Made the following repairs during peer review:

#UE4doc

diff --git a/Engine/Documentation/Source/Engine/Performance/ReducingPackageSize/ReducingPackageSize.INT.udn b/Engine/Documentation/Source/Engine/Performance/ReducingPackageSize/ReducingPackageSize.INT.udn
--- a/Engine/Documentation/Source/Engine/Performance/ReducingPackageSize/ReducingPackageSize.INT.udn
+++ b/Engine/Documentation/Source/Engine/Performance/ReducingPackageSize/ReducingPackageSize.INT.udn
@@ -1,10 +1,10 @@
 ﻿Availability:Public
 Title: Reducing Packaged Game Size 
-Crumbs:
+Crumbs:%ROOT%, Engine, Engine/Performance
 Description:How to reduce the size of your packaged game.
 Skilllevel: Intermediate
-Version: 4.8
-related:Platforms/Android/ReducingAPKSize
+Version: 4.12
+Related:Platforms/Android/ReducingAPKSize
 Parent:Engine/Performance
 
 [VAR:Topic]
@@ -48,19 +48,19 @@ In the following guide we will cover steps you can take to help reduce your proj
 
 [EXCERPT:main]
 ## Creating a new empty project
-When you first start working on your Android based mobile project you might be tempted to use an existing project as a base to work from or create a new project that has the Starter Content enabled. 
-**DO NOT** do this and instead create a completely new, empty C++ or Blueprint based project and then use the [Migration Tool](Engine/Content/Browser/UserGuide/Migrate) to bring over the content that you want to use. 
-This way you ensure that the only content in your project is content that you need to be there.
+When you first start working on your Android based mobile project, you might be tempted to use an existing project as a base to work from or create a new project that has the Starter Content enabled. 
+**DO NOT** do this, instead, create a completely new, empty C++ or Blueprint based project, and then use the [Migration Tool](Engine/Content/Browser/UserGuide/Migrate) to bring over the content that you want to use. 
+This way, you ensure that the only content in your project is content that needs to be there.
 
 ## Compressing Cooked Content
 The easiest and quickest way to reduce your APK package size is to tell UE4 to compress the APK packages during the packing process. 
-To turn on package compression you need to do the following in the UE4 Editor.
+To turn on package compression, you need to do the following in the UE4 Editor.
 
-1. First open up the projects settings by going to the Main Toolbar and selecting the **Edit** option then selecting **Project Settings**.
+1. First, open up the project's settings by going to the Main Toolbar and selecting the **Edit** option, then selecting **Project Settings**.
 
 	![](T_Project_Settings_Menu.png)
 
-1. Under the Project section click on the **Packaging** section to show the options for how the project will be packaged.
+1. Under the Project section, click on the **Packaging** section to show the options for how the project will be packaged.
 
 	![](T_Project_Packaging.png)
 
@@ -68,55 +68,55 @@ To turn on package compression you need to do the following in the UE4 Editor.
 
 	![](T_Expose_Advanced_Props.png)
 
-1. Look for the option labeled **Create compressed cooked packages** and enable it if not already enabled.
+1. Look for the option labeled **Create compressed cooked packages** and enable it (if not already enabled).
 
 	![](T_Enable_Comp_Packages.png)
 	
 If you have not packaged your game with the Create compressed cooked packages checkmark box enabled you should see a huge difference in size when your project is re-packaged with it enabled.
-In fact it is not uncommon for some projects APK packages to be reduced in size by up to 50% when create compressed cooked packages is enabled.
+In fact it is not uncommon for some project's APK packages to be reduced in size by up to 50% when **Create compressed cooked packages** is enabled.
 
 ##  Setting up a Project's Levels
-One overlooked area that can lead to a bloating of APK package sizes is failing to properly setup your projects options like which levels should be loaded by default or which levels should be used between level transitions. 
-To set which level or levels should be used for this type of interaction you can do the following.
+One overlooked area that can lead to the bloating of APK package sizes is failing to properly setup your project's options, like which levels should be loaded by default, or which levels should be used between level transitions. 
+To set which level (or levels) should be used for this type of interaction, you can do the following.
 
-1. First on the **Main Toolbar** go to  **Edit** > **Project Settings**.
+1. First, on the **Main Toolbar** go to  **Edit** > **Project Settings**.
 
 	 ![](T_Project_Settings.png)
 
-1. Then under **Project** in the **Maps & Modes** section, set the map in the **Default Maps** to a map that is specific to your project.
+1. Then, under **Project** in the **Maps & Modes** section, set the map in the **Default Maps** to a map that is specific to your project.
 
 	![](T_Set_Maps.png)
 
 ## Checking what Content will be Cooked
-To ensure that only the content that is related to your project is added to your APK file, you can check which content is being added to your projects APK by looking in the **Cooked** folder in your project.
+To ensure that only the content that is related to your project is added to your APK file, you can check which content is being added to your project's APK by looking in the **Cooked** folder in your project.
 You can find the Cooked folder by going to **(ProjectName)\Saved\Cooked** and then choosing the folder with the format you cooked your project for.
 
 [region:note]
-Please note that you will only be able to see the content in the Cooked folder after you complete your first cook. 
+Please note that you will only be able to see the content in the **Cooked** folder after you complete your first cook. 
 [/region]
 
 ![](T_Cooked_Content_Folder.png)
 
 The image below shows the Texture folder from the Match 3 sample game. 
-The content in the folder has been sorted by size so it is easier to see what assets are the biggest.
-Once the biggest assets have been established you can then examine the assets inside of UE4 to see if in fact they can be reduced in size without sacrificing the integrity of the asset.
+The content in the folder has been sorted by size, so it is easier to see what assets are the biggest.
+Once the biggest assets have been determined, you can then examine the assets inside of UE4 to see if in fact they can be reduced in size without sacrificing the integrity of the asset.
 
 ![](T_Cooked_Content.png)
 
 ## Development VS Shipping Build Size
 
-When trying to figure out the final size your project will be keep in mind that the Development version of your project will be slightly larger than your Shipping build size.
-On the Medieval Match example game the difference in size between a Development and Shipping build was around **14 percent**.
-However since each project has different requirements, the saving between the two different builds types for your project could be more or less than a 14 percent savings.
+When trying to figure out the final size your project will be, keep in mind that the Development version of your project will be slightly larger than your Shipping build size.
+On the Medieval Match example game, the difference in size between a Development and Shipping build was around 14 percent.
+However, since each project has different requirements, the saving between the two different build types for your project could be more or less than a 14 percent savings.
 
 ## Per-Device Texture LOD's
-With the release of Unreal Engine 4.8 you can now specify which size Texture should be used on which devices. 
+With the release of Unreal Engine 4.8, you can now specify which size Texture should be used on which devices. 
 You can read more about this in the [Texture Guidelines for Mobile Platforms](Platforms/Mobile/Textures) section.
 
 ## Removing Unused Content
-When your project is finally ready to be packaged up for the store make sure to first remove any unused or testing content by selecting it in the Content Browser and deleting it.
+When your project is finally ready to be packaged up for the store, make sure to first remove any unused or testing content by selecting it in the Content Browser and deleting it.
 You can delete content by first selecting the content you want to delete and then either press the **Delete** key on the keyboard or use the Delete option in the Right - Click menu. 
-When you do try and delete something from your project the **Delete Assets** menu will be displayed like in the following image.
+When you do try and delete something from your project, the **Delete Assets** menu will be displayed like in the following image.
 
 [region:tip]
 This is the preferred method of deleting objects in UE4 and should be used over just simply deleting/removing the assets from their location in the content folder.
@@ -124,27 +124,26 @@ This is the preferred method of deleting objects in UE4 and should be used over
 
 ![](T_Delete_Menu.png)
 
-The Delete Assets menu will inform you if the asset you are trying to delete is referenced by another asset or not. 
-If it is not referenced you can simply delete it, however if it is referenced by other assets you can use the **Replace References** option to replace the reference to that asset to another asset that is supposed to be packaged in your project.
-
+The Delete Assets menu will inform you if the asset you are trying to delete is referenced by another asset. 
+If it is not referenced, you can simply delete it, however if it is referenced by other assets, you can use the **Replace References** option to replace the reference to that asset to another asset that is supposed to be packaged in your project.
 
 ## Removing Unused Plug-ins
-Disabling any unused plug-in's inside of the UE4 Editor is another way to ensure that un-need content and code do not make it into your projects final APK package file.
-This is because some Plug-in's require a certain set of base assets and code to work correctly.
-If the Plug-in is not disabled the assets and code that are required to make the plug-in work will not be removed from your project.
-While this might not reduce your projects size as much as reducing or removing a large Texture would, every little bit can help when trying to slim your project down to the 50MB size.
-To disable a plug-in inside of UE4 you need to do the following. 
+Disabling any unused plug-ins inside of the UE4 Editor is another way to ensure that unnecessary content and code do not make it into your project's final APK package file.
+This is because some plug-ins require a certain set of base assets and code to work correctly.
+If the plug-in is not disabled, the assets and code that are required to make the plug-in work will not be removed from your project.
+While this might not reduce your project's size as much as reducing or removing a large Texture would, every little bit can help when trying to slim your project down to the 50MB size.
+To disable a plug-in inside of UE4, you need to do the following. 
 
 [region:note]
 	Make sure you fully test your project to ensure that the plug-ins you are disabling do not interfere with the functionality of your project.
 [/region]
 
-1. First open up the Plug-ins Manager by going to the **Main Toolbar** and clicking on the **Window** option and then selecting the **Plug-in** option.
+1. First, open up the Plug-ins Manager by going to the **Main Toolbar** and clicking on the **Window** option and then selecting the **Plug-in** option.
 
 	![](T_Open_Plugins_Browser.png)
 
-1. From the Plug-in Browser select the different sections and disable any plug-in's you are not using by un-checking the check mark box next to **Enabled**.
-  When you have disabled all the Plug-ins you will need to re-start the UE4 Editor and then make sure to re-package your project.
+1. From the Plug-in Browser, select the different sections and disable any plug-ins you are not using by un-checking the check mark box next to **Enabled**.
+  When you have disabled all the plug-ins, you will need to re-start the UE4 Editor and then make sure to re-package your project.
 
 	[OBJECT:EmbeddedVideo]
 	[PARAMLITERAL:width]
@@ -160,8 +159,8 @@ To disable a plug-in inside of UE4 you need to do the following.
 
 ## Package Black List
 
-With the release of Unreal Engine 4.9 and later you can now place a text file in your projects **Build/Platform/(Target Platform Folder)** directory that will tell the cooker to exclude partial or complete file paths from being packaged into your project.
-Project's can have multiple Black List files for Debug, Development, Test, and Shipping builds that can be setup to include or exclude whatever project data you want.
+With the release of Unreal Engine 4.9 and later you can now place a text file in your project's **Build/Platform/(Target Platform Folder)** directory that will tell the cooker to exclude partial or complete file paths from being packaged into your project.
+Projects can have multiple Black List files for Debug, Development, Test, and Shipping builds that can be setup to include or exclude whatever project data you want.
 You can even have different Black List files for each platform your project supports like one for Android, one for iOS, ect. 
 Here is an example of what the Black List files would look like for Tappy Chicken on Android. 
 
@@ -173,25 +172,25 @@ Here is an example from Tappy Chicken where the cooker is told to not include th
 
 	![](T_Black_List.png)
 
-The first three entries should be safe for any project, but there may be cases were the **AssetRegistry.bin** is needed at runtime depending on your project.
-It is highly recommend that you thoroughly test to ensure that the content you are removing will not negatively impact your project.
-An easy and simple way to go about this is to open up all your projects levels to ensure that they correctly load and have no errors or warnings.
-If you run into further problems and are not sure where they are coming from check the Logcat for errors.
+The first three entries should be safe for any project, but there may be cases were the **AssetRegistry.bin** is needed at runtime (depending on your project).
+It is highly recommended that you thoroughly test to ensure that the content you are removing will not negatively impact your project.
+An easy and simple way to go about this is to open up all your project's levels to ensure that they correctly load and have no errors or warnings.
+If you run into further problems and are not sure where they are coming from, check the Logcat for errors.
 
 ## Max Dynamic Point Lights
 
-With the release of UE 4.9 a new option has been introduced that can help reduce the number of shader that are generated for dynamic lights
- making this feature is particularly useful for mobile games that do not require dynamic lighting. 
+With the release of UE 4.9 a new option has been introduced that can help reduce the number of shaders that are generated for dynamic lights. 
+Making this feature is particularly useful for mobile games that do not require dynamic lighting. 
 The property, **Max Dynamic Point Lights** can be found under the **Mobile** section of the **Project Settings**.
 
 ![](Max_Dyn_Point_Lights.png)
 
-To completely disable this feature set the **Max Dynamic Point Lights** from **4** to **0** and then re-start the UE4 Editor when prompted.
+To completely disable this feature, set the **Max Dynamic Point Lights** from **4** to **0** and then re-start the UE4 Editor when prompted.
 
 ![](Max_Dyn_Point_Lights_To_0.png)
 
-Depending on the size of your project and the amount of Materials used the size saving you can achieve can range from just a few MB into the multiple MB range.
-Using the default Third Person template for an example, we can see what the size savings are when Max Dynamic Point Lights are enabled verses when Max Dynamic Point Lights and disabled completely. 
+Depending on the size of your project and the amount of Materials used, the size saving you can achieve can range from just a few MB to the multiple MB range.
+Using the default Third Person template for an example, we can see what the size savings are when Max Dynamic Point Lights are enabled versus when Max Dynamic Point Lights are disabled completely. 
 
 | State| Packaged Game Size | Savings|
 | ---------| ------| -----|
diff --git a/Engine/Documentation/Source/Platforms/Android/ReducingAPKSize/ReducingAPKSizes.INT.udn b/Engine/Documentation/Source/Platforms/Android/ReducingAPKSize/ReducingAPKSizes.INT.udn
--- a/Engine/Documentation/Source/Platforms/Android/ReducingAPKSize/ReducingAPKSizes.INT.udn
+++ b/Engine/Documentation/Source/Platforms/Android/ReducingAPKSize/ReducingAPKSizes.INT.udn
@@ -3,25 +3,27 @@ Title: Reducing APK Package Size
 Crumbs:%ROOT%, Platforms, Platforms/Android
 Description:How to reduce the size of your packaged Android game.
 Skilllevel: Intermediate
-Version: 4.8
+Version: 4.12
 Parent:Platforms/Android
+Related:Engine/Performance/ReducingPackageSize
 tags:Mobile
 tags:Android
 type:overview
 
 [TOC(start:2 end:3)]
 
-No matter what type of Android device your UE4 project is targeting reducing the size of the Android application package or **APK** to be under the 50MB limit for the Google Play store can be challenging. 
-In the following guide we will cover steps you can take to help reduce your projects final APK package size to be as small as possible using nothing but the tools provided to you in the UE4 Editor. 
+No matter what type of Android device your UE4 project is targeting, reducing the size of the Android application package (**APK**) to be under the 50MB limit for the Google Play store can be challenging. 
+In the following guide, we will cover steps you can take to help reduce your project's final APK package size to be as small as possible, using nothing but the tools provided to you in the UE4 Editor. 
 
 ## ETC1 Texture Considerations
-When packing your project for Android using the ETC1 Texture format you need to be aware that ETC1 will not compress Textures that have embedded Alpha information in them and instead treat them as fully uncompressed Textures. 
-While this will not greatly affect your projects run time performance, it will increase the size of your projects APK.
-If your project does require Textures that use Alpha's it is better to use two separate Textures, one for the Diffuse and one for the Alpha, 
-instead of a single Texture with the Alpha embedded inside of it as this will ensure that the Textures will be compressed when your project is cooked and packaged.
+
+When packing your project for Android using the ETC1 Texture format, you need to be aware that ETC1 will not compress Textures that have embedded Alpha information in them, and instead treat them as fully uncompressed Textures. 
+While this will not greatly affect your project's run time performance, it will increase the size of your project's APK.
+If your project does require Textures that use Alphas, it is better to use two separate Textures; one for the Diffuse, and one for the Alpha, 
+instead of a single Texture with the Alpha embedded inside of it. This will ensure that your Textures will be compressed when the project is cooked and packaged.
 
 ## Platform-Agnostic Techniques
 
-In addition to the Android-specific methods listed above, there are many ways to reduce the size of any packaged game regardless of the target platform.
+In addition to the Android-specific methods listed above, there are many ways to reduce the size of any packaged game (regardless of the target platform).
 
 [INCLUDE:Engine/Performance/ReducingPackageSize#main(offset:1)]
