#UE4 docs minor adjustments to highlighting, added a link to refelection captures

diff --git a/Engine/Documentation/Source/Platforms/Mobile/Performance/MobilePerformance.INT.udn b/Engine/Documentation/Source/Platforms/Mobile/Performance/MobilePerformance.INT.udn
--- a/Engine/Documentation/Source/Platforms/Mobile/Performance/MobilePerformance.INT.udn
+++ b/Engine/Documentation/Source/Platforms/Mobile/Performance/MobilePerformance.INT.udn
@@ -16,22 +16,22 @@ tags:mobile
 Below are things to watch out for or suggestions on how to set up your content to maximize performance on mobile platforms. These are listed from most severe to least.
 
 * Make sure lighting is built before running on a device.
-* Make sure most Post Process features are disabled.  In fact, the only ones you can use for iPad4 are Temporal AA and the Film ones, including vignette.  A bunch of expensive features like Bloom and DOF are enabled by default as we have the same settings between PC and mobile.  Post Process features can easily cost 60+ms with the default settings.  We hope to be able to use DOF and light shafts on higher end devices.  `show postprocessing` can be used on devices to toggle PP and see how much it is costing.
-* Make sure precomputed visibility is setup properly.  To do this, you need to place PrecomputedVisibilityVolumes around anywhere that the player can walk, and build lighting.  You need to make sure the same P level is used when building lighting and when running the game (aka do not build the sublevel by itself).  You can verify that it is working by entering `stat initviews` on device or in previewer and make sure the `Statically occluded primitives` is > 0.  Use `r.ShowPrecomputedVisibilityCells 1` to visualize the cells in the editor.
-* Use masked and translucent Materials sparingly, and only in places where they cover a small part of the screen.  iOS devices are very optimal in shading opaque surfaces, as they will only shade each pixel once, most of the time.  But for masked and translucency, every single layer will have to be shaded.  The total GPU time of a frame can easily be doubled or more by having overdraw.  Use the Shader complexity viewmode in Mobile previewer to investigate hot spots.
+* Make sure most **Post Process** features are disabled.  In fact, the only ones you can use for iPad4 are **Temporal AA** and the Film ones, including vignette.  A bunch of expensive features like Bloom and DOF are enabled by default as we have the same settings between PC and mobile.  **Post Process** features can easily cost 60+ms with the default settings.  We hope to be able to use DOF and light shafts on higher end devices.  The console command `show postprocessing` can be used on devices to toggle **Post Processing** and see how much it is costing.
+* Make sure **Precomputed Visibility** is setup properly.  To do this, you need to place **Precomputed Visibility Volumes** around anywhere that the player can walk, and then build lighting.  You need to make sure the same **Persistent** level is used when building lighting and when running the game (aka do not build the sublevel by itself).  You can verify that it is working by entering `stat initviews` on device or in previewer and make sure the `Statically occluded primitives` is > 0.  Use `r.ShowPrecomputedVisibilityCells 1` to visualize the cells in the editor.
+* Use masked and translucent **Materials** sparingly, and only in places where they cover a small part of the screen.  iOS devices are very optimal in shading opaque surfaces, as they will only shade each pixel once, most of the time.  But for masked and translucency, every single layer will have to be shaded.  The total GPU time of a frame can easily be doubled or more by having overdraw.  Use the **Shader Complexity Viewmode** in **Mobile Previewer** to investigate hot spots.
 * Draw calls of the entire scene should be &lt;=700 for any view.  The areas with poor occlusion, like looking over the landscape, are going to be the biggest challenges for this.  This can be seen with `stat openglrhi` on device, `stat d3d11rhi` in ES2 previewer on PC.
 * Triangle count of the entire scene should be &lt;=500k for any view.  This has been determined to be the poly count needed to hit 30fps on both iPad4 and iPad Air.  This can be seen with `stat openglrhi` on device, `stat d3d11rhi` in ES2 previewer on PC.
-* Materials should use as few texture lookups and instructions as possible.  You have 5 texture samplers available, but using all of those makes for a fairly expensive Material.  For now, use the PC instruction count to optimize, and ES2 previewer shader complexity to visualize the total cost.
-* Only use independent texture fetches in Materials.  This means UVs in the pixel shader (BaseColor, Roughness, etc) must not be manipulated in any way, such as scaling.  Use the new CustomizedUVs feature instead to do the scaling in the vertex shader.  Some special features like environment mapping require math on the UVs and that is ok for special cases. 
+* **Materials** should use as few texture lookups and instructions as possible.  You have 5 texture samplers available, but using all of those makes for a fairly expensive Material.  For now, use the PC instruction count to optimize, and ES2 previewer shader complexity to visualize the total cost.
+* Only use independent texture fetches in Materials.  This means UVs in the pixel shader (BaseColor, Roughness, etc) must not be manipulated in any way, such as scaling.  Use the **CustomizedUVs** feature instead to do the scaling in the vertex shader.  Some special features like environment mapping require math on the UVs and that is ok for special cases. 
 * Prefer square textures as they have less memory waste, and always use power of 2 dimensions (512, 1024, 2048).  Use `listtextures` in the ES2 previewer to see where all the texture memory is going. 
 
 ## Performance Tiers
 
-UE4 supports a variety of lighting features on mobile devices. Using these features costs performance and may cause your game to perform poorly on slower mobile devices. While it is possible to mix and match most of UE4's mobile lighting features, it can be useful to categorize these features into the following tiers. When building a mobile game, you should decide which features to use based on the quality of graphics your game requires and the types of devices you need to support. Check the device compatibility pages under [iOS Development](Platforms/iOS/DeviceCompatibility) and [Android Development](Platforms/Android/DeviceCompatibility) for more information on what devices we have tested here at Epic and which tiers we think are most appropriate.
+Unreal Engine 4 (UE4) supports a variety of lighting features on mobile devices. Using these features costs performance and may cause your game to perform poorly on slower mobile devices. While it is possible to mix and match most of UE4's mobile lighting features, it can be useful to categorize these features into the following tiers. When building a mobile game, you should decide which features to use based on the quality of graphics your game requires and the types of devices you need to support. Check the device compatibility pages under [iOS Development](Platforms/iOS/DeviceCompatibility) and [Android Development](Platforms/Android/DeviceCompatibility) for more information on what devices we have tested here at Epic and which tiers we think are most appropriate.
 
 ###LDR (Low Dynamic Range)
 
-Low dynamic range or LDR mode is the highest performance tier supported in UE4 and is recommended for games that do not require lighting or Post Processing features. To use this mode, you must disable `Mobile HDR` for your project in the Rendering section of the [Project Settings Editor](Engine/UI/ProjectSettings). 
+Low dynamic range or LDR mode is the highest performance tier supported in UE4 and is recommended for games that do not require lighting or **Post Processing** features. To use this mode, you must disable `Mobile HDR` for your project in the Rendering section of the [Project Settings Editor](Engine/UI/ProjectSettings). 
 
 **Advantages**
 
@@ -48,32 +48,32 @@ Low dynamic range or LDR mode is the highest performance tier supported in UE4 a
 
 **Recommendations**
 
-* Make sure all of your Materials have their shading model set to Unlit for maximum performance.
+* Make sure all of your **Materials** have their shading model set to **Unlit** for maximum performance.
 * Do not place lights in your scene or otherwise depend on lighting for maximum performance.
-* Consider performing as many operations in the Material's vertex shader as possible. You can do so by enabling Customized UVs, connecting nodes to them, and reading them in the pixel shader using a Texture Coordinate node.
+* Consider performing as many operations in the **Material's** vertex shader as possible. You can do so by enabling **Customized UVs**, connecting nodes to them, and reading them in the pixel shader using a **Texture Coordinate** node.
 
 
 ### Basic Lighting
 
-In this tier, you leverage static lighting and fully rough Materials to create levels with interesting lighting while maximizing performance to reach a broader range of mobile devices.
+In this tier, you leverage static lighting and fully rough **Materials** to create levels with interesting lighting while maximizing performance to reach a broader range of mobile devices.
 
 **Advantages**
 
 * Access to static lighting and global illumination features.
-* Full HDR pipeline with access to some Post Processing features like tone mapping.
+* Full HDR pipeline with access to some **Post Processing** features like tone mapping.
 * Translucency is blended in linear space allowing you to author content as you normally would for any other UE4 game.
 
 
 **Limitations**
 
-* As you are marking all Materials as fully rough, your Materials will not have interesting specular reflections.
+* As you are marking all **Materials** as **Fully Rough**, your **Materials** will not have interesting specular reflections.
 * If you also choose to disable lightmap directionality, normal maps will have no effect.
 
 
 **Recommendations**
 
-* Author all of your Materials with the fully rough flag set.
-* Consider disabling light map directionality in your Materials for additional performance.
+* Author all of your **Materials** with the fully rough flag set.
+* Consider disabling light map directionality in your **Materials** for additional performance.
 * Use only static lights in your maps.
 * Disable bloom and anti-aliasing. Also, stick to the basic set of film and color controls. See [mobile Post Process](Platforms/Mobile/PostProcessEffects) for more details.
 
@@ -95,7 +95,7 @@ In this tier, you take advantage of most of the HDR lighting features available
 
 * Consider enabling bloom to take full advantage of the HDR lighting pipeline.
 * Realistic specular reflections in combination with HDR lighting can lead to specular aliasing. You should take advantage of the normal map to roughness feature to help reduce specular aliasing due to high frequency information in your normal map. You should also consider enabling anti-aliasing in the Post Process volume to help reduce this artifact.
-* Place reflection captures intelligently to get the best results.
+* Place **Reflection Captures** intelligently to get the best results. See [Reflections: Placing Reflection Captures](Resources/Showcases/Reflections/#placingreflectioncaptures) for more information.
 * Use only static lights in your scene.
 
 
@@ -119,8 +119,8 @@ In this tier, you take advantage of all of the HDR lighting features available f
 ## Shader Complexity View Mode
 
 You can use the shader complexity view mode in the PC ES2 previewer to get a sense of Material cost.  To use it, 
-just enter `viewmode shadercomplexity` on the console in the PC ES2 previewer. Just like the SM5 shader complexity, 
-green means good performance, bright red means very expensive, white or pink means EXTREMELY expensive.  Here are some 
+just enter `viewmode shadercomplexity` in the console in the **PC ES2 Previewer**. Just like the SM5 shader complexity, 
+green means good performance, bright red means very expensive, white or pink means **extremely** expensive.  Here are some 
 examples from the mobile showcase:
 
 [REGION:imagetable]
@@ -131,18 +131,18 @@ examples from the mobile showcase:
 
 <!--
 [EXCERPT:Image1]
-The pillar Materials are fairly expensive here, and the translucent volumetric sheets are VERY expensive.  In this scene, I asked to have the translucent sheets removed, as they cost too much. 
+The pillar Materials are fairly expensive here, and the translucent volumetric sheets are **very** expensive.  In this scene, I asked to have the translucent sheets removed, as they cost too much. 
 [/EXCERPT:Image1]
 [EXCERPT:Image2]
 Here the pillars are quite expensive as they used all five texture lookups and did quite a bit of layering. Otherwise, this is just about efficient enough for 30fps.
 [/EXCERPT:Image2]
 [EXCERPT:Image3]
-Here the tree is causing pixels to be EXTREMELY expensive.  If the player could ever get up close or cover his screen with that Material, the cost would be enormous.
+Here the tree is causing pixels to be **extremely** expensive.  If the player could ever get up close or cover his screen with that Material, the cost would be enormous.
 [/EXCERPT:Image3]
 -->
 
 ## Mobile Content Scale Factor 
-Mobile Content Scale Factor is a way to scale the resolution of your project to best fit the screen resolution of the mobile device being used to view your project. 
+**Mobile Content Scale Factor** is a way to scale the resolution of your project to best fit the screen resolution of the mobile device being used to view your project. 
 You can create and store multiple device profiles by creating a new .INI file called **DefaultDeviceProfiles.INI** in your project's config folder. 
 Inside of this file, you can specify what will happen to the resolution of your project by entering the following command. 
 
@@ -212,7 +212,7 @@ The following images show what a Material will look like when the Material Quali
 ### Setting Material Quality Level
 You can set the Material Quality Level in the following ways inside UE4:
 
-**From the Console**: You can set the quality level that your Materials will use in the Editor by pressing the Backtick key to open up the UE4 console and entering **r.MaterialQualityLevel ** followed by one of the following numbers.
+**From the Console**: You can set the quality level that your Materials will use in the Editor by pressing the Backtick key to open up the UE4 console and entering `r.MaterialQualityLevel` followed by one of the following numbers.
 
 ![](Platforms\Mobile\Performance\MQS_Console_Command.png)
 
@@ -280,7 +280,7 @@ Click for full image.
 [/REGION]
 
 ### Mobile Material Options 
-If you open up any Material and click on the Main Material Node then look in the Details Panel under the **Mobile** section you will find the following two properties.
+If you open up any **Material** and click on the Main Material Node then look in the Details Panel under the **Mobile** section you will find the following two properties.
 
 * **Fully Rough**: Forces a Material to be completely rough saving a number of instructions and a single texture sampler.
 * **Use Lightmap Directionality**: Uses lightmap directionally and per pixel normals. If disabled the lightmaps will be flatter but cheaper.
