-Interactive Widget Component Overview
-Creating 3D Widget How-To
-Creating Virtual Keyboard How-to
-Widget Component Page Updates

diff --git a/Engine/Documentation/Source/Engine/Components/Widget/WidgetComponents.INT.udn b/Engine/Documentation/Source/Engine/Components/Widget/WidgetComponents.INT.udn
--- a/Engine/Documentation/Source/Engine/Components/Widget/WidgetComponents.INT.udn
+++ b/Engine/Documentation/Source/Engine/Components/Widget/WidgetComponents.INT.udn
@@ -7,32 +7,52 @@ Related: Engine/UMG/UserGuide/WidgetBlueprints
 Related: Engine/MediaFramework
 Related: Engine/Blueprints/UserGuide/Components
 SkillLevel: Beginner
-Version: 4.9
+Version: 4.13
 
-Widget Components allow you to manifest 3D UI elements created through [**Unreal Motion Graphics**](Engine/UMG/UserGuide) in your game world. 
+Widget Components allow you to manifest 3D UI elements created through [Unreal Motion Graphics](Engine/UMG/UserGuide) in your game world. 
 
 ## Widget Component
 
-The **WidgetComponent** itself is a 3D instance of a Widget Blueprint that you can interact with in your game world. 
+The **Widget** Component itself is a 3D instance of a Widget Blueprint that you can interact with in your game world. 
 
-In the example below, we have added a WidgetComponent which will now allow us to specify a Widget Class Blueprint. 
+In the example below, we have added a Widget Component which will now allow us to specify a Widget Class Blueprint. 
 
-![](Engine/UMG/HowTo/Create3DWidgets/UMG_3D_21.png)
+![](UMG_3D_21.png)
 
-_This example uses a Widget Blueprint to display a interactive menu that appears in the game world._
+[REGION:caption]
+This example uses a Widget Blueprint to display a interactive menu that appears in the game world.
+[/REGION]
+
+You can change the size of the WidgetComponent in the world by changing the **Draw Size** or using **Draw at Desired Size**. 
+
+![](UMG_3D_22.png)
+
+Once you place the Actor that contains the WidgetComponent in your level, the Widget Class Blueprint is displayed in the world. 
 
-You can change settings to how the WidgetComponent is perceived in the world by changing the Draw Size or Max Interaction distance. 
+![](UMG_3D_30.png)
 
-![](Engine/UMG/HowTo/Create3DWidgets/UMG_3D_22.png)(w:720)
+[REGION:caption]
+Here we have an interactive keypad that sits on a wall.
+[/REGION]
+
+## Widget Component Property Reference
 
-Once we place the Actor that contains our WidgetComponent in our level, our Widget Class Blueprint is displayed in the world. 
+Here are the properties that can be adjusted via the **Details** panel for the Widget Component. 
 
-![](Engine/UMG/HowTo/Create3DWidgets/UMG_3D_30.png)(w:720)
+| Option | Description |
+| --- | ----------- |
+| **Space** | The coordinate space in which to render the widget (World or Screen). With World the widget is rendered in the world as a mesh while and can be occluded while Screen renders the widget on the screen completely outside of the world and is never occluded. |
+| **Widget Class** | The class of User Widget to create and display an instance of. |
+| **Draw Size** | The size of the displayed quad. |
+| **Manually Redraw** | Should the Widget wait to be told to redraw to actually draw or not. |
+| **Redraw Time** | The time in between draws, if 0 - redraw every frame. If 1, we would redraw every second. This will work with the Manually Redraw as well. So you can say, manually redraw, but only redraw at this maximum rate. |
+| **Window Focusable** | Is the virtual window created to host the widget focusable? Should this take focus from the user. |
+| **Draw at Desired Size** | Causes the render target to automatically match the desired size specified by the Widget Class. Expensive if is done every frame. 
+| **Pivot** | The alignment/pivot point that the widget is placed at relative to the position. |
 
-See **[](Engine/UMG/HowTo/Create3DWidgets)** for a step-by-step guide to creating this example using a WidgetCompoment and Widget Blueprint. 
 
 [REGION:note]
-This Component is considered as experimental and may be substantially changed or removed in future releases. 
+See [](Engine\UMG\HowTo\InWorldWidgetInteraction) and [](Engine\UMG\HowTo\VirtualKeyboards) for creating 3D UI elements. 
 [/REGION]
 
 
@@ -41,3 +61,5 @@ This Component is considered as experimental and may be substantially changed or
 
 
 
+
+
diff --git a/Engine/Documentation/Source/Engine/UMG/HowTo/InWorldWidgetInteraction/InWorldWidgetInteraction.INT.udn b/Engine/Documentation/Source/Engine/UMG/HowTo/InWorldWidgetInteraction/InWorldWidgetInteraction.INT.udn
new file mode 100644
--- /dev/null
+++ b/Engine/Documentation/Source/Engine/UMG/HowTo/InWorldWidgetInteraction/InWorldWidgetInteraction.INT.udn
@@ -0,0 +1,154 @@
+Availability: Docs
+Title:Creating 3D Widget Interaction
+Crumbs: %ROOT%, Engine, Engine/UMG, Engine/UMG/HowTo
+Description:An example of how to create, display and remove UI elements from the screen during gameplay. 
+Related: Engine\UMG\UserGuide\WidgetInteraction\
+Related: Engine/UMG
+Related: Engine/UMG/UserGuide
+Version: 4.13
+SkillLevel: Intermediate
+
+Traditional UI consists in 2D space on top of a 2D or 3D world (for example menus, health bars or score displays). 
+However there may be times when you want a piece of UI to exist within your 3D environment that players can interact with (a keypad, a virtual menu, an inventory or other interactive UI element). 
+In VR projects, this is even more prevalent and UMG supports these types of interactions with the [Widget Interaction](Engine\UMG\UserGuide\WidgetInteraction) component. 
+
+In this how-to we create a widget containing a button that can be placed in the level and clicked on by a player. 
+
+
+## Steps
+
+[REGION:note]
+For this how-to, we are using the **Blueprint First Person Template** with **Starter Content** enabled. 
+[/REGION]
+
+1.	Inside the **Content/FirstPersonBP/Blueprints** folder, open the **FirstPersonCharacter** Blueprint. 
+
+	![](WidgetInteraction_01.png)
+
+1.	Click the **Add Component** button and search for and add the **Widget Interaction** component. 
+
+	![](WidgetInteraction_02.png)
+
+1.	**Left-click** drag and drop the Widget Interaction component on top of the **Sphere** component attached to **FP_Gun**. 
+
+	![](WidgetInteraction_03.png)
+
+	This will allow us to point the Widget Interaction component in the direction the gun is pointing. 
+
+1.	In the **Details** panel for the Widget Interaction component, zero out the **Location** and check the **Show Debug** option. 
+
+	![](WidgetInteraction_04.png)
+
+1.	Inside the **Event Graph**, add a **Right Mouse Button** Mouse Event node.
+
+1.	Drag the Widget Interaction component into the graph, then drag off it and add a **Press Pointer Key** node. 
+
+	![](WidgetInteraction_05.png)
+
+1.	Connect to the **Pressed** pin of the **Right Mouse Button** node and set the **Key** to **Left Mouse Button**. 
+
+	![](WidgetInteraction_06.png)
+
+	By default UMG registers click events as a result of a Left Mouse Button key press, using this node will allow us to simulate that when the Right Mouse Button is pressed. 
+	You could use any Key Input event to call this function (such as a trigger press or button press to simulate a Left Mouse Button click). 
+
+1.	Add a **Release Pointer Key** node (set to **Left Mouse Button**) and connect it to the **Released** pin of the Right Mouse Button node. 
+
+	![](WidgetInteraction_07.png)
+
+1.	In the **Content Browser**, create a **Widget Blueprint** called **InteractiveWidget**. 
+
+	![](WidgetInteraction_08.png)
+
+1.	In the **Visual Designer**, remove the **Canvas Panel** then add a **Button** with a **Text** widget on top of the button. 
+
+	[REGION:lightbox]
+	[![](WidgetInteraction_09.png)](WidgetInteraction_09.png)
+	[/REGION]
+
+	[REGION:caption]
+	Click for full view. 
+	[/REGION]
+
+	For this example we are purely adding a button that can be clicked on. The text on the button we will indicate the number of times the button has been clicked by the player. 
+	You can create any type of interactive widget and interact with it using the Widget Interaction component (things like sliders, combo boxes, check boxes, etc. can be interacted with in world space). 
+
+1.	In the **Details** panel for the **Button**, under **Style/Hovered** change the **Tint** to a different color. 
+
+	![](WidgetInteraction_10.png)
+
+1.	In the **Details** panel for the **Text**, change the **Text** content to **0** and increase the **Font Size** to **48**. 
+
+	![](WidgetInteraction_11.png)
+
+1.	In the **Graph**, create a **Text** Variable called **Value** then **Compile** and set **Value** to **0**. 
+
+	![](WidgetInteraction_12.png)
+
+	This is the variable that we will update when the player clicks on the button which we will bind to our Text widget. 
+
+1.	On the **Designer** tab for the **Text** widget, click the **Bind** option next to **Text** and select the **Value** property.
+
+	![](WidgetInteraction_13.png)
+
+1.	In the **Details** panel for the **Button** widget under **Events**, click the **+** next to **On Clicked**. 
+
+	![](WidgetInteraction_14.png)
+
+	This will fire whenever the button is clicked on which is called from inside our **FirstPersonCharacter** Blueprint. 
+
+1.	On the **Graph** tab, hold **Ctrl** and drag in the **Value** variable then off it use **To String** followed by **String To Int**. 
+	
+	![](WidgetInteraction_15.png)
+
+	We are converting the Text to a String, then from a String to an Int so that we can increase the value when the button is clicked. 
+
+1.	Off the **String To Int** pin, use an **Integer + Integer** node set to **+ 1**. 
+
+1.	Hold **Alt** and drag in the **Value** variable to set it, then connect the remaining pins to the **OnClicked** event as shown. 
+
+	![](WidgetInteraction_16.png)
+
+	The **ToText(int)** node will be automatically created when you connect the out pin of **Integer + Integer** to the text input pin on **Value**. 
+
+1.	In the **Content Browser**, create a new **Blueprint** based on **Actor** called **ExampleWidget** and add a **Widget** component to it. 
+
+	![](WidgetInteraction_18.png)
+
+	This will represent the 3D version of our **InteractiveWidget** Widget Blueprint. 
+
+1.	In the **Details** panel of the **Widget** component under **User Interface**, change the **Widget Class** to **InteractiveWidget**. 
+
+	![](WidgetInteraction_19.png)
+
+1.	In the **Content Browser**, drag the **ExampleWidget** Blueprint into the level and scale, rotate and position as desired. 
+
+	![](WidgetInteraction_20.png)
+
+1.	Click **Play** to play in the editor. 
+
+	
+## End Result
+
+When you play in the editor, you should have something similar to below:
+
+[OBJECT:EmbeddedVideo]
+[PARAMLITERAL:width]
+640
+[/PARAMLITERAL]
+[PARAMLITERAL:height]
+360
+[/PARAMLITERAL]
+[PARAMLITERAL:videoid]
+bOZy7H9IMOI
+[/PARAMLITERAL]
+[/OBJECT]
+
+You can use the **Left Mouse Button** to fire your weapon and when pointing at the button, you enter the **Hovered** state (causing it to change its style). 
+When you **Right Mouse Button** click on the button, the text value will increase each time it is pressed. 
+
+As noted during the steps, we are using the Right Mouse Button to register click events however you can use gamepad button clicks to register Pointer Key Presses/Releases or (as shown below) Motion Controller Trigger presses. 
+
+![](WidgetInteraction_21.png)
+
+Here we point our Right Motion Controller (with an attached Widget Interaction component) at the button and press the Right Trigger to register a button click.
diff --git a/Engine/Documentation/Source/Engine/UMG/HowTo/UMGHowTo.INT.udn b/Engine/Documentation/Source/Engine/UMG/HowTo/UMGHowTo.INT.udn
--- a/Engine/Documentation/Source/Engine/UMG/HowTo/UMGHowTo.INT.udn
+++ b/Engine/Documentation/Source/Engine/UMG/HowTo/UMGHowTo.INT.udn
@@ -28,6 +28,9 @@ Refer to the table below for more information and click each link to view its co
 ## Advanced
 [REGION:wrapper transparent]
 [](FANCY:Engine\UMG\HowTo\EventBasedUI "%Engine\UMG\HowTo\EventBasedUI:description%")
+[PUBLISH:docs]
+[](FANCY:Engine\UMG\HowTo\VirtualKeyboards "%Engine\UMG\HowTo\VirtualKeyboards:description%")
+[/PUBLISH:docs]
 [/REGION]
 
 
diff --git a/Engine/Documentation/Source/Engine/UMG/HowTo/VirtualKeyboards/ActorSetup/ActorSetup.INT.udn b/Engine/Documentation/Source/Engine/UMG/HowTo/VirtualKeyboards/ActorSetup/ActorSetup.INT.udn
new file mode 100644
--- /dev/null
+++ b/Engine/Documentation/Source/Engine/UMG/HowTo/VirtualKeyboards/ActorSetup/ActorSetup.INT.udn
@@ -0,0 +1,116 @@
+Availability: Docs
+Title:3. Actor Setup & Testing
+Crumbs: %ROOT%, Engine, Engine/UMG, Engine/UMG/HowTo
+Description:Demonstrates how you can create and interaction with a 3D Virtual Keyboard with UMG. 
+Related: Engine\UMG\UserGuide\WidgetInteraction\
+Related: Engine/UMG
+Related: Engine/UMG/UserGuide
+Version: 4.13
+SkillLevel: Advanced
+
+[VAR:Steps]
+[OBJECT:Navigation]
+	[PARAM:previous]
+		[Previous Step](Engine\UMG\HowTo\VirtualKeyboards\WidgetSetup "%Engine\UMG\HowTo\VirtualKeyboards\WidgetSetup:title%")
+	[/PARAM]
+	[PARAM:current]
+	[/PARAM]
+	[PARAM:home]
+		[Creating Virtual Keyboards](Engine\UMG\HowTo\VirtualKeyboards)
+	[/PARAM]
+	[PARAM:next]
+	[/PARAM]
+[/OBJECT]
+[/VAR]
+
+%Steps%
+
+In this final section, we add our Widget Blueprints as **Widget** Components to our Actors and define what happens when the player clicks on the **Enter** Button, checking to see if the entered text is the same as our **GoalText**. 
+With everything in place, we give it a test to see the end result which prints **Success** to the screen if the code is correct or **Wrong Code** if the code is incorrect.  
+
+
+## Steps
+
+1.	Open the **KeypadDisplay** Blueprint and add a **Widget** component with class set to **KeypadDisplayWB** and **Draw at Desired Size**. 
+
+	![](Keypad15.png)
+
+	Draw at Desired Size is a quick way to cause the render target to automatically match the desired size inside the Widget Blueprint.
+
+1.	Open the **Keypad** Blueprint and add a **Widget** component with class set to **KeypadWB** and **Draw at Desired Size**. 
+
+1.	Uncheck the **Window Focusable** option for this Widget Component. 
+
+	![](Keypad16.png)
+
+	We don't want this component to take focus away from the Keypad Display widget, so this needs to be unchecked. 
+
+1.	Also inside **Keypad**, add a **WidgetInteraction** component with **Enable Hit Testing** set to false. 
+
+	![](Keypad17.png)
+
+	We add this WidgetInteraction component so that we don't have to get a reference to the WidgetInteraction component in the player character and we disable Hit Testing as we don't want it to generate any hover effects. 
+
+	WidgetInteraction components by default are set to **Virtual User Index 0** when they are created and all of them communicate using the Virtual User Index ID. 
+	Since the WidgetInteraction component in the player character is also using Virtaul User Index 0, this is the same as getting a reference to it. 
+
+1.	On the Graph of **Keypad**, off **Event Begin Play** take the **Widget** component and **Get User Widget Object** to **Cast to KeypadWB**. 
+
+	![](Keypad18.png)
+
+1.	Off the **As Keypad WB** pin, **Set Widget Interaction** and pass through the **WidgetInteraction** component from Step 4. 
+
+1.	Also off the **As Keypad WB** pin, call the **Assign Enter Pressed** which will create an **Enter Pressed** Custom Event. 
+
+	![](Keypad19.png)
+
+	Now when we click on the Enter button this custom event will be called which we can use to check if the correct code was entered. 
+
+1.	Create a **Variable** of the **Keypad Display** Actor type called **TargetDisplay** and set it to **Editable**. 
+
+	![](Keypad20.png)
+
+1.	Create a **Variable** of the **Text** type called **GoalText**, then **Compile** and set the **Default Value** to **1212**. 
+
+	![](Keypad21.png)
+
+1.	Hold **Ctrl** and drag in **TargetDisplay** to get the **Widget** and **User Widget Object** then **Cast to KeypadDisplayWB**. 
+
+	![](Keypad22.png)
+
+1.	Off the **As Keypad Display WB** pin, get the **Editable Text Box** and **GetText(Text Box)** to see if it **Equals** **Goal Text**. 
+
+	![](Keypad23.png)
+
+1.	Add a **Branch** with two **Print Strings** where off **True** (use Success) and **False** (use Wrong Code). 
+
+	![](Keypad24.png)
+
+1.	Drag the **Keypad** and **KeypadDisplay** Blueprints into the level and rotate/position them as desired. 
+
+1.	In the **Details** panel for the **Keypad** Actor, set the **Target Display** variable to the **KeypadDisplay** in the level. 
+
+	![](Keypad25.png)
+
+1.	Play in the level. 
+
+## End Result
+
+Approach the keypad display and **Right-click** in it to focus on it. You can look at the buttons and **Right-click** on them to issue the **Send Key Char** command and pass through the value to the display. 
+We are using the Right Mouse button to interact with the widgets in the level, but you could swap this to use gamepad or Motion Controller button presses as well. 
+
+[OBJECT:EmbeddedVideo]
+[PARAMLITERAL:width]
+640
+[/PARAMLITERAL]
+[PARAMLITERAL:height]
+360
+[/PARAMLITERAL]
+[PARAMLITERAL:videoid]
+N2Fo4izRgh0
+[/PARAMLITERAL]
+[/OBJECT]
+
+While this example only prints to the screen if the correct code was entered or not, you could instead fire off some event that opens a door, spawns an enemy or any other form of gameplay releated event! 
+
+
diff --git a/Engine/Documentation/Source/Engine/UMG/HowTo/VirtualKeyboards/CharacterSetup/CharacterSetup.INT.udn b/Engine/Documentation/Source/Engine/UMG/HowTo/VirtualKeyboards/CharacterSetup/CharacterSetup.INT.udn
new file mode 100644
--- /dev/null
+++ b/Engine/Documentation/Source/Engine/UMG/HowTo/VirtualKeyboards/CharacterSetup/CharacterSetup.INT.udn
@@ -0,0 +1,63 @@
+Availability: Docs
+Title:1. Character Setup
+Crumbs: %ROOT%, Engine, Engine/UMG, Engine/UMG/HowTo
+Description:Here we add the Widget Interaction component to our character and create the other assets we will need. 
+Version: 4.13
+SkillLevel: Advanced
+
+[VAR:Steps]
+[OBJECT:Navigation]
+	[PARAM:previous]
+	[/PARAM]
+	[PARAM:current]
+	[/PARAM]
+	[PARAM:home]
+		[Creating Virtual Keyboards](Engine\UMG\HowTo\VirtualKeyboards)
+	[/PARAM]
+	[PARAM:next]
+		[Next Step](Engine\UMG\HowTo\VirtualKeyboards\WidgetSetup "%Engine\UMG\HowTo\VirtualKeyboards\WidgetSetup:title%")
+	[/PARAM]
+[/OBJECT]
+[/VAR]
+
+%Steps%
+
+In this step we prep our project and add the **Widget Interaction** component to our player character so we can interact with 3D widgets. 
+
+[REGION:note]
+For this how-to, we are using the **Blueprint First Person Template** with **Starter Content** enabled. 
+[/REGION]
+
+## Steps
+
+1.	Inside the **Content/FirstPersonBP/Blueprints** folder, open the **FirstPersonCharacter** Blueprint. 
+
+1.	Add a **WidgetInteraction** component attached to the **Sphere** then zero out the **Location** and **Show Debug** in the **Details** panel. 
+
+	![](Keypad01.png)
+
+	Our Widget Interaction component will now point in the direction that player points the gun indicated in-game by a debug marker.
+
+1.	In the **Graph**, add a **Right Mouse Button** Event and drag in the **Widget Interaction** Component. 
+
+1.	Off **Widget Interaction** add **Press Pointer Key** and **Release Pointer Key** both set to **Left Mouse Button** and connect as shown. 
+
+	![](Keypad02.png)
+
+	Widgets respond to mouse/touch interaction via Left Mouse Button click or touch controls and we can simulate that using the **Press Pointer Key** function. 
+	In our example we use the **Right Mouse Button** to simulate a Left Mouse Button click, however you could use a gamepad button press or even Motion Controller Trigger press to fire off the Press Pointer Key function. 
+
+1.	In the **Content Browser** create two **Class Blueprints** based on **Actor** and two **Widget Blueprints**. 
+
+	![](Keypad03.png)
+
+	For the two Actor Blueprints, call one **Keypad** (to represent the keypad entry) and the other **KeypadDisplay** (to represent and display the values entered from the Keypad). 
+	For the Widget Blueprints use the same names as the Actor Blueprints however to distinguish them, suffix them with **WB** (for Widget Blueprint). 
+
+
+## End Result
+
+Our project is ready to go and our character can interact with Widgets in the world, next we'll set up our interactive Widgets. 
+
+%Steps%
+
diff --git a/Engine/Documentation/Source/Engine/UMG/HowTo/VirtualKeyboards/CreatingVirtualKeyboards.INT.udn b/Engine/Documentation/Source/Engine/UMG/HowTo/VirtualKeyboards/CreatingVirtualKeyboards.INT.udn
new file mode 100644
--- /dev/null
+++ b/Engine/Documentation/Source/Engine/UMG/HowTo/VirtualKeyboards/CreatingVirtualKeyboards.INT.udn
@@ -0,0 +1,53 @@
+Availability: Docs
+Title:Creating Virtual Keyboard Interactions
+Crumbs: %ROOT%, Engine, Engine/UMG, Engine/UMG/HowTo
+Description:Demonstrates how you can create and interaction with a 3D Virtual Keyboard with UMG. 
+Version: 4.13
+SkillLevel: Advanced
+parent:Engine/Editor
+order:1
+checkpoint: editorqs
+
+If you are looking to create a 3D keypad or virtual keyboard that allows a player to enter specified characters, you can do so with a combination of the [Widget Interaction](Engine\UMG\UserGuide\WidgetInteraction) component and [Widget](Engine/Components/Widget) components. 
+The Widget Interaction component will help facilitate which "buttons" are clicked on and processes sending the characters through to a Widget component in an Actor that will handle the display.
+
+In this how-to we create a keypad that can be interacted with in and provides entries to a text box. 
+
+[OBJECT:EmbeddedVideo]
+[PARAMLITERAL:width]
+640
+[/PARAMLITERAL]
+[PARAMLITERAL:height]
+360
+[/PARAMLITERAL]
+[PARAMLITERAL:videoid]
+N2Fo4izRgh0
+[/PARAMLITERAL]
+[/OBJECT]
+
+_At the end of this guide, you will have a keypad and display that appears inside the 3D game world and can be interacted with._
+
+One important aspect in setting this up will be to determine the Widgets that are (or are not) focusable. If the Text Box that displays the characters the player enters is set to focusable and the buttons that the player clicks on to send characters through is also focusable, the Text Box would lose player focus when a button (or any other widget) is clicked on and not receive the commands issued. 
+
+The first step however is to set up our character with the Widget Interaction component and create the assets to display our UI. 
+
+[OBJECT:TopicButtonList]
+	[PARAM:icon]
+		![](%ROOT%/build_icon.png)(convert:false)
+	[/PARAM]
+	[PARAM:title]
+		Steps
+	[/PARAM]
+	[PARAM:description]
+	[/PARAM]
+	[PARAM:links]
+	* [](Engine\UMG\HowTo\VirtualKeyboards\CharacterSetup)
+	* [](Engine\UMG\HowTo\VirtualKeyboards\WidgetSetup)
+	* [](Engine\UMG\HowTo\VirtualKeyboards\ActorSetup)
+	[/PARAM]
+[/OBJECT]
+
+[REGION:call_to_action]
+	[Click to Start](Engine\UMG\HowTo\VirtualKeyboards\CharacterSetup "%Engine\UMG\HowTo\VirtualKeyboards\CharacterSetup:title%")
+[/REGION]
+
diff --git a/Engine/Documentation/Source/Engine/UMG/HowTo/VirtualKeyboards/WidgetSetup/WidgetSetup.INT.udn b/Engine/Documentation/Source/Engine/UMG/HowTo/VirtualKeyboards/WidgetSetup/WidgetSetup.INT.udn
new file mode 100644
--- /dev/null
+++ b/Engine/Documentation/Source/Engine/UMG/HowTo/VirtualKeyboards/WidgetSetup/WidgetSetup.INT.udn
@@ -0,0 +1,98 @@
+Availability: Docs
+Title:2. Widget Blueprint Setup
+Crumbs: %ROOT%, Engine, Engine/UMG, Engine/UMG/HowTo
+Description:Here we create the keypad and keypad display widgets that will appear in the level. 
+Version: 4.13
+SkillLevel: Advanced
+
+[VAR:Steps]
+[OBJECT:Navigation]
+	[PARAM:previous]
+		[Previous Step](Engine\UMG\HowTo\VirtualKeyboards\CharacterSetup "%Engine\UMG\HowTo\VirtualKeyboards\CharacterSetup:title%")
+	[/PARAM]
+	[PARAM:current]
+	[/PARAM]
+	[PARAM:home]
+		[Creating Virtual Keyboards](Engine\UMG\HowTo\VirtualKeyboards)
+	[/PARAM]
+	[PARAM:next]
+		[Next Step](Engine\UMG\HowTo\VirtualKeyboards\ActorSetup "%Engine\UMG\HowTo\VirtualKeyboards\ActorSetup:title%")
+	[/PARAM]
+[/OBJECT]
+[/VAR]
+
+%Steps%
+
+In this step we create the layout of our interactive UI Widgets and script what happens when clicking a button. 
+
+
+## Steps
+
+1.	Open the **KeypadDisplayWB** and replace the **Canvas Panel** in the Hierarchy with a **TextBox** wrapped with a **Size Box**. 
+
+	![](Keypad04a.png)
+
+1.	For the **Size Box**, set to **Desired on Screen** and override the **Width** and **Height** to **100x50**. 
+
+	![](Keypad04.png)
+
+1.	For the **Text Box**, set **Hint Text** to **1212** and adjust the **Font Size** to **28**. 
+
+	![](Keypad05.png)
+
+	For example, the player must use the keypad to enter the text 1212 to signal a successful entry but you could use any value you wish. 
+
+1.	Open the **KeypadWB** and set the **Hierarchy** up to create 3 Buttons with the text **0**, **1** and **2** on them as shown below.  
+
+	![](Keypad06.png)
+
+	You may also want to press **F2** on each of the Buttons in the Hierarchy to change their names to reflect which button they are. 
+	
+1.	Add a second **Horizontal Box** containing two **Buttons** with text on them, one set to **Del** the other set to **Ent**. 
+
+	![](Keypad07.png)
+
+	Our keypad will only have three numerical buttons and two others for deleting (**Del**) or committing (**Ent**) entries. 
+
+1.	Hold **Ctrl** and select each of your Buttons, then set the **Is Focusable** option to **False** for each. 
+
+	![](Keypad08.png)
+
+	In order for our setup to work, we need to ensure that no other widgets we interact with will take focus away from the Keypad Display widget. 
+	If we left these set to true, when we go to click on a button, focus will be taken away from the Keypad Display and shifted to the button preventing us from passing the key press and value over to the Keypad Display. 
+
+1.	On the **Graph**, add **OnClicked** Events for each of the Buttons. 
+
+	![](Keypad09.png)
+
+1.	Create two **Variables**, a **String** type called **CharToSend** and a **WidgetInteractionComponent** called **WidgetInteraction**. 
+
+	![](Keypad10.png)
+
+1.	Hold **Alt** and drag in **CharToSend** and connect to the **OnClicked** Events then set the variable value as shown. 
+
+	![](Keypad11.png)
+
+1.	Hold **Ctrl** and drag in **WidgetInteraction** then off it **Send Key Char** and connect **CharToSend** and the **OnClicked** Events. 
+
+	![](Keypad12.png)
+
+	We use the **Send Key Char** function here as Text Boxes in Slate do not get Keys sent directly to them, for example they don't recognize that the A key was pressed so send A through.
+	Instead they listen for Send Char, that way the lower level system operating code can do things like determine if Shift was held down while pressing the A key or if Caps Lock was on while pressing A. 
+
+1.	Off the **Del** OnClicked Event, use the **Press and Release Key** set to **Backspace** from the **WidgetInteraction** variable. 
+
+	![](Keypad13.png)
+
+1.	Create an **Event Dispatcher** called **EnterPressed** and call it from **OnClicked(Ent)**. 
+
+	![](Keypad14.png)
+
+
+## End Result
+
+Our Widget Blueprints are set up to display our keypad and keypad display, next we will add them as Widget Components to our Actor Blueprints so that we can have them appear inside our level. 
+By the end of the next step we should be able to test out our system and interact with both elements during gameplay. 
+
+%Steps%
+
diff --git a/Engine/Documentation/Source/Engine/UMG/UserGuide/WidgetInteraction/WidgetInteraction.INT.udn b/Engine/Documentation/Source/Engine/UMG/UserGuide/WidgetInteraction/WidgetInteraction.INT.udn
new file mode 100644
--- /dev/null
+++ b/Engine/Documentation/Source/Engine/UMG/UserGuide/WidgetInteraction/WidgetInteraction.INT.udn
@@ -0,0 +1,82 @@
+Availability:Docs
+Title: Widget Interaction Component
+Crumbs:%ROOT%, Engine, Engine/UMG, Engine/UMG/UserGuide
+Description:Outlines how to use the Widget Interaction Component to interact with 3D Widgets.
+Version: 4.13
+Related: Engine/UMG/UserGuide
+Related: Engine/Blueprints
+Related: Engine/Components/Widget
+
+[REGION:warning]
+This feature is considered experimental and some aspects may not work as expected and may change in future updates. 
+[/REGION]
+
+[REGION:banner]
+![](WidgetInteractionBanner.png)
+[/REGION]
+
+If you are using a [Widget Component](Engine/Components/Widget) to display UI that exists in 3D in your game world and you want to allow players to interact with that widget, the **Widget Interaction** component provides a way for that interaction to occur. 
+
+The Widget Interaction component performs a [Raycast](Gameplay/HowTo/UseRaycasts/) to see if it hits a Widget Component placed in the world and if it hits one, you can set up rules to determine how to interact with it. 
+The interaction is performed by simulating a defined key press, for example if a Button can be clicked on with the Left Mouse Button, you could tell other forms of input to simulate a Left Mouse Button click (controller button press, motion controller trigger press, etc.)
+
+[REGION:note]
+For step-by-step examples, see [](Engine\UMG\HowTo\InWorldWidgetInteraction\) and [](Engine\UMG\HowTo\VirtualKeyboards). 
+[/REGION]
+
+## Adding a Widget Interaction Component
+
+Typically you will want to add the Widget Interaction component to your player [](Gameplay/Framework/Pawn) or [](Gameplay/Framework/Pawn/Character) class from the [](Engine/Blueprints/UserGuide/Components/). 
+
+![](WidgetInteractionComponent.png)
+
+Above we add and attach it to a gun our character is carrying so that where the gun points, the Widget Interaction component points. 
+
+
+## Widget Interaction Property Reference
+
+Once you have added a Widget Interaction component, you can adjust several properties forom the **Details** panel to define how it functions. 
+Aside from common component properties like **Transform**, **Rendering** or **Sockets**, the properties below are specifically used with the Widget Interaction Component. 
+
+| Option | Description |
+| --- | ----------- |
+| [REGION:tablesection]**Interaction**[/REGION] ||
+| **Virtual User Index** | Represents the Virtual User Index. Widget Interaction Components operate by using this Virtual User Index which captures and handles focus states separately. Each Virtual User should be represented by a different Index to ensure that they will maintain separate capture and focus states for each. As each Widget Interaction component comes online it informs Slate of the Virtual User Index it has been assigned and can send events as if it were a real Slate user. |
+| **Pointer Index** | Each user Virtual Controller or virtual Finger Tips being simulated should use a different Pointer Index.|
+| **Interaction Distance** | The distance in game units the component should be able to interact with a widget component. |
+| **Interaction Source** | Determines where we should project from and start the trace (World, Mouse, Center Screen or Custom). If you set this to Custom, you'll need to call `SetCustomHitResult()` and provide the result of a custom hit test from whatever location you wish. |
+| **Enable Hit Testing** | Should the interaction component perform hit testing (Automatic or Custom) and attempt to simulate hover - if you were going to emulate a keyboard you would want to turn this option off if the virtual keyboard was separate from the virtual pointer device and used a secon interaction component.|
+| [REGION:tablesection]**Debugging**[/REGION] ||
+| **Show Debug** | Shows debugging lines and a hit sphere to help debug interaction. |
+| **Debug Color** | Determines the color of the debug lines when **Show Debug** is enabled. |
+| [REGION:tablesection]**Events**[/REGION] ||
+| **On Hovered Widget Changed** | Called when the hovered Widget Component changes. The interaction component functions at the Slate level - so it's unable to report anything about what Widget is under the hit result.|
+
+
+## Widget Interaction Blueprint Node Reference
+
+The Widget Interaction component can simulate different types of input methods (such as press, release or press + release) which you can find under the **Interaction** section in the Blueprint Context Menu. 
+You can also get information such as the last "hit" position registered by the Widget Interaction component or if a Widget Component in the world is currently being hovered upon. 
+
+![](WidgetInteractionBlueprint.png)
+
+| Node | Description |
+| --- | ----------- |
+| **Get Hovered Widget Component** | Gets the currently hovered Widget Component. |
+| **Get Last Hit Result** | Gets the last hit result generated by the component. Returns the custom hit result if that was set. |
+| **Is Over Focusable Widget** | Returns true if a widget under the hit result is focusable. For example, Slate widgets that return true for `SupportsKeyboardFocus()`.|
+| **Is Over Hit Test Visible Widget** | Returns true if a widget under the hit result has a visibility setting that makes it hit test visible. For example, Slate widgets that return true for `GetVisibility().IsHitTestVisible()`.|
+| **Is Over Interactable Widget** | Returns true if a widget under the hit result is interactive. For example, Slate widgets that return true for `IsInteractble()`.|
+| **Press and Release Key** | Performs both the press and release of a simulated keyboard key.| 
+| **Press Key** | Performs a press of a simulated keyboard key. Avoid using this for `a-z|A-Z`, things like the Editable Textbox in Slate expect **OnKeyChar** to be called to signal a specific character being sent to the widget so for those cases you would use **SendKeyChar** instead. |
+| **Press Pointer Key** | Performs a press of a simulated mouse/pointer event (pointer could be mouse, could be touch, etc). Normally you would just use Left/Right Mouse Button for the Key, however advanced usage could be imagined where you send other keys to signal widgets to take special actions if they're under the cursor.  |
+| **Release Key** | Releases a key as if it had been released from the keyboard.  |
+| **Release Pointer Key** | Releases a key as if the mouse/pointer where the source of it. Similar to Press Pointer Key, normally you would use Left/Right Mouse Button for the key however advanced use cases are possible. |
+| **Scroll Wheel** | Sends a scroll wheel event to the widget under the last hit result. |
+| **Send Key Char** | Transmits a list of characters to a widget by simulating a OnKeyChar event for each key listed in the string. |
+| **Set Custom Hit Result** | Sets a custom hit result. This is only take into account if the Interaction Source is set to Custom.  |
+
+The properties for the Widget Interaction component can also be reterived in the form of a getter or set via setter nodes. 
+
+![](WidgetInteractionVariables.png)
+
