Various fixes during peer review process.

#UE4doc

diff --git a/Engine/Documentation/Source/Platforms/Android/GettingStarted/1/AndroidQS_1.INT.udn b/Engine/Documentation/Source/Platforms/Android/GettingStarted/1/AndroidQS_1.INT.udn
--- a/Engine/Documentation/Source/Platforms/Android/GettingStarted/1/AndroidQS_1.INT.udn
+++ b/Engine/Documentation/Source/Platforms/Android/GettingStarted/1/AndroidQS_1.INT.udn
@@ -4,6 +4,7 @@ Crumbs:%ROOT%, Platforms, Platforms/Android, Platforms/Android/GettingStarted
 Description: Install the CodeWorks for Android 1R4 SDK.
 Version:4.9
 Skilllevel:Intermediate
+Tags:Android
 
 [VAR:Steps]
 [OBJECT:Navigation]
@@ -69,10 +70,10 @@ If you already have CodeWorks for Android 1R4 installed, but your environment va
 [/REGION]
 
 [region:warning]
-It is highly recommended that you install CodeWorks for Android 1R4 while the editor and launcher are both closed to ensure that there are not problems with the install.
+It is highly recommended that you install CodeWorks for Android 1R4 while the editor and launcher are both closed to ensure that there are no problems with the install.
 [/region]
 
-1. To install CodeWorks for Android 1R4 navigate to the **Engine** > **Extras** > **Android** folder. 
+1. To install CodeWorks for Android 1R4, navigate to the **Engine** > **Extras** > **Android** folder. 
 
 	[region:note]
 	The CodeWorks for Android 1R4 installer requires an installation of Visual Studio to complete successfully. Supported versions are 2010, 2012 or 2013 (not 2015).
@@ -85,20 +86,20 @@ It is highly recommended that you install CodeWorks for Android 1R4 while the ed
 
 	[region:note]
 	If you downloaded UE4 from **Github** you will need to perform the following steps before you can run the CodeWorks for Android 1R4 installer. 
-	Please note that the follow applies to both Windows and Mac based PCs that download the UE4 source code from Github.
-	If you are using the binary build downloaded from the launcher then you can skip this part and go right into installing CodeWorks for Android 1R4. 
+	Please note that the following applies to both Windows and Mac based PCs that download the UE4 source code from Github.
+	If you are using the binary build downloaded from the launcher, then you can skip this part and go right into installing CodeWorks for Android 1R4. 
 
 	1. Once UE4 has been downloaded from Github and extracted, run the **Setup.bat** file to download the required files as well as the **CodeWorks for Android 1R4** installer.
 
-	1. After the **Setup.bat** has finished running locate and run the CodeWorks for Android 1R4 installer following the instructions below.
+	1. After the **Setup.bat** has finished running, locate and run the CodeWorks for Android 1R4 installer, following the instructions below.
 
 	[/region]
 
-1. **Double-click** on the **AndroidWorks-1R1-windows.exe** file to initiate the install process.
+1. Double-click on the **AndroidWorks-1R1-windows.exe** file to initiate the install process.
 
 	![](AndroidWorks_11.png)
 
-1. Click the **Next** button on the first window to begin the install process.
+1. Click the **Next** button on the first window to begin the installation process.
     
 	![](AndroidWorks_00.png)
 
@@ -124,12 +125,12 @@ It is highly recommended that you install CodeWorks for Android 1R4 while the ed
 
 1.	A new window will be displayed that will show you the download and install progress of the required CodeWorks for Android 1R4 tools. 
 	[region:note]
-	You will need an internet connection to complete this portion. Depending on the speed of that internet connection this could take a few minute to a few hours to complete.
+	You will need an internet connection to complete this portion. Depending on the speed of your internet connection, this could take a few minute to a few hours to complete.
 	[/region]
 	
 	![](AndroidWorks_06.png)
 
-1.	When all required tools have been downloaded and installed the following window will be displayed. Press the **Next** button continue the install process.
+1.	When all required tools have been downloaded and installed, the following window will be displayed. Press the **Next** button continue the install process.
 	
 	![](AndroidWorks_08.png)
 
@@ -157,7 +158,7 @@ If you already have CodeWorks for Android 1R4 installed, but your environment va
 [/REGION]
 
 [region:warning]
-It is highly recommended that you install CodeWorks for Android 1R4 while the editor and launcher are both closed to ensure that there are not problems with the install.
+It is highly recommended that you install CodeWorks for Android 1R4 while the editor and launcher are both closed to ensure that there are no problems with the install.
 [/region]
 Before you can install CodeWorks for Android 1R4 to your Mac you will need to make sure that both Java 6 and Java 7 are installed. 
 Use the links provide below to download and install the correct versions of Java.
@@ -176,17 +177,17 @@ Use the links provide below to download and install the correct versions of Java
 	Please note that the follow applies to both Windows and Mac based PCs that download the UE4 source code from Github.
 	If you are using the binary build downloaded from the launcher then you can skip this part and go right into installing CodeWorks for Android 1R4. 
 
-	1. Once UE4 has been downloaded from Github and extracted, run the **Setup.bat** file to download the required files as well as the CodeWorks for Android 1R4 installer.
+	1. Once UE4 has been downloaded from Github and extracted, run the **Setup.command** file to download the required files as well as the CodeWorks for Android 1R4 installer.
 
-	1. After  **Setup.bat** has finished running locate and run the CodeWorks for Android 1R4 installer following the instructions below.
+	1. After **Setup.command** has finished running locate and run the CodeWorks for Android 1R4 installer following the instructions below.
 
 	[/region]
 
-1. **Double-click** on the **AndroidWorks-1R1-osx.dmg** to load the CodeWorks for Android 1R4 installer.
+1. Double-click on the **AndroidWorks-1R1-osx.dmg** to load the CodeWorks for Android 1R4 installer.
 	
 	![](AndroidWorks_Mac_00.png)
 
-1. **Double-click** the CodeWorks for Android 1R4 installer icon to start the install process.
+1. Double-click the CodeWorks for Android 1R4 installer icon to start the installation process.
 
 	![](AndroidWorks_Mac_01.png)
 
@@ -198,7 +199,7 @@ Use the links provide below to download and install the correct versions of Java
 
 	![](AndroidWorks_Mac_03.png)
 
-1. Specify a location where CodeWorks for Android 1R4 should be installed then press the **Next** button.
+1. Specify a location where CodeWorks for Android 1R4 should be installed, then press the **Next** button.
 
 	[region:note]
 	It is perfectly fine to leave this set to the default location.
@@ -216,11 +217,11 @@ Use the links provide below to download and install the correct versions of Java
 
 1. The CodeWorks for Android 1R4 Component Manager will now begin to download and install the required tools. 
 	[region:note]
-	For this part of the setup will take a few minutes to completed and you will need an internet connection to complete it.
+	You will need an internet connection to complete this portion. Depending on the speed of your internet connection, this could take a few minute to a few hours to complete.
 	[/region]
 	![](AndroidWorks_Mac_07.png)
 
-1. Once everything has been installed the following window will be displayed. Click on the **Next** button to continue.
+1. Once everything has been installed, the following window will be displayed. Click on the **Next** button to continue.
 
 	![](AndroidWorks_Mac_08.png)
 
@@ -229,7 +230,7 @@ Use the links provide below to download and install the correct versions of Java
 	![](AndroidWorks_Mac_09.png)
 
 	[region:note]
-	If you downloaded UE4 from Github, when CodeWorks for Android 1R4 has finished installing, restart your Mac so the environment variables are set and then run the **Setup.bat** file one more time to download the required Android Development files.
+	If you downloaded UE4 from Github, when CodeWorks for Android 1R4 has finished installing, restart your Mac so the environment variables are set and then run the **Setup.command** file one more time to download the required Android Development files.
 	[/region]
 
 [/EXCERPT:MacOSSetup]
diff --git a/Engine/Documentation/Source/Platforms/Android/Reference/AndroidGettingStarted.INT.udn b/Engine/Documentation/Source/Platforms/Android/Reference/AndroidGettingStarted.INT.udn
--- a/Engine/Documentation/Source/Platforms/Android/Reference/AndroidGettingStarted.INT.udn
+++ b/Engine/Documentation/Source/Platforms/Android/Reference/AndroidGettingStarted.INT.udn
@@ -84,7 +84,7 @@ There are several methods for getting the components. If you have already instal
 <!--
 [EXCERPT:TADPWindowsSetup]
 
-CodeWorks for Android 1R4, by Nvidia is a one-stop shop for Android development needs. If you install CodeWorks for Android 1R4, all of your requirements to develop for Android will be met.
+CodeWorks for Android 1R4 (by Nvidia) is a one-stop shop for Android development needs. If you install CodeWorks for Android 1R4, all of your requirements to develop for Android will be met.
 
 The installer for CodeWorks for Android 1R4 is included with Unreal Engine 4 for Windows. It can be found here: 
 
@@ -206,7 +206,7 @@ If you continue seeing the message about the Android SDK not being installed pro
 
 To set an environment variable:
 
-* **Right-click** on your **Computer** icon OR press **Windows-Break** on your keyboard.
+* Right-click on your **Computer** icon OR press **Windows-Break** on your keyboard.
 * Select **Advanced system settings** on the left.
 * Click **Environment Variables...**
 * For each variable (see below):
@@ -243,7 +243,7 @@ Unless you installed CodeWorks for Android 1R4, you will need to verify that som
 Once you have completed the installation (either by editing your .bash_profile file, or by installing CodeWorks for Android 1R4), you will need to restart both the Editor and Launcher applications.
 [/region]
 
-To set an environment variable you need to edit your **.bash_profile** file in your home directory. 
+To set an environment variable you need to edit the **.bash_profile** file in your home directory. 
 The file is hidden in Finder, so you’ll need to use Terminal to open it for editing. 
 If you’re not familiar with Unix terminal commands please do the following.
 
@@ -308,12 +308,12 @@ Here are the name of the variables that we depend on and where they should point
 [EXCERPT:ATTWindowsSetup]
 If you do not see your device listed in the **Launch** menu, here are some things to try: 
 - Wait a few seconds and look again, as the background scanning for connected devices can take a short amount of time. 
-- Check your USB debugging connection is all working. 
+- Check that your USB debugging connection is working. 
 - Open a command prompt (Press **Windows+R**, and run `cmd.exe`) 
 	- Type the following command: `adb devices`
 	- If you do not see any devices listed, your USB driver is not installed properly. How to install the driver depends on the model of your Android device. Search online for your device model and 
 	"usb driver", then install the driver. 
-	- Additionally, you will need to make sure your device has Developer mode enabled, and USB debugging is turned on. See the [Android Developer Using Hardware Devices guide](http://developer.android.com/tools/device.html) for more information (specifically items 2 and 3).  
+	- Additionally, you will need to make sure your device has Developer mode enabled with USB debugging is turned on. See the [Android Developer Using Hardware Devices guide](http://developer.android.com/tools/device.html) for more information (specifically items 2 and 3).  
 	- Afterwards, make sure the `adb devices` command works; without this, the engine will not be able to do anything with your Android device. 
 	
 If your device shows up in the Launch menu, but gets stuck deploying to the device, it is most likely because the device is not authorized for your PC to talk to it. If you click on the **Show Log** line 
@@ -327,16 +327,16 @@ in the progress bar, you may see that the last line indicates that it is calling
 [EXCERPT:ATTMacOSSetup]
 If you do not see your device listed in the Launch menu, here are some things to try: 
 * Wait a few seconds and look again, as the background scanning for connected devices can take a short amount of time to register the device. 
-* Check your USB debugging connection is all working. 
+* Check that your USB debugging connection is working.
 * Open the Terminal app and then try the following.
-	* Type the following command into the terminal **adb devices**.
+	* Type the following command into the terminal: `adb devices`.
 	* If you do not see any devices listed, your USB driver is not installed properly. How to install the driver depends on the model of your Android device. Search online for your device model and **USB driver**, then install the required driver. 
-	* Additionally, you will need to make sure your device has Developer mode enabled, and USB debugging is turned on. See the **[Android Developer Using Hardware Devices](http://developer.android.com/tools/device.html)** guide for more information (specifically items 2 and 3).
-	* Afterwards, make sure the **adb devices** command works; without this, the engine will not be able to do anything with your Android device.
+	* Additionally, you will need to make sure your device has Developer mode enabled with USB debugging is turned on. See the [Android Developer Using Hardware Devices](http://developer.android.com/tools/device.html) guide for more information (specifically items 2 and 3).
+	* Afterwards, make sure the `adb devices` command works; without this, the engine will not be able to do anything with your Android device.
 
 If your device shows up in the **Launch menu**, but gets stuck deploying to the device, it is most likely because the device is not authorized to interact with your Mac. 
-If you click on the **Show Log** line in the progress bar, and see that the last line indicates that it is calling adb uninstall try the following.
-* Open up the Terminal app and run the adb devices command. 
+If you click on the **Show Log** line in the progress bar, and see that the last line indicates that it is calling `adb uninstall`, try the following.
+* Open up the Terminal app and run the `adb devices` command. 
 * If the device is listed but shows unauthorized, that indicates that you need to get the device to ask if you want to authorize the Mac. 
 * It is recommended you check the **Always allow** check box to always allow your Mac to talk to this device.
 
@@ -350,7 +350,7 @@ Not all Android devices are made the same. In particular, there are 4 different
 [EXCERPT:AndroidFormats]
 |Format | Description|
 |-------|------------|
-|DXT |This is what your desktop computer uses, and used by Tegra-based devices.|
+|DXT |This is what your desktop computer uses (also used by Tegra-based devices).|
 |ATC |This is used in Qualcomm Snapdragon based devices.|
 |PVR |This is used in ImgTec PowerVR based devices.|
 |ETC1 |This is supported on all devices, but does not support textures with an alpha channel, so those textures are uncompressed, making your game's download large and use more runtime memory.|
@@ -365,8 +365,8 @@ the best format at runtime. It will work on any device (as will **Android (ETC1)
 
 [region:note]
 Unreal Engine version 4.8 and later will now support the use of **ES31** and **AEP** rendering on Tegra K1 devices. 
-Please note that due to the limited number of devices that currently support ES31 and AEP rendering very little performance and compatibility testing has been completed at this time.
-If you are developing a project that is aimed at using one of theses devices it is highly recommended that you perform extensive performance and compatibility testing to ensure that the device delivers your project in the manner you intended. 
+Please note that due to the limited number of devices that currently support ES31 and AEP rendering, very little performance and compatibility testing has been completed at this time.
+If you are developing a project that is aimed at using one of these devices, it is highly recommended that you perform extensive performance and compatibility testing to ensure that the device delivers your project in the manner you intended. 
 [/region]
 
 ### Choosing a Format
diff --git a/Engine/Documentation/Source/Platforms/GearVR/Debugging/GearVRDebugging.INT.udn b/Engine/Documentation/Source/Platforms/GearVR/Debugging/GearVRDebugging.INT.udn
--- a/Engine/Documentation/Source/Platforms/GearVR/Debugging/GearVRDebugging.INT.udn
+++ b/Engine/Documentation/Source/Platforms/GearVR/Debugging/GearVRDebugging.INT.udn
@@ -2,7 +2,7 @@ Availability:Public
 Title:Samsung Gear VR Debugging
 Description:Information on how to debug applications for the Gear VR headset.
 Navigation:platform
-Crumbs:
+Crumbs:%ROOT%, Platforms, Platforms/GearVR
 Platform:Gear VR
 Related: Platforms/Mobile/Performance
 Related: Engine/Performance
@@ -21,18 +21,18 @@ tags:GearVR
 [TOC(start:2 end:3)]
 
 Being able to debug your project on its target device is critical to tracking down performance issues and other possible project breaking errors. 
-In the following section you will learn what tools are available to assist you in debugging your Gear VR projects and how to use them.
+In the following section, you will learn what tools are available to assist you in debugging your Gear VR projects and how to use them.
 
 ### Enabling Gear VR Developer Mode
 Having the ability to view your Gear VR Project without having to use the Gear VR Headset can greatly speed up the time it takes to debug an issue by allowing you quick and easy access to things like Unreal's console. 
 In the following sections we will cover how you can set up your Gear VR projects so that they can be viewed without the Gear VR Headset. 
 
 [region:note]
-Before you can do this you must make sure that you have some type of Gear VR application installed to your phone that has your device signature in it. 
-This can be any kind of application, download or compiled.
+Before you can do this, you must make sure that you have some type of Gear VR application installed to your phone that has your device signature in it. 
+This can be any kind of application (download or compiled).
 [/region]
 
-1. First on your Samsung Galaxy Note 4 Smartphone go to **Settings** > **Application Manager** > **Gear VR Service**.
+1. First, on your Samsung Galaxy Note 4 Smartphone go to **Settings** > **Application Manager** > **Gear VR Service**.
 	
 	[REGION:lightbox]
 	[![](T_GearVR_Dev_Mode_01.png)(w:600 h:355)](T_GearVR_Dev_Mode_01.png)
@@ -42,7 +42,7 @@ This can be any kind of application, download or compiled.
 	Click for full image.
 	[/REGION]
 
-1. Next tap on the **Manage Storage** option.
+1. Next, tap on the **Manage Storage** option.
   
 	[REGION:lightbox]
 	[![](T_GearVR_Dev_Mode_02.png)(w:250 h:444)](T_GearVR_Dev_Mode_02.png)
@@ -52,7 +52,7 @@ This can be any kind of application, download or compiled.
 	Click for full image.
 	[/REGION]
 
-1. Then tap on the **VR Service Version** **6** times.
+1. Then, tap on the **VR Service Version** **6** times.
 
 	[REGION:lightbox]
 	[![](T_GearVR_Dev_Mode_03.png)(w:250 h:225)](T_GearVR_Dev_Mode_03.png)
@@ -72,23 +72,23 @@ This can be any kind of application, download or compiled.
 	Click for full image.
 	[/REGION]
 
-1. Now you should be able to launch your Gear VR Project and have it display without having to use the Gear VR Headset.
+1. Now, you should be able to launch your Gear VR Project and have it display without having to use the Gear VR Headset.
 
 	[region:tip]
-	Be careful when enabling VR Developer mode as it can drastically drain your phone's battery as you are overriding any settings to turn off the Gear VR Headset when not in use.
+	Be careful when enabling VR Developer mode (it can drastically drain your phone's battery) as you are overriding any settings, to turn off the Gear VR Headset when not in use.
 	[/region]
 
-1. If for any reason you need to disable VR Developer mode just slide the **Developer Mode** slider from **On** to **Off** and VR Developer mode will be disabled.
+1. If for any reason you need to disable VR Developer mode, just slide the **Developer Mode** slider from **On** to **Off** and VR Developer mode will be disabled.
 
 ### Enabling the UE4 Console
-The UE4 Console is a very powerful tool that allows you to look at and in some cases affect data that is currently being displayed in UE4. 
+The UE4 Console is a very powerful tool that allows you to look at, and in some cases, affect data that is currently being displayed in UE4. 
 Accessing the UE4 console on Mobile devices requires a slightly different approach than what you might be used to on PC and is explained below.
 
 [region:tip]
 You can get a full list of available console commands by entering **DumpConsoleCommands** into the console while on a PC and Unreal will then list all available console commands in the console.
 [/region]
 
-1. First make sure that you have enabled Gear VR Developer mode on your phone. 
+1. First, make sure that you have enabled Gear VR Developer mode on your phone. 
 	This way you can launch your project and see it without having to use the Gear VR Headset which, if used, would prevent you from being able to touch your phone's screen.
 	
 	[REGION:lightbox]
@@ -99,7 +99,7 @@ You can get a full list of available console commands by entering **DumpConsoleC
 	Click for full image.
 	[/REGION]
 
-1. Next launch your Gear VR project by tapping on the projects icon to launch it. 
+1. Next, launch your Gear VR project by tapping on the project's icon to launch it. 
 	When the project has been launched, you should see something looking like this on your phone.
 	
 	[REGION:lightbox]
@@ -111,7 +111,7 @@ You can get a full list of available console commands by entering **DumpConsoleC
 	[/REGION]
 
 1. With your Gear VR project now running, tap the phone's screen with **Four (4)** fingers at once to open up the Unreal Engine Console. 
-	If done correctly you should see something like this on your phone's screen. 
+	If done correctly, you should see something like this on your phone's screen. 
 
 	[REGION:lightbox]
 	[![](T_GearR_Open_Console.png)(w:500 h:281)](T_GearR_Open_Console.png)
@@ -142,29 +142,29 @@ You can get a full list of available console commands by entering **DumpConsoleC
 	Click for full image.
 	[/REGION]
 
-1. When you are done debugging or if you just wish to clear the screen of debug information, you can input **Stat None** to clear the screen off all debug information.
+1. When you are done debugging, or if you just wish to clear the screen of debug information, you can input **Stat None** to clear the screen of all debug information.
 
 	[region:tip]
-	The best option for console debugging with the Gear VR Headset is to use a Bluetooth keyboard. This way you can use the console without having to remove the phone from the headset.
+	The best option for console debugging with the Gear VR Headset is to use a Bluetooth keyboard. This way, you can use the console without having to remove the phone from the headset.
 	[/region]
 
 ### Viewing Gear VR Debug Logs
 Having the ability to view debug logs is of critical importance when developing any kind of 3D project. 
-Debug logs can provide vital information and insight into why your project is not working and in some case's even make suggestions as to what you can do to fix things. 
-In the following section we will go over what you need to do in order to get debug logs from your Smartphone.
+Debug logs can provide vital information and insight into why your project is not working and, in some cases, even make suggestions as to what you can do to fix things. 
+In the following section, we will go over what you need to do in order to get debug logs from your Smartphone.
 
 [region:note]
-Check out the official Android documentation for much more detailed information on all of various things you can do with [Logcat](http://developer.android.com/tools/help/logcat.html).
+Check out the official Android documentation for more detailed information on all of various things you can do with [Logcat](http://developer.android.com/tools/help/logcat.html).
 [/region]
 
-1. First make sure that your Smartphone with the application you want to get debug information from is plugged into your PC via USB.
+1. First, make sure that your Smartphone with the application you want to get debug information from is plugged into your PC via USB.
 
-1. Next open up a Windows Command Shell by first pressing the **Windows Key + R** and then entering **CMD** in the **Run** dialogue box that is opened.
+1. Next, open up a Windows Command Shell by first pressing the **Windows Key + R** and then entering **CMD** in the **Run** dialogue box that is opened.
 	Then press either the **Enter** key on your keyboard or just click on the **OK** button with your mouse. 
 	 
 	 ![](T_Open_CMD_Prompt.png)
 
-1. Once the **OK** button has been toggled a Windows Command Line window like the following should now be displayed.
+1. Once the **OK** button has been toggled, a Windows Command Line window like the following should now be displayed.
 	
 	[REGION:lightbox]
 	[![](T_Windows_CMD.png)(w:422 h:269)](T_Windows_CMD.png)
@@ -174,7 +174,7 @@ Check out the official Android documentation for much more detailed information
 	Click for full image.
 	[/REGION]
 
-1. Now check to make sure that your PC can see your Smartphone by entering **ADB Devices** into the command line window and then press the **Enter** key to display devices attached to your PC.
+1. Now, check to make sure that your PC can see your Smartphone by entering **ADB Devices** into the command line window and then press the **Enter** key to display devices attached to your PC.
 	
 	[REGION:lightbox]
 	[![](T_Windows_ADB_Devices.png)(w:422 h:269)](T_Windows_ADB_Devices.png)
@@ -188,9 +188,9 @@ Check out the official Android documentation for much more detailed information
 	If you enter **ADB Devices** but do not see anything show up, make sure that you have gone through the [](Platforms/GearVR) setup section and installed / setup everything that is listed there and then try this step again. 
 	[/region]
 
-1. Now that you have verified that you can see your Smartphone in ADB Devices it is now time to get the log information from your phone. 
-	But first we are going to clear the log cache to ensure that we are getting only current debug information and not old debug information. 
-	To clear the log enter the following command.
+1. Now that you have verified that you can see your Smartphone in ADB Devices, it is now time to get the log information from your phone. 
+	But firs,t we are going to clear the log cache to ensure that we are getting only current debug information and not old debug information. 
+	To clear the log, enter the following command.
 
 		adb logcat -c
 	[REGION:lightbox]
@@ -217,7 +217,7 @@ Check out the official Android documentation for much more detailed information
  To see the live output of the log, omit the **>logoutput.txt** part and all log updates will happen in the command line window.
  [/region]
 
-1. To complete the log capture press **CTRL + C** on the keyboard and when completed you should now have a log file named,** logoutput.txt**, that is saved to your **Users** folder in Windows.
+1. To complete the log capture, press **CTRL + C** on the keyboard and when completed, you should now have a log file named,** logoutput.txt**, that is saved to your **Users** folder in Windows.
 
 	[REGION:lightbox]
 	[![](T_ADB_Log_Location.png)(w:381 h:338)](T_ADB_Log_Location.png)
@@ -228,12 +228,12 @@ Check out the official Android documentation for much more detailed information
 	[/REGION]
 
 	[region:note]
-	Data will still continue to be output to the log file while the command line window is open and running the ADB logcat command. To stop this all you have to do is press **CTRL + C** and data will stop being written to the log file.
+	Data will still continue to be output to the log file while the command line window is open and running the ADB logcat command. To stop this, all you have to do is press **CTRL + C** and data will stop being written to the log file.
 	[/region] 
 
 1. The information that is contained within the log file was generated from adding **UE4**, **OVR**,  **VrApi**, **VrLib**,**VRSVC**, and **HmdInfo** to the logcat command. 
-	When you look at the log file that was created and search for anyone of those terms you will find all debug information related to those terms. 
-	You can see an example of this in the image below, some of the terms used, OVR, UE4, Vrapi, and VrApi all show up in the log with information that is related to those terms.
+	When you look at the log file that was created and search for any of those terms, you will find all debug information related to those terms. 
+	You can see an example of this in the image below, some of the terms used (OVR, UE4, Vrapi, and VrApi) show up in the log with information that is related to those terms.
 	
 	[REGION:lightbox]
 	[![](T_ADB_Log_Info.png)(w:516 h:349)](T_ADB_Log_Info.png)
@@ -248,15 +248,15 @@ It also provides you with detailed information about the problem that can then b
 
 ### Wireless Deploying & Debugging
 Getting debug information or deploying to your Smartphone while it is in the Gear VR Headset is impossible due to the only USB port being used by the Gear VR Headset. 
-However you can work around this issue by using your Smartphone's Wifi connection to send debug information to your development PC. 
-You can even deploy projects in this manner however it is extremely slow when compared to using the USB cable. 
-In the following section we will cover all you need to know about setting up your Smartphone so that it can send debug information and be deployed to over wifi.
+However, you can work around this issue by using your Smartphone's Wifi connection to send debug information to your development PC. 
+You can even deploy projects in this manner; however, it is extremely slow when compared to using the USB cable. 
+In the following section, we will cover all you need to know about setting up your Smartphone so that it can send debug information and be deployed over wifi.
 
 [region:note]
-Please note that in order for this to work you have to have your development PC and Galaxy Note 4 on the same network.
+Please note that in order for this to work, you have to have your development PC and Galaxy Note 4 on the same network.
 [/region]
 
-1. Next open up a Windows Command Shell by first pressing the **Windows Key + R** and then entering **CMD** in the **Run** dialogue box that is opened.
+1. Next, open up a Windows Command Shell by first pressing the **Windows Key + R** and then entering **CMD** in the **Run** dialogue box that is opened.
 	Then press either the **Enter** key on your keyboard or just click on the **OK** button with your mouse. 
 	 
 	[REGION:lightbox]
@@ -267,7 +267,7 @@ Please note that in order for this to work you have to have your development PC
 	Click for full image.
 	[/REGION]
 
-1. Once the **OK** button has been toggled a Windows Command Line window like the following should now be displayed.
+1. Once the **OK** button has been toggled, a Windows Command Line window like the following should now be displayed.
 	
 	[REGION:lightbox]
 	[![](T_Windows_CMD.png)(w:427 h:269)](T_Windows_CMD.png)
@@ -278,7 +278,7 @@ Please note that in order for this to work you have to have your development PC
 	[/REGION]
 
 
-1. Now enter the following command to see what your Smartphone's IP address is and make sure to take note of it. The IP address you want to use is the one that comes after **inet**.
+1. Now, enter the following command to see what your Smartphone's IP address is and make sure to take note of it. The IP address you want to use is the one that comes after **inet**.
 	
 		adb shell ip -f inet addr show wlan0
 	
@@ -290,7 +290,7 @@ Please note that in order for this to work you have to have your development PC
 	Click for full image.
 	[/REGION]
 
-1. Next enter the following command to change the network connection to **TCP** and set the port to listen to.
+1. Next, enter the following command to change the network connection to **TCP** and set the port to listen to.
 
 		adb tcpip 5555
 	
@@ -302,11 +302,11 @@ Please note that in order for this to work you have to have your development PC
 	Click for full image.
 	[/REGION]
 
-1. Once the port has been set, it is now time to connect to the Smartphone. To do this all you need to do is enter the following command and press the Enter key.
+1. Once the port has been set, it is now time to connect to the Smartphone. To do this, all you need to do is enter the following command and press the Enter key.
 	
 		adb connect 10.3.31.133 - Replace with your devices IP Address.
 		
-1. You should now see a message that says **connected to** and the IP address of your device.
+1. You should now see a message that says **connected to** with the IP address of your device.
 	
 	![](T_ADB_Connected.png)
 
@@ -321,8 +321,8 @@ Please note that in order for this to work you have to have your development PC
 	[/REGION]
 
 
-With your Smartphone and PC now connected to one another via the network you can now begin to get debug information from your Smartphone or even deploy projects to your phone without having to use the USB cable. 
-However keep in mind that deploying to your Smartphone from your PC over the network is four to five times slower than using the USB cable.
+With your Smartphone and PC now connected to one another via the network, you can now begin to get debug information from your Smartphone or even deploy projects to your phone without having to use the USB cable. 
+However, keep in mind that deploying to your Smartphone from your PC over the network is slower than using the USB cable.
 
 ### Return to Wired Deploying & Debugging
 1. To return to using the USB cable for debugging and deployment, input the following command into the command window.
diff --git a/Engine/Documentation/Source/Platforms/GearVR/Prerequisites/GearVRPrerequisites.INT.udn b/Engine/Documentation/Source/Platforms/GearVR/Prerequisites/GearVRPrerequisites.INT.udn
--- a/Engine/Documentation/Source/Platforms/GearVR/Prerequisites/GearVRPrerequisites.INT.udn
+++ b/Engine/Documentation/Source/Platforms/GearVR/Prerequisites/GearVRPrerequisites.INT.udn
@@ -1,7 +1,7 @@
 Title: Samsung Gear VR Headset Prerequisites
 Description:Required initial setup for the Gear VR Headset
 Availability:Public
-Crumbs:
+Crumbs:%ROOT%, Platforms, Platforms/GearVR
 Navigation:platform
 Platform:Gear VR
 Skilllevel: Intermediate
@@ -18,57 +18,55 @@ tags:GearVR
 Developing for the Gear VR Headset requires **Unreal Engine 4.7.4** or later.
 [/REGION]
 
-With the 4.7 release of Unreal Engine 4 (UE4) you can now build Virtual Reality(VR) projects for the Samsung Gear VR Headset that can transform your Samsung Galaxy Note 4, Galaxy S6, and Galaxy S6 Edge Smartphone into an immersive VR device. 
-In the following sections we will go over everything you need to know to get your Samsung Galaxy Note 4 Smartphone properly setup so that it can work with the Samsung Gear VR headset.
+With the 4.7 release of Unreal Engine 4 (UE4), you can now build Virtual Reality (VR) projects for the Samsung Gear VR Headset that can transform your Samsung Galaxy Note 4, Galaxy S6, and Galaxy S6 Edge Smartphone into an immersive VR device. 
+In the following sections, we will go over everything you need to know to get your Samsung Galaxy Note 4 Smartphone properly setup so that it can work with the Samsung Gear VR headset.
 
 ## Gear VR Official Documentation
 
 For a more in-depth look at how the Gear VR SDK works or for more information on Oculus Mobile VR development check out the official Oculus documents linked below.
 
-* **[Oculus Mobile Software Development Kit(SDK) Device and Environment Setup Guide](http://static.oculus.com/sdk-downloads/documents/Oculus_Mobile_v0.4.0_Device_and_Environment_Setup_Guide.pdf)**
-* **[Oculus Mobile Development Documentation](http://static.oculus.com/sdk-downloads/documents/Oculus_Mobile_v0.5.0_SDK_Documentation.pdf)**
-
-
+* [Oculus Mobile Software Development Kit(SDK) Device and Environment Setup Guide](http://static.oculus.com/sdk-downloads/documents/Oculus_Mobile_v0.4.0_Device_and_Environment_Setup_Guide.pdf)
+* [Oculus Mobile Development Documentation](http://static.oculus.com/sdk-downloads/documents/Oculus_Mobile_v0.5.0_SDK_Documentation.pdf)
 
 ## Gear VR Headset Pre-requisites 
 
-Before we can get into UE4 Gear VR development we first must make sure that the following prerequisites are fulfilled as if we do not, we run the risk of our Gear VR projects not working with the Gear VR Headset. 
+Before we can get into UE4 Gear VR development, we first must make sure that the following prerequisites are fulfilled because if we don't, we run the risk of having Gear VR projects not working with the Gear VR Headset. 
 
 * Make sure you have the USB cord used for charging your Samsung Galaxy Note 4 Smartphone available and plugged into a USB port on the PC you are developing on. 
-The USB cord is of vital importance as this is how you transfer data between your development PC and your Smartphone.
+The USB cord is of vital importance, as this is how you transfer data between your development PC and your Smartphone.
 
-* Only certain types of Samsung Galaxy Smartphone's will work with the Gear VR Headset. 
-	The **[Oculus website](https://support.oculus.com/hc/en-us/articles/204113096-Samsung-Gear-VR-Innovator-Edition-Samsung-Galaxy-Note-4-Compatibility)** has a list of phone models that will work with the Gear VR Headset.
+* Only certain types of Samsung Galaxy Smartphones will work with the Gear VR Headset. 
+	The [Oculus website](https://support.oculus.com/hc/en-us/articles/204113096-Samsung-Gear-VR-Innovator-Edition-Samsung-Galaxy-Note-4-Compatibility) has a list of phone models that will work with the Gear VR Headset.
 	Make sure to double check that the Samsung Galaxy Note 4 Smartphone you are buying is the correct model that will work with the Gear VR Headset.
 
-* Make sure to download and install the specific USB drivers for your Galaxy Note 4 from the **[Samsung Website](http://www.samsung.com/us/support/downloads)**. 
+* Make sure to download and install the specific USB drivers for your Galaxy Note 4 from the [Samsung Website](http://www.samsung.com/us/support/downloads). 
 
-* Download Unreal Engine version **4.7.4** or later from the **[UE4 Launcher](https://www.unrealengine.com/)** or **[Github](https://www.unrealengine.com/ue4-on-github)**. 
+* Download Unreal Engine version **4.7.4** or later from the [UE4 Launcher](https://www.unrealengine.com/) or [Github](https://www.unrealengine.com/ue4-on-github). 
 
-* Install the **CodeWorks for Android 1R4** which can be found in **Engine\Extras\Android**.  
+* Install the **CodeWorks for Android 1R4**, which can be found in **Engine\Extras\Android**.  
 
-* Check to see if **[Developer Mode](https://docs.unrealengine.com/latest/INT/Platforms/Android/GettingStarted/2/index.html )** has been enabled on your Samsung Galaxy Note 4 Smartphone and if not, enable it. 
+* Check to see if [Developer Mode](https://docs.unrealengine.com/latest/INT/Platforms/Android/GettingStarted/2/index.html ) has been enabled on your Samsung Galaxy Note 4 Smartphone and if not, enable it. 
 
-* Also make sure to check that USB Debugging has been enabled to ensure that you can load data onto your Smartphone from your PC.
+* Also, make sure to check that USB Debugging has been enabled to ensure that you can load data onto the Smartphone from your PC.
 	
 	![](Platforms\GearVR\GVR_Dev_Options_Enabled.png)
 
 * With your Smartphone plugged into your development PC via USB, make sure to accept the RSA Fingerprint key from your PC.
-This way your development PC and Smartphone can send data back and forth.
-Also make sure to check **Always allow this computer** to ensure that your development PC can always interact with this Smartphone when plugged in. 
+This way, your development PC and Smartphone can send data back and forth.
+Also, make sure to check **Always allow this computer** to ensure that your development PC can always interact with this Smartphone when plugged in. 
 
 	![](Platforms\GearVR\GVR_USB_Debuging.png)
 
-* Generate and download the **OSIG File** for your Samsung Smartphone using your Smartphone's Device ID and the **[Oculus](https://developer.oculus.com/tools/osig/)** website.
-If you are not sure how to get your Smartphone's **Device ID** the **[Oculus website](https://developer.oculus.com/tools/osig/)** has detailed directions on how to do this.
+* Generate and download the **OSIG File** for your Samsung Smartphone using your Smartphone's Device ID and the [Oculus](https://developer.oculus.com/tools/osig/) website.
+If you are not sure how to get your Smartphone's **Device ID**, the [Oculus website](https://developer.oculus.com/tools/osig/) has detailed directions on how to do this.
 	
 	![](Platforms\GearVR\GVR_OSIG_Website.png)
 
 * After you download the OSIG file, create a folder called **assets** and place the **OSIG** file inside of it. 
-Then place that folder in **Engine/Build/Android/Java**. 
+Then, place that folder in **Engine/Build/Android/Java**. 
 
 	[region:warning]
-	**If you do not do this your project will not be able to run on your Galaxy Smartphone.**
+	If you do not do this, your project will not be able to run on your Galaxy Smartphone.
 	[/region]
 	
 	![](Platforms\GearVR\GVR_OSIG_File_Location.png)
diff --git a/Engine/Documentation/Source/Platforms/GearVR/QuickStart/2/GearVRQuickStart_2.INT.udn b/Engine/Documentation/Source/Platforms/GearVR/QuickStart/2/GearVRQuickStart_2.INT.udn
--- a/Engine/Documentation/Source/Platforms/GearVR/QuickStart/2/GearVRQuickStart_2.INT.udn
+++ b/Engine/Documentation/Source/Platforms/GearVR/QuickStart/2/GearVRQuickStart_2.INT.udn
@@ -1,6 +1,6 @@
 Availability:Public
 Title:2. Gear VR Project Setup
-Crumbs:%ROOT%, Platforms/GearVR
+Crumbs:%ROOT%, Platforms, Platforms/GearVR
 Description:How to setup your C++ based UE4 project to be used with the Gear VR Headset.
 Navigation:platform
 Skilllevel: Intermediate
@@ -25,15 +25,15 @@ Version: 4.9
 %Steps%
 
 
-1. Once the project has finished being created and compiled, Visual Studio 2013 will automatically open (If already installed) along with the editor.
+1. Once the project has finished being created and compiled, Visual Studio 2013 will automatically open (if already installed) along with the editor.
 
 	![](GVR_Editor_Opened.png)
 
 	[region:note]
-	If you are using Unreal Engine 4.8 or later and a Blueprint based project Visual Studio 2013 will not open up as it is not needed. 
+	If you are using Unreal Engine 4.8 or later and a Blueprint based project, Visual Studio 2013 will not open up as it is not needed. 
 	[/region]
 
-1. Once the UE4 Editor has loaded, from the **Main Toolbar**, open the **Windows** tab and select the **Plugins** menu item. 
+1. Once the UE4 Editor has loaded, from the **Main Toolbar**, open the **Windows** tab and select the **Plugins** menu. 
 
 	[OBJECT:EmbeddedVideo]
 	[PARAMLITERAL:width]
@@ -47,11 +47,11 @@ Version: 4.9
 	[/PARAMLITERAL]
 	[/OBJECT]
 
-1. From the Plug-In menu go to the **Head Mounted Displays** section and **disable** the **Oculus** option and **enable** the **Gear VR** option. Then click the **Restart Now** button to re-start the editor so the changes can take place. 
+1. From the **Plugins** menu, go to the **Head Mounted Displays** section, **disable** the **Oculus** option, and **enable** the **Gear VR** option. Then, click the **Restart Now** button to re-start the editor so the changes can take place. 
 
 	![](GVR_Plugin_Options_Toggle.png)
 
-1. Once the Editor has re-started, select **Edit** from the **Main Toolbar** and then select **Project Settings**  and click the **Configure Now** button for Android support in the **APKPacking** section.
+1. Once the Editor has re-started, select **Edit** from the **Main Toolbar**, select **Project Settings** and click the **Configure Now** button for Android support in the **APKPacking** section.
 
 	[OBJECT:EmbeddedVideo]
 	[PARAMLITERAL:width]
@@ -65,11 +65,11 @@ Version: 4.9
 	[/PARAMLITERAL]
 	[/OBJECT]
 
-1. Next change the **Minimum SDK Version** from **9** to **19**.
+1. Next, change the **Minimum SDK Version** from **9** to **19**.
 
 	![](GVR_Android_SDK_Version.png)
 
-1. Next scroll down till you see the **Configure the Android Manifest for deployment to Gear VR** checkbox and enable it. 
+1. Next, scroll down till you see the **Configure the Android Manifest for deployment to Gear VR** checkbox and enable it. 
 
 	[Region:note]
 	This step is of vital importance as without this on your project will not be able to run on the Gear VR Headset.
@@ -77,16 +77,16 @@ Version: 4.9
 
 	![](GVR_Deployment_Checkbox.png)
 
-1. Next, under the **Android SDK** section fill, in the location of the required Android SDK files on your development PC. 
+1. Next, under the **Android SDK** section, fill-in the location of the required Android SDK files on your development PC. 
 	You can find the needed folders in the location where you installed the Tegra files. 
 	
 	[Region:note]
-	If you have already installed CodeWorks for Android 1R4 this information should be already filled in for you. The only reason that you should change something here is if the information is missing or point to the wrong location.
+	If you have already installed CodeWorks for Android 1R4, this information should be already filled in for you. The only reason that you should change something here is if the information is missing or pointing to the wrong location.
 	[/region]
 
 	![](GVR_Android_SDK_Setup.png)
 
-1. Finally under **File** > **Package Project** > **Packing Settings** un-check **Full Rebuild** to help speed up iteration times. 
+1. Finally, under **File** > **Package Project** > **Packing Settings**, un-check **Full Rebuild** to help speed up iteration times. 
 
 	[Region:note]
 	Remember to **Re-Enable** this option when your project is ready for its final publish to ensure that your project will work as intended on the Gear VR Headset. 
