#UE4 Docs: minor formatting updates to property tables

diff --git a/Engine/Documentation/Source/Engine/Animation/Blendspaces/Blendspaces.INT.udn b/Engine/Documentation/Source/Engine/Animation/Blendspaces/Blendspaces.INT.udn
--- a/Engine/Documentation/Source/Engine/Animation/Blendspaces/Blendspaces.INT.udn
+++ b/Engine/Documentation/Source/Engine/Animation/Blendspaces/Blendspaces.INT.udn
@@ -178,52 +178,25 @@ When setting up your Blend Space asset, there are additional options that can be
 
 ![](BlendSpaceDetails.png)
 
-### Blend Space
-
+| Property | Description |
 |---|---|
+| [REGION:tablesection]Blend Space[/REGION] ||
 |**Display Editor Vertically**| Only available for 1D Blend Space assets, toggles the Blend Space graph display in either vertical or horizontal display modes.|
-
-### Input Interpolation
-
-|---|---|
+| [REGION:tablesection]Input Interpolation[/REGION] ||
 |**Axis to Scale Animation**| Not available for 	1D Blend Space assets, this allows you (if you have input interpolation) to determine which axis to drive animation speed (scale). For example, for locomotion animation, speed axis will drive animation speed (thus scale).|
 |**Interpolation Params**| Input interpolation parameter for all 3 axis, for each axis input, decide how you'd like to interpolate input to by setting Interpolation Time and Interpolation Type. |
-
-### Sample Interpolation
-
-|---|---|
+| [REGION:tablesection]Sample Interpolation[/REGION] ||
 |**Target Weight Interpolation Speed Per Sec**| When target samples are set, how fast you would like to get to target which improves target blending. For example, if you interpolate input, when you move from left to right rapidly, you will interpolate through forward, but if you use target weight interpolation, you will skip forward but interpolate between left and right.|
 |**Per Bone Blend**| Define target weight interpolation per bone. This will blend in different speed per each bone setting.|
-
-### Animation Notifies
-
-|---|---|
+| [REGION:tablesection]Animation Notifies[/REGION] ||
 |**Notify Trigger Mode**| Assign the current mode used by the Blend Space to decide which animation notifies to fire: All Animation (all notify events will fire), HighestWeightedAnimation (Notify events will only fire from the highest weighted animation) or None (no notify events will fire from any animations).|
-
-### Additive Settings
-
-|---|---|
+| [REGION:tablesection]Additive Settings[/REGION] ||
 |**Preview Base Pose**| Preview Base pose for additive Blend Space.|
-
-### Animation
-
-|---|---|
+| [REGION:tablesection]Animation[/REGION] ||
 |**Skeleton**| Pointer to the Skeleton this asset can be played on (unadjustable).|
-
-### Meta Data
-
-|---|---|
+| [REGION:tablesection]Meta Data[/REGION] ||
 |**Meta Data**| This is Meta Data that can be saved with the asset. The meta data is a Blueprintable class derived from the Anim Meta Data class. This allows you to add custom meta data to animation assets (Anim Sequence, Anim Montage, Anim Composite, and Blendspace are all supported). You can query the data from an animation asset in C++ by calling the `GetMetaData()` method, or by using `GetSectionMetaData()` on Anim Montages. |
-
-### Thumbnail
-
-|---|---|
+| [REGION:tablesection]Thumbnail[/REGION] ||
 |**Orbit Pitch**| The pitch of the orbit camera around the asset. |
 |**Orbit Yaw**| The yaw of the orbit camera around the asset. |
 |**Orbit Zoom**| The offset from the bounds sphere distance from the asset.|
-
-
-
-
-
-
