#UE4 Docs: set page to publish

diff --git a/Engine/Documentation/Source/Engine/Blueprints/TechnicalGuide/NativizingBlueprints/Advanced/AdvancedNativeBlueprints.INT.udn b/Engine/Documentation/Source/Engine/Blueprints/TechnicalGuide/NativizingBlueprints/Advanced/AdvancedNativeBlueprints.INT.udn
--- a/Engine/Documentation/Source/Engine/Blueprints/TechnicalGuide/NativizingBlueprints/Advanced/AdvancedNativeBlueprints.INT.udn
+++ b/Engine/Documentation/Source/Engine/Blueprints/TechnicalGuide/NativizingBlueprints/Advanced/AdvancedNativeBlueprints.INT.udn
@@ -1,4 +1,4 @@
-Availability:Docs
+Availability:Public
 Title:Advanced Blueprint Nativization
 Crumbs: %ROOT%, Programming, Engine/Blueprints/TechnicalGuide
 Description:Technical guide for gameplay programmers exposing gameplay elements to Blueprints.
diff --git a/Engine/Documentation/Source/Engine/Blueprints/TechnicalGuide/NativizingBlueprints/NativizingBlueprints.INT.udn b/Engine/Documentation/Source/Engine/Blueprints/TechnicalGuide/NativizingBlueprints/NativizingBlueprints.INT.udn
--- a/Engine/Documentation/Source/Engine/Blueprints/TechnicalGuide/NativizingBlueprints/NativizingBlueprints.INT.udn
+++ b/Engine/Documentation/Source/Engine/Blueprints/TechnicalGuide/NativizingBlueprints/NativizingBlueprints.INT.udn
@@ -1,4 +1,4 @@
-Availability:Docs
+Availability:Public
 Title:Nativizing Blueprints
 Crumbs: %ROOT%, Programming, Engine/Blueprints/TechnicalGuide
 Description:Technical guide for gameplay programmers exposing gameplay elements to Blueprints.
