#UE4 Docs: set page to publish

diff --git a/Engine/Documentation/Source/Engine/LevelStreaming/HowTo/HowToStreamLevels.INT.udn b/Engine/Documentation/Source/Engine/LevelStreaming/HowTo/HowToStreamLevels.INT.udn
--- a/Engine/Documentation/Source/Engine/LevelStreaming/HowTo/HowToStreamLevels.INT.udn
+++ b/Engine/Documentation/Source/Engine/LevelStreaming/HowTo/HowToStreamLevels.INT.udn
@@ -1,4 +1,4 @@
-Availability: Docs
+Availability: Public
 Title:Level Streaming How-Tos
 Crumbs:%ROOT%, Engine, Engine/UI/LevelEditor
 Description:How to dynamically stream in levels using Level Streaming Volumes, Blueprints, and C++ code
diff --git a/Engine/Documentation/Source/Engine/LevelStreaming/HowTo/StreamWithBlueprints/StreamBP.INT.udn b/Engine/Documentation/Source/Engine/LevelStreaming/HowTo/StreamWithBlueprints/StreamBP.INT.udn
--- a/Engine/Documentation/Source/Engine/LevelStreaming/HowTo/StreamWithBlueprints/StreamBP.INT.udn
+++ b/Engine/Documentation/Source/Engine/LevelStreaming/HowTo/StreamWithBlueprints/StreamBP.INT.udn
@@ -1,4 +1,4 @@
-Availability: Docs
+Availability: Public
 Title:Load and Unload Levels with Blueprints
 Crumbs:%ROOT%, Engine, Engine/UI/LevelEditor
 Description:How to stream levels with a custom streaming Actor created with Blueprints
diff --git a/Engine/Documentation/Source/Engine/LevelStreaming/HowTo/StreamWithC++/StreamC++.INT.udn b/Engine/Documentation/Source/Engine/LevelStreaming/HowTo/StreamWithC++/StreamC++.INT.udn
--- a/Engine/Documentation/Source/Engine/LevelStreaming/HowTo/StreamWithC++/StreamC++.INT.udn
+++ b/Engine/Documentation/Source/Engine/LevelStreaming/HowTo/StreamWithC++/StreamC++.INT.udn
@@ -1,4 +1,4 @@
-Availability: Docs
+Availability: Public
 Title:Load and Unload Levels with C++
 Crumbs:%ROOT%, Engine, Engine/UI/LevelEditor
 Description:How to stream levels with a custom streaming Actor created in C++
diff --git a/Engine/Documentation/Source/Engine/LevelStreaming/HowTo/StreamWithVolumes/StreamLSV.INT.udn b/Engine/Documentation/Source/Engine/LevelStreaming/HowTo/StreamWithVolumes/StreamLSV.INT.udn
--- a/Engine/Documentation/Source/Engine/LevelStreaming/HowTo/StreamWithVolumes/StreamLSV.INT.udn
+++ b/Engine/Documentation/Source/Engine/LevelStreaming/HowTo/StreamWithVolumes/StreamLSV.INT.udn
@@ -1,4 +1,4 @@
-Availability: Docs
+Availability: Public
 Title:Stream Sublevels with Level Streaming Volumes
 Crumbs:%ROOT%, Engine, Engine/UI/LevelEditor
 Description:How to use volumes to control the streaming of Levels based on the player's viewpoint.
diff --git a/Engine/Documentation/Source/Engine/LevelStreaming/Overview/Overview.INT.udn b/Engine/Documentation/Source/Engine/LevelStreaming/Overview/Overview.INT.udn
--- a/Engine/Documentation/Source/Engine/LevelStreaming/Overview/Overview.INT.udn
+++ b/Engine/Documentation/Source/Engine/LevelStreaming/Overview/Overview.INT.udn
@@ -1,4 +1,4 @@
-Availability: Docs
+Availability: Public
 Title:Level Streaming Overview
 Crumbs:%ROOT%, Engine, Engine/UI/LevelEditor
 Description:Streaming Levels can be loaded with Level Streaming Volumes or programmatically with Blueprints or C++.
diff --git a/Engine/Documentation/Source/Engine/LevelStreaming/StreamingVolumes/StreamingVolumes.INT.udn b/Engine/Documentation/Source/Engine/LevelStreaming/StreamingVolumes/StreamingVolumes.INT.udn
--- a/Engine/Documentation/Source/Engine/LevelStreaming/StreamingVolumes/StreamingVolumes.INT.udn
+++ b/Engine/Documentation/Source/Engine/LevelStreaming/StreamingVolumes/StreamingVolumes.INT.udn
@@ -1,4 +1,4 @@
-Availability: Docs
+Availability: Public
 Title:Level Streaming Volumes
 Crumbs:
 Description:Guide to using volumes to control the streaming of Levels based on the player's viewpoint.
diff --git a/Engine/Documentation/Source/Engine/Levels/CollaborateWithSublevels/CollaborateWithSublevels.INT.udn b/Engine/Documentation/Source/Engine/Levels/CollaborateWithSublevels/CollaborateWithSublevels.INT.udn
--- a/Engine/Documentation/Source/Engine/Levels/CollaborateWithSublevels/CollaborateWithSublevels.INT.udn
+++ b/Engine/Documentation/Source/Engine/Levels/CollaborateWithSublevels/CollaborateWithSublevels.INT.udn
@@ -1,4 +1,4 @@
-Availability: Docs
+Availability: Public
 Title:Collaborate with Sublevels
 Crumbs: %ROOT%, Engine
 Description:You can use sublevels to collaborate with other designers working in the same space. 
diff --git a/Engine/Documentation/Source/Engine/Levels/LevelsWindow/LevelsWindow.INT.udn b/Engine/Documentation/Source/Engine/Levels/LevelsWindow/LevelsWindow.INT.udn
--- a/Engine/Documentation/Source/Engine/Levels/LevelsWindow/LevelsWindow.INT.udn
+++ b/Engine/Documentation/Source/Engine/Levels/LevelsWindow/LevelsWindow.INT.udn
@@ -1,6 +1,7 @@
 Availability: Public
 Title:Managing Multiple Levels
 Description:Manage your persistent level and sublevels with the Levels window.
+crumbs:
 Version: 4.12
 parent:Engine/Levels
 type:reference
@@ -144,7 +145,7 @@ No additional details are visible for persistent Levels, although there is a dro
 ![](LevelDetails_Persistent.png)(w:500)
 
 For sublevels, you can set the Level's offset **Position** and **Rotation**, as well as the **Streaming Volumes** to use and the debug **Level Color**. Advanced settings like the 
-[minimum time between unload requests](Engine/LevelStreaming/Overview#addinghysteresistounloadingrequests) to improve performance are also accessible here.
+[minimum time between unload requests](Engine/LevelStreaming/StreamingVolumes#addinghysteresistounloadingrequests) to improve performance are also accessible here.
 
 ![](LevelDetails_Streaming.png)
 
