#UE4 Docs: set page to publish

diff --git a/Engine/Documentation/Source/Engine/Sequencer/HowTo/BurnIns/BurnIns.INT.udn b/Engine/Documentation/Source/Engine/Sequencer/HowTo/BurnIns/BurnIns.INT.udn
--- a/Engine/Documentation/Source/Engine/Sequencer/HowTo/BurnIns/BurnIns.INT.udn
+++ b/Engine/Documentation/Source/Engine/Sequencer/HowTo/BurnIns/BurnIns.INT.udn
@@ -1,4 +1,4 @@
-Availability: Docs
+Availability: Public
 Title:Applying Burn Ins
 Crumbs: %ROOT%, Engine, Engine/Sequencer, Engine/Sequencer/HowTo
 Description:Shows how you can apply overlays from UMG to rendered Sequencer movies.  
diff --git a/Engine/Documentation/Source/Engine/Sequencer/HowTo/CustomBurnIns/CustomBurnIns.INT.udn b/Engine/Documentation/Source/Engine/Sequencer/HowTo/CustomBurnIns/CustomBurnIns.INT.udn
--- a/Engine/Documentation/Source/Engine/Sequencer/HowTo/CustomBurnIns/CustomBurnIns.INT.udn
+++ b/Engine/Documentation/Source/Engine/Sequencer/HowTo/CustomBurnIns/CustomBurnIns.INT.udn
@@ -1,4 +1,4 @@
-Availability: Docs
+Availability: Public
 Title:Applying Custom Burn Ins
 Crumbs: %ROOT%, Engine, Engine/Sequencer, Engine/Sequencer/HowTo
 Description:Shows you to create and apply custom burn ins to rendered movies. 
diff --git a/Engine/Documentation/Source/Engine/Sequencer/HowTo/SequencerHowTo.INT.udn b/Engine/Documentation/Source/Engine/Sequencer/HowTo/SequencerHowTo.INT.udn
--- a/Engine/Documentation/Source/Engine/Sequencer/HowTo/SequencerHowTo.INT.udn
+++ b/Engine/Documentation/Source/Engine/Sequencer/HowTo/SequencerHowTo.INT.udn
@@ -30,9 +30,7 @@ Refer to the table below for more information and click each link to view its co
 [](FANCY:Engine\Sequencer\HowTo\CameraRigRail "%Engine\Sequencer\HowTo\CameraRigRail:description%")
 [](FANCY:Engine\Sequencer\HowTo\CineCameraActors "%Engine\Sequencer\HowTo\CineCameraActors:description%")
 [](FANCY:Engine\Sequencer\HowTo\TriggeringSequences "%Engine\Sequencer\HowTo\TriggeringSequences:description%")
-[PUBLISH:docs]
 [](FANCY:Engine\Sequencer\HowTo\BurnIns "%Engine\Sequencer\HowTo\BurnIns:description%")
-[/PUBLISH:docs]
 [/REGION]
 
 
@@ -43,9 +41,7 @@ Refer to the table below for more information and click each link to view its co
 [](FANCY:Engine\Sequencer\HowTo\Spawnables "%Engine\Sequencer\HowTo\Spawnables:description%")
 [](FANCY:Engine\Sequencer\HowTo\SequenceRecorder "%Engine\Sequencer\HowTo\SequenceRecorder:description%")
 [](FANCY:Engine\Sequencer\HowTo\TracksLevelVisibility "%Engine\Sequencer\HowTo\TracksLevelVisibility:description%")
-[PUBLISH:docs]
 [](FANCY:Engine\Sequencer\HowTo\CustomBurnIns "%Engine\Sequencer\HowTo\CustomBurnIns:description%")
-[/PUBLISH:docs]
 [/REGION]
 
 
