#UE4 Docs: set page to publish

diff --git a/Engine/Documentation/Source/Engine/Tools/DevTools/DeveloperTools.INT.udn b/Engine/Documentation/Source/Engine/Tools/DevTools/DeveloperTools.INT.udn
--- a/Engine/Documentation/Source/Engine/Tools/DevTools/DeveloperTools.INT.udn
+++ b/Engine/Documentation/Source/Engine/Tools/DevTools/DeveloperTools.INT.udn
@@ -1,4 +1,4 @@
-Availability:Docs
+Availability:Public
 Title:Developer Tools
 Crumbs: %ROOT%, Engine
 Description:Descriptions of the options found under the Developer Tools file menu. 
@@ -40,7 +40,7 @@ See the table below for more information and links to additional pages.
 | **Device Manager** | Opens the Device Manager. |
 | **Device Profiles** | Opens the [**Device Profiles**](Platforms/DeviceProfiles) tool. Allows you to manage device-specific configurations inside your project. |
 | **Merge Actors** | Opens the Merge Actors tool. |
-| **Pixel Inspector** | Opens the [**Pixel Inspector**](Engine\Tools\DevTools\PixelInspector). Provides information about a specific Pixel from a viewport. |
+| **Pixel Inspector** | Opens the [**Pixel Inspector**](Engine/Tools/DevTools/PixelInspector). Provides information about a specific Pixel from a viewport. |
 | **Session Frontend** | Opens the [**Session Frontend**](Engine/Deployment/UnrealFrontend) tool. Multipurpose tool used for daily game development (prepping builds, deployment, launching). |
 | **Widget Reflector** | Opens the Widget Reflector. |
 
diff --git a/Engine/Documentation/Source/Engine/Tools/DevTools/PixelInspector/PixelInspector.INT.udn b/Engine/Documentation/Source/Engine/Tools/DevTools/PixelInspector/PixelInspector.INT.udn
--- a/Engine/Documentation/Source/Engine/Tools/DevTools/PixelInspector/PixelInspector.INT.udn
+++ b/Engine/Documentation/Source/Engine/Tools/DevTools/PixelInspector/PixelInspector.INT.udn
@@ -1,4 +1,4 @@
-Availability:Docs
+Availability:Public
 Title:Pixel Inspector
 Crumbs: %ROOT%, Engine
 Description:The Pixel Inspector shows you what is happening in the various parts of the GBuffer at a given pixel.
