Updated Level Designer Quick Start Guide for content/images.

diff --git a/Engine/Documentation/Source/Engine/QuickStart/3/LDQS_3.INT.udn b/Engine/Documentation/Source/Engine/QuickStart/3/LDQS_3.INT.udn
--- a/Engine/Documentation/Source/Engine/QuickStart/3/LDQS_3.INT.udn
+++ b/Engine/Documentation/Source/Engine/QuickStart/3/LDQS_3.INT.udn
@@ -27,12 +27,12 @@ checkpoint: editorqs
 %Steps%
 
 Next we will create a new [Level](Engine/Levels) that we will use to build our game environment in. 
-While there are several different ways in which you can create a new Level, we will use the most common method by using the [Content Browser](Engine/Content/Browser) which is used to create several different types of assets inside the Unreal Editor. 
+While there are several different ways in which you can create a new Level, we will use the **File Menu** method which we will present us with level selection options. 
 
 
 ## Steps
 
-1. Inside the Unreal Editor, locate the **Content Browser** window and click the **Add New** Button then select the **Level** asset. 
+1. Inside the Unreal Editor, click the **File Menu** option then select **New Level...**. 
 
 	![](ContentBrowserLevel.png)
 
@@ -40,7 +40,7 @@ While there are several different ways in which you can create a new Level, we w
 
 	![](level_new.png)(w:550)
 
-	Above, the **Default** level includes some of the commonly used assets for constructing levels while **Empty Level** is a completely blank level with no assets. For the purposes of this guide we are going to start from scratch with a completely blank slate. 
+	Above the **Default** level includes some of the commonly used assets for constructing levels, there is also a **VR-Basic** level which includes some assets for constructing levels with the VR Editor and **Empty Level** is a completely blank level with no assets. For the purposes of this guide we are going to start from scratch with a completely blank slate. 
 
 1. Click the **Empty Level** to select it.
 
