- Sequencer EDL How-to.
- Added section to Overview Page linking to How-to.
- Added to Sequencer How-to Overview page.

diff --git a/Engine/Documentation/Source/Engine/Sequencer/HowTo/BurnIns/BurnIns.INT.udn b/Engine/Documentation/Source/Engine/Sequencer/HowTo/BurnIns/BurnIns.INT.udn
--- a/Engine/Documentation/Source/Engine/Sequencer/HowTo/BurnIns/BurnIns.INT.udn
+++ b/Engine/Documentation/Source/Engine/Sequencer/HowTo/BurnIns/BurnIns.INT.udn
@@ -10,9 +10,11 @@ checkpoint: editorqs
 Related: Engine/Sequencer/Overview
 Related: Engine/Sequencer/ReferenceEditor
 
+[EXCERPT:intro]
 When creating and rendering out cinematic movies, you may want to include overlays with information about the scene you are viewing such as the name of the shot, the date, time or frame information. 
 These overlays are referred to as **Burn Ins** as they are burned into the movie when it is rendered out. 
 Often times you will see this process as part of real world movie making to provide directors, editors or anyone viewing the scene with contextual information about the scene (some will even be watermarked with copyright information). 
+[/EXCERPT:intro]
 
 **Sequencer** allows you to apply burn ins to your rendered out movies and provides a default burn in with general information about the shot. 
 In this guide we take a look at how to enable the default burn in options when rendering out your cinematics. 
diff --git a/Engine/Documentation/Source/Engine/Sequencer/HowTo/ImportExportEDLs/ImportExportEDLs.INT.udn b/Engine/Documentation/Source/Engine/Sequencer/HowTo/ImportExportEDLs/ImportExportEDLs.INT.udn
new file mode 100644
--- /dev/null
+++ b/Engine/Documentation/Source/Engine/Sequencer/HowTo/ImportExportEDLs/ImportExportEDLs.INT.udn
@@ -0,0 +1,128 @@
+Availability: Docs
+Title:Importing & Exporting Edit Decision Lists (EDLs)
+Crumbs: %ROOT%, Engine, Engine/Sequencer, Engine/Sequencer/HowTo
+Description:Illustrates how you can import and export EDLs for use in external video editing software applications. 
+Version: 4.13
+SkillLevel: Advanced
+parent:Engine/Sequencer/HowTo
+order:1
+checkpoint: editorqs
+Related: Engine/Sequencer/Overview
+Related: Engine\Sequencer\ReferenceEditor\
+
+[EXCERPT:intro]
+Sequencer in Unreal Engine 4 not only allows you to render and export out your entire cinematic, but you can also export your cinematic split up into each of your shots along with an **Edit Decision List** or **EDL** which is a file that can be used with most video editing applications like Premiere, Avid or Final Cut. 
+[/EXCERPT:intro]
+
+The EDL contains an ordered list of reel and time code data representing where each video clip can be obtained in order to produce the final cut. 
+In this how-to we export a cinematic and EDL file, take it into an external program like Premiere, make edits to it and re-import the edited cinematic back into Unreal along with those changes. 
+
+## Steps
+
+[REGION:note]
+For this how-to, we are using the **Blueprint Third Person Template** project and have created a sample cinematic however you can use your own project and cinematic sequence. 
+We are also using Adobe Premiere Pro CS6 as our external editing software; however you can use any software that supports EDL importing. 
+[/REGION]
+
+1.	Make sure your cinematic you want to export includes a **Shot Track** containing your shots, then click the **Render Movie** button.
+
+1.	Define an **Output Directory** and **Filename Format** using the **{shot}** token and make sure **Write Edit Decision List** is checked. 
+
+	[OBJECT:EmbeddedVideo]
+	[PARAMLITERAL:width]
+	640
+	[/PARAMLITERAL]
+	[PARAMLITERAL:height]
+	360
+	[/PARAMLITERAL]
+	[PARAMLITERAL:videoid]
+	c2J89J3b8WA
+	[/PARAMLITERAL]
+	[/OBJECT]
+
+	Above we show the sample cinematic that we want to export along with an EDL and click the **Render Movie** button to bring up the settings dialog window. 
+	Inside the settings window, we create a new folder to export the cinematic to and prefix our {shot} token with the name of our cinematic. 
+	By using the {shot} token each of our shots will be exported as individual video files along with the full cinematic. 
+
+1.	Inside your software editing suite, **Import** the **EDL** file exported from the previous step. 
+
+	[OBJECT:EmbeddedVideo]
+	[PARAMLITERAL:width]
+	640
+	[/PARAMLITERAL]
+	[PARAMLITERAL:height]
+	360
+	[/PARAMLITERAL]
+	[PARAMLITERAL:videoid]
+	LN0OeRF0G8A
+	[/PARAMLITERAL]
+	[/OBJECT]
+
+	In Premiere Pro, you can **Right-click** in the window in the lower-left and choose **Import** then point to the EDL file. 
+	We left the EDL Information and New Sequence dialog windows with the default information. 
+
+1.	Next you'll need to link the data from the EDL to the exported media files in your editing software. 
+
+	[OBJECT:EmbeddedVideo]
+	[PARAMLITERAL:width]
+	640
+	[/PARAMLITERAL]
+	[PARAMLITERAL:height]
+	360
+	[/PARAMLITERAL]
+	[PARAMLITERAL:videoid]
+	iUJenhTwFRk
+	[/PARAMLITERAL]
+	[/OBJECT]
+
+	In Premeire Pro, **Right-click** the media folder and select **Link Media** then point each file to its corresponding media file. 
+
+1.	Modify the shots in your editing software as desired, then **Export** as an **EDL**. 
+
+	[OBJECT:EmbeddedVideo]
+	[PARAMLITERAL:width]
+	640
+	[/PARAMLITERAL]
+	[PARAMLITERAL:height]
+	360
+	[/PARAMLITERAL]
+	[PARAMLITERAL:videoid]
+	AH3Fcy9v8tQ
+	[/PARAMLITERAL]
+	[/OBJECT]
+
+	Above for our sample scene, we adjust the length of the shots as well as re-arrange their position in the sequence. 
+
+1.	Return to **Unreal Engine 4** and in your Sequence, **Right-click** your **Shot Track** and **Import EDL**. 
+
+	[OBJECT:EmbeddedVideo]
+	[PARAMLITERAL:width]
+	640
+	[/PARAMLITERAL]
+	[PARAMLITERAL:height]
+	360
+	[/PARAMLITERAL]
+	[PARAMLITERAL:videoid]
+	uI-wIYvMEiQ
+	[/PARAMLITERAL]
+	[/OBJECT]
+
+	The changes we made externally are now reflected inside our Sequence in Unreal Engine 4. 
+
+## End Result
+
+The import/export EDL option within Sequencer gives you the flexibility to determine where you want to do your editing. 
+
+You can do most of your setup inside Sequencer, then pass your sequence off to an editor who works inside an external application like Premiere to perform the final touches who can then pass it back to you to tweak in Sequencer. 
+
+By exporting your sequence as an EDL and working externally, you have the ability to add post production effects or after effects to your scene before you rendering out the final cinematic. 
+
+
+
+
+
+
+
+
+
+
diff --git a/Engine/Documentation/Source/Engine/Sequencer/HowTo/SequencerHowTo.INT.udn b/Engine/Documentation/Source/Engine/Sequencer/HowTo/SequencerHowTo.INT.udn
--- a/Engine/Documentation/Source/Engine/Sequencer/HowTo/SequencerHowTo.INT.udn
+++ b/Engine/Documentation/Source/Engine/Sequencer/HowTo/SequencerHowTo.INT.udn
@@ -42,6 +42,9 @@ Refer to the table below for more information and click each link to view its co
 [](FANCY:Engine\Sequencer\HowTo\SequenceRecorder "%Engine\Sequencer\HowTo\SequenceRecorder:description%")
 [](FANCY:Engine\Sequencer\HowTo\TracksLevelVisibility "%Engine\Sequencer\HowTo\TracksLevelVisibility:description%")
 [](FANCY:Engine\Sequencer\HowTo\CustomBurnIns "%Engine\Sequencer\HowTo\CustomBurnIns:description%")
+[PUBLISH:docs]
+[](FANCY:Engine\Sequencer\HowTo\ImportExportEDLs "%Engine\Sequencer\HowTo\ImportExportEDLs:description%")
+[/PUBLISH:docs]
 [/REGION]
 
 
diff --git a/Engine/Documentation/Source/Engine/Sequencer/Overview/SequencerOverview.INT.udn b/Engine/Documentation/Source/Engine/Sequencer/Overview/SequencerOverview.INT.udn
--- a/Engine/Documentation/Source/Engine/Sequencer/Overview/SequencerOverview.INT.udn
+++ b/Engine/Documentation/Source/Engine/Sequencer/Overview/SequencerOverview.INT.udn
@@ -9,7 +9,7 @@ Version: 4.12
 Related: Engine\Sequencer\ReferenceEditor
 Related: Engine\Sequencer\HowTo
 
-The **Sequencer** Editor gives users the ability to create in-game cinematics through a specialized multi-track editor similar to [**Matinee**](Engine\Matinee\). 
+The **Sequencer** Editor gives users the ability to create in-game cinematics through a specialized multi-track editor similar to [Matinee](Engine\Matinee\). 
 By creating **Level Sequences** and adding **Tracks**, users can define the make-up of each Track which will determine the content for their scene.
 Tracks can consist of things like Animation (for animating a character), Transformation (moving things around in the scene like Camera), Audio (for including a music or sound effect), as well as several other Track types.  
 
@@ -86,7 +86,7 @@ We are "possessing" that Skeletal Mesh that exists in order to tell it what to d
 So how is this beneficial and when would I elect to use a "spawnable" object? Since Sequencer is the owner and spawns the object, it is not bound to a particular level. Anything deemed as spawnable can be used in any level, so you could create a scene and have it take place in any environment and reuse it in new levels without having to recreate it. 
 
 [REGION:note]
-Please see [**Creating Spawnables**](Engine\Sequencer\HowTo\Spawnables\) for more information. 
+Please see [Creating Spawnables](Engine\Sequencer\HowTo\Spawnables\) for more information. 
 [/REGION]
 
 
@@ -153,7 +153,7 @@ Sequencer allows you to use specialized viewports to make the editing process ea
 
 ![](SpecialViewports.png)
 
-These [**Cinematic Viewports**](Engine\Sequencer\HowTo\ViewportOptions) give a better idea of what the scene will look like and can be enabled from the **Viewport** options button. 
+These [Cinematic Viewports](Engine\Sequencer\HowTo\ViewportOptions) give a better idea of what the scene will look like and can be enabled from the **Viewport** options button. 
 
 ##Cinematic Actors
 
@@ -172,7 +172,7 @@ The Camera Rig Crane Actor can be used to simulate crane-like camera movements t
 You can attach a Camera to the Camera Rig Crane Actor then manipulate the crane's movement through the **Details** panel and **Crane Controls** values. Those values can affect the Crane Pitch, Crane Yaw and Crane Arm Length. These values can be keyframed within Sequencer which will allow you to adjust them during your cinematics as desired. 
 
 [REGION:note]
-Please see [**Shooting from a Camera Rig Crane**](Engine\Sequencer\HowTo\CameraRigCrane) for more information. 
+Please see [Shooting from a Camera Rig Crane](Engine\Sequencer\HowTo\CameraRigCrane) for more information. 
 [/REGION]
 
 ### Camera Rig Rail
@@ -184,7 +184,7 @@ The Camera Rig Rail Actor is a spline based tool that a Camera can be attached t
 You can select each spline point and alter the tangents to produce the path you want your camera to follow. The camera that is attached to the rail can be rotated independently and the position of where the camera is on the rail can be adjusted with the **Current Position on Rail** property inside the **Details** panel. This value can be keyframed allowing you to alter the value through Sequencer. 
 
 [REGION:note]
-Please see [**Shooting from a Camera Rig Rail**](Engine\Sequencer\HowTo\CameraRigRail) for more information. 
+Please see [Shooting from a Camera Rig Rail](Engine\Sequencer\HowTo\CameraRigRail) for more information. 
 [/REGION]
 
 ### Cine Camera Actor
@@ -200,7 +200,7 @@ Cine Camera Actor (black camera) and default Camera Actor (blue camera) pictured
 The Cine Camera Actor includes settings for Look at Tracking (following an Actor), Filmback Settings (16:9 DSLR, Super 8mm, 35mm VistaVision, etc.), Lens and Focus Settings as well as Current Aperture and Focus Distances. While using a normal Camera Actor is perfectly fine, to achieve a more cinematic feel and expose more camera settings it is recommended that you shoot your scenes with a Cine Camera Actor. 
 
 [REGION:note]
-Please see [**Using Cine Camera Actors**](Engine\Sequencer\HowTo\CineCameraActors) for more information. 
+Please see [Using Cine Camera Actors](Engine\Sequencer\HowTo\CineCameraActors) for more information. 
 [/REGION]
 
 
@@ -223,7 +223,7 @@ mvCPXslVb8Y
 Add a new recording, pick an Actor to record and click Record to start capturing. When stopping the recording, new Level Sequence assets will be created which you can then incorporate into your existing cinematic or you can use pieces from the recorded Level Sequence to create new sequences.
 
 [REGION:note]
-Please see [**Recording Gameplay**](Engine/Sequencer/HowTo/SequenceRecorder) for more information. 
+Please see [Recording Gameplay](Engine/Sequencer/HowTo/SequenceRecorder) for more information. 
 [/REGION]
 
 ##Render Movie Settings
@@ -238,6 +238,25 @@ This will display the **Render Movie Settings** window which you can use to defi
 
 Click the **Capture Movie** button to start the capture process of recording your scene. 
 
+### Custom Burn-ins
+
+[INCLUDE:Engine/Sequencer/HowTo/BurnIns#intro]
+
+[REGION:note]
+Please see [](Engine/Sequencer/HowTo/BurnIns) for a step-by-step guide on using burn-ins. 
+[/REGION]
+
+[PUBLISH:docs]
+### Import/Export EDLs
+
+[INCLUDE:Engine/Sequencer/HowTo/ImportExportEDLs#intro]
+
+[REGION:note]
+Please see [](Engine/Sequencer/HowTo/ImportExportEDLs) for a step-by-step guide on using EDLs.
+[/REGION]
+[/PUBLISH:docs]
+
+
 
 ## Workflow Considerations
 
@@ -245,5 +264,5 @@ Once you have a basic understanding of how the Sequencer Editor works, you may w
 
 
 [REGION:note]
-Please see [**Workflow Considerations**](Engine\Sequencer\Workflow\ToolUsage) for more information. 
+Please see [Workflow Considerations](Engine\Sequencer\Workflow\ToolUsage) for more information. 
 [/REGION]
\ No newline at end of file
