#UE4 Docs: Applying peer review feedback.

diff --git a/Engine/Documentation/Source/Engine/Performance/ReducingPackageSize/ReducingPackageSize.INT.udn b/Engine/Documentation/Source/Engine/Performance/ReducingPackageSize/ReducingPackageSize.INT.udn
--- a/Engine/Documentation/Source/Engine/Performance/ReducingPackageSize/ReducingPackageSize.INT.udn
+++ b/Engine/Documentation/Source/Engine/Performance/ReducingPackageSize/ReducingPackageSize.INT.udn
@@ -83,10 +83,15 @@ To set which level (or levels) should be used for this type of interaction, you
 
 	 ![](T_Project_Settings.png)
 
-1. Then, under **Project** in the **Maps & Modes** section, set the map in the **Default Maps** to a map that is specific to your project.
+1. Then, under **Project** in the **Maps & Modes** section look for the **Default Maps** section.
 
 	![](T_Set_Maps.png)
 
+1. Click on the arrow next to **Game Default Map** and from the asset list select the map you want to use for the games default starting level. 
+	
+	![](T_Project_Settings_Default_Game_Map.png)
+
+1. Repeat this process until you have levels setup for each required map type.
 ## Checking what Content will be Cooked
 To ensure that only the content that is related to your project is added to your APK file, you can check which content is being added to your project's APK by looking in the **Cooked** folder in your project.
 You can find the Cooked folder by going to **(ProjectName)\Saved\Cooked** and then choosing the folder with the format you cooked your project for.
diff --git a/Engine/Documentation/Source/Platforms/Android/GettingStarted/1/AndroidQS_1.INT.udn b/Engine/Documentation/Source/Platforms/Android/GettingStarted/1/AndroidQS_1.INT.udn
--- a/Engine/Documentation/Source/Platforms/Android/GettingStarted/1/AndroidQS_1.INT.udn
+++ b/Engine/Documentation/Source/Platforms/Android/GettingStarted/1/AndroidQS_1.INT.udn
@@ -1,30 +1,19 @@
 Availability:Public
-Title:1. Install the Android SDK
+Title:1. Required Android Setup
 Crumbs:%ROOT%, Platforms, Platforms/Android, Platforms/Android/GettingStarted
 Description: Install the CodeWorks for Android 1R4 SDK.
-Version:4.9
+Version:4.12
 Skilllevel:Intermediate
 Tags:Android
-
-[VAR:Steps]
-[OBJECT:Navigation]
-	[PARAM:previous]
-	[/PARAM]
-	[PARAM:current]
-	[/PARAM]
-	[PARAM:home]
-		[Quick Start Home](Platforms/Android/GettingStarted)
-	[/PARAM]
-	[PARAM:next]
-		[Next Step](Platforms/Android/GettingStarted/2 "%Platforms/Android/GettingStarted/2:title%")
-	[/PARAM]
-[/OBJECT]
-[/VAR]
+type:sequential
+parent:Platforms/Android/GettingStarted
+order:1
+Checkpoint:AndroidQS
 
 
-%Steps%
-
+[nav]
 
+## Steps
 The first thing we need to do to start developing for Android is to install CodeWorks for Android 1R4, which will 
 ensure that you have all the required programs and code needed to deploy to Android devices.
 In the following section, we will cover how you can get CodeWorks for Android 1R4 installed on your development PC.
@@ -95,7 +84,7 @@ It is highly recommended that you install CodeWorks for Android 1R4 while the ed
 
 	[/region]
 
-1. Double-click on the **AndroidWorks-1R1-windows.exe** file to initiate the install process.
+1. Double-click on the **CodeWorksforAndroid-1R4-windows.exe** file to initiate the install process.
 
 	![](AndroidWorks_11.png)
 
@@ -183,7 +172,7 @@ Use the links provide below to download and install the correct versions of Java
 
 	[/region]
 
-1. Double-click on the **AndroidWorks-1R1-osx.dmg** to load the CodeWorks for Android 1R4 installer.
+1. Double-click on the ** CodeWorks for Android 1R4-osx.dmg** to load the CodeWorks for Android 1R4 installer.
 	
 	![](AndroidWorks_Mac_00.png)
 
@@ -236,5 +225,8 @@ Use the links provide below to download and install the correct versions of Java
 [/EXCERPT:MacOSSetup]
 -->
 
-%Steps%
+## End Result
+When completed you will now have all of the required Android SDK's installed and ready for use in building a UE4 project that can work on Android  devices.
+
+[nav]
 
diff --git a/Engine/Documentation/Source/Platforms/Android/GettingStarted/2/AndroidQS_2.INT.udn b/Engine/Documentation/Source/Platforms/Android/GettingStarted/2/AndroidQS_2.INT.udn
--- a/Engine/Documentation/Source/Platforms/Android/GettingStarted/2/AndroidQS_2.INT.udn
+++ b/Engine/Documentation/Source/Platforms/Android/GettingStarted/2/AndroidQS_2.INT.udn
@@ -2,27 +2,20 @@ Availability:Public
 Title:2. Set Up Your Android Device
 Crumbs:%ROOT%, Platforms, Platforms/Android, Platforms/Android/GettingStarted
 Description:Set up your Android device for development.
-Version:4.9
+Version:4.12
 Skilllevel:Intermediate
+tags:Android
+tags:Mobile
+type:sequential
+parent:Platforms/Android/GettingStarted
+order:2
+Checkpoint:AndroidQS
 
 Now that Android Works has been installed it is now time to setup your Android device so that it can be used in development. In the following section we will go over all you need to know to setup your Android device so that it can be used to run your UE4 project.
-[VAR:Steps]
-[OBJECT:Navigation]
-	[PARAM:previous]
-		[Previous Step](Platforms/Android/GettingStarted/1 "%Platforms/Android/GettingStarted/1:title%")
-	[/PARAM]
-	[PARAM:current]
-	[/PARAM]
-	[PARAM:home]
-		[Quick Start Home](Platforms/Android/GettingStarted)
-	[/PARAM]
-	[PARAM:next]
-		[Next Step](Platforms/Android/GettingStarted/3 "%Platforms/Android/GettingStarted/3:title%")
-	[/PARAM]
-[/OBJECT]
-[/VAR]
 
-%Steps%
+[nav]
+
+## Steps
 
 %Globals:OSSelect%
 
@@ -203,7 +196,8 @@ If you do not see your device after you input the adb devices command, here are
 
 -->
 
+## End Result
+When completed you will now have an Android based smartphone that is ready for your UE4 project to be deployed to and run on. 
 
-
-%Steps%
+[nav]
 
diff --git a/Engine/Documentation/Source/Platforms/Android/GettingStarted/3/AndroidQS_3.INT.udn b/Engine/Documentation/Source/Platforms/Android/GettingStarted/3/AndroidQS_3.INT.udn
--- a/Engine/Documentation/Source/Platforms/Android/GettingStarted/3/AndroidQS_3.INT.udn
+++ b/Engine/Documentation/Source/Platforms/Android/GettingStarted/3/AndroidQS_3.INT.udn
@@ -2,31 +2,23 @@ Availability:Public
 Title:3. Create Your Project
 Crumbs:%ROOT%, Platforms, Platforms/Android, Platforms/Android/GettingStarted
 Description:Start a template project
-Version:4.9
+Version:4.12
 Skilllevel:Intermediate
+tags:Android
+tags:Mobile
+type:sequential
+parent:Platforms/Android/GettingStarted
+order:3
+Checkpoint:AndroidQS
 
-[VAR:Steps]
-[OBJECT:Navigation]
-	[PARAM:previous]
-		[Previous Step](Platforms/Android/GettingStarted/2 "%Platforms/Android/GettingStarted/2:title%")
-	[/PARAM]
-	[PARAM:current]
-	[/PARAM]
-	[PARAM:home]
-		[Quick Start Home](Platforms/Android/GettingStarted)
-	[/PARAM]
-	[PARAM:next]
-		[Next Step](Platforms/Android/GettingStarted/4 "%Platforms/Android/GettingStarted/4:title%")
-	[/PARAM]
-[/OBJECT]
-[/VAR]
+[nav]
+
+## Steps
 
 In the following section we will be creating a new project for use to test deployment to Android with.
 The project that we create will be a Blueprint based project that will use the Blueprint Third Person template.
 
 
-%Steps%
-
 
 1. From the **Unreal Projects Browser** create a new Blueprint based project with the following options. 
 * Use the Third Person Template
@@ -37,13 +29,15 @@ The project that we create will be a Blueprint based project that will use the B
 
 ![](new_android_project.png)
 
+1. Once that is completed, press the **Create Project** button to create the project.
 
 [REGION:tip]
 While we have selected a Blueprint based project for this example, any type of project, C++ or Blueprint, can be used to packaged and deployed to an Android device.
 [/REGION]
 
+## End Result
+When completed you will now have a UE4 project that is setup for use with a mobile device. 
 
-
-%Steps%
+[nav]
 
 
diff --git a/Engine/Documentation/Source/Platforms/Android/GettingStarted/4/AndroidQS_4.INT.udn b/Engine/Documentation/Source/Platforms/Android/GettingStarted/4/AndroidQS_4.INT.udn
--- a/Engine/Documentation/Source/Platforms/Android/GettingStarted/4/AndroidQS_4.INT.udn
+++ b/Engine/Documentation/Source/Platforms/Android/GettingStarted/4/AndroidQS_4.INT.udn
@@ -2,29 +2,19 @@ Availability:Public
 Title:4. Deploy Your Level
 Crumbs:%ROOT%, Platforms, Platforms/Android, Platforms/Android/GettingStarted
 Description:Launch your level with one-click deploying.
-Version:4.9
 Skilllevel:Intermediate
+tags:Android
+tags:Mobile
+type:sequential
+parent:Platforms/Android/GettingStarted
+order:4
+Checkpoint:AndroidQS
 
-[VAR:Steps]
-[OBJECT:Navigation]
-	[PARAM:previous]
-		[Previous Step](Platforms/Android/GettingStarted/3 "%Platforms/Android/GettingStarted/3:title%")
-	[/PARAM]
-	[PARAM:current]
-	[/PARAM]
-	[PARAM:home]
-		[Quick Start Home](Platforms/Android/GettingStarted)
-	[/PARAM]
-	[PARAM:next]
-		[Next Step](Platforms/Android/GettingStarted/5 "%Platforms/Android/GettingStarted/5:title%")
-	[/PARAM]
-[/OBJECT]
-[/VAR]
 
 
 In the following section we will be taking a look at how you can deploy a single level from the project you are working in to your Android device for quick testing. 
 
-%Steps%
+[nav]
 
 [region:note]
 For this section to work you will need to make sure that your Android device is plugged into your development PC via USB. If you have not done this already please do this before going any further.
@@ -47,6 +37,8 @@ The chart below shows you which Texture formats work with which Android based de
 
 [/REGION:tip]
 
+## Steps
+
 1.	To test your current level on a Android based device you first need to make sure you have the level you want to test open. 
 	From this example we will be using the **ThirdPersonExampleMap** level from the Blueprint based project that was created in the last step.
 	
@@ -79,6 +71,7 @@ The chart below shows you which Texture formats work with which Android based de
 
 	![](Deployed_Project_Android.png)
 
+## End Result
+When completed the UE4 project that you created will now be running on your Android based smartphone.
 
-
-%Steps%
+[nav]
diff --git a/Engine/Documentation/Source/Platforms/Android/GettingStarted/5/AndroidQS_5.INT.udn b/Engine/Documentation/Source/Platforms/Android/GettingStarted/5/AndroidQS_5.INT.udn
--- a/Engine/Documentation/Source/Platforms/Android/GettingStarted/5/AndroidQS_5.INT.udn
+++ b/Engine/Documentation/Source/Platforms/Android/GettingStarted/5/AndroidQS_5.INT.udn
@@ -2,29 +2,22 @@
 Title:5. Package Your Game
 Crumbs:%ROOT%, Platforms, Platforms/Android, Platforms/Android/GettingStarted
 Description:Package your game for installation.
-Version:4.9
+Version:4.12
 Skilllevel:Intermediate
+tags:Android
+tags:Mobile
+type:sequential
+parent:Platforms/Android/GettingStarted
+order:5
+Checkpoint:AndroidQS
 
-[VAR:Steps]
-[OBJECT:Navigation]
-	[PARAM:previous]
-		[Previous Step](Platforms/Android/GettingStarted/4 "%Platforms/Android/GettingStarted/4:title%")
-	[/PARAM]
-	[PARAM:current]
-	[/PARAM]
-	[PARAM:home]
-		[Quick Start Home](Platforms/Android/GettingStarted)
-	[/PARAM]
-	[PARAM:next]
-		[Next Step](Platforms/Android/GettingStarted/6 "%Platforms/Android/GettingStarted/6:title%")
-	[/PARAM]
-[/OBJECT]
-[/VAR]
+[nav]
 
 In order for you to be able to share your project with others or put it on a digital store to sell, the project needs to be packaged. 
 In the following section we will take a look at how you can package your finished project for deployment inside of UE4.
 
-%Steps%
+## Steps
+
 
 %Globals:OSSelect%
 
@@ -57,6 +50,8 @@ In the following section we will take a look at how you can package your finishe
 	[/PARAM]
 [/OBJECT]
 
+
+
 <!--
 [EXCERPT:WindowsSetup]
 
@@ -148,4 +143,6 @@ We are currently packaging in **Development** because we are testing this game,
 [/EXCERPT:MacOSSetup]
 -->
 
-%Steps%
+## End Result
+When completed you will now have a .APK file that contains your UE4  project that can now be deployed via the .BAT file to any Android based smartphone.
+[nav]
\ No newline at end of file
diff --git a/Engine/Documentation/Source/Platforms/Android/GettingStarted/6/AndroidQS_6.INT.udn b/Engine/Documentation/Source/Platforms/Android/GettingStarted/6/AndroidQS_6.INT.udn
--- a/Engine/Documentation/Source/Platforms/Android/GettingStarted/6/AndroidQS_6.INT.udn
+++ b/Engine/Documentation/Source/Platforms/Android/GettingStarted/6/AndroidQS_6.INT.udn
@@ -2,29 +2,20 @@ Availability:Public
 Title:6. On your Own
 Crumbs:%ROOT%, Platforms, Platforms/Android, Platforms/Android/GettingStarted
 Description:Read more about Android development with Unreal Engine.
-Version:4.9
+Version:4.12
 Skilllevel:Intermediate
+tags:Android
+tags:Mobile
+type:sequential
+parent:Platforms/Android/GettingStarted
+order:6
+Checkpoint:AndroidQS
 
 The following section will go over some things that you can do on your own to practice and refine your skills at deploying UE4 projects to Android devices.
 
-[VAR:Steps]
-[OBJECT:Navigation]
-	[PARAM:previous]
-		[Previous Step](Platforms/Android/GettingStarted/5 "%Platforms/Android/GettingStarted/5:title%")
-	[/PARAM]
-	[PARAM:current]
-	[/PARAM]
-	[PARAM:home]
-		[Quick Start Home](Platforms/Android/GettingStarted)
-	[/PARAM]
-	[PARAM:next]
-	[/PARAM]
-[/OBJECT]
-[/VAR]
+[nav]
 
 
-%Steps%
-
 Now that we have gone over how to get your UE4 project packaged and ready to be tested on your Android device, 
 try and accomplish the following using what you have learned from this Quick Start guide. 
 
@@ -54,4 +45,4 @@ For more information about designing game content and levels for mobile devices
 You can also preview your mobile game without deploying to a device by using the [](Platforms\Mobile\Previewer).
 
 
-%Steps%
+[nav]
diff --git a/Engine/Documentation/Source/Platforms/Android/GettingStarted/AndroidQuickStart.INT.udn b/Engine/Documentation/Source/Platforms/Android/GettingStarted/AndroidQuickStart.INT.udn
--- a/Engine/Documentation/Source/Platforms/Android/GettingStarted/AndroidQuickStart.INT.udn
+++ b/Engine/Documentation/Source/Platforms/Android/GettingStarted/AndroidQuickStart.INT.udn
@@ -2,13 +2,15 @@ Availability: Public
 Title: Android Quick Start
 Crumbs:%ROOT%, Platforms, Platforms/Android
 Description:Setting up for development for the Android platform.
-Version:4.9
+Version:4.12
 Skilllevel:Intermediate
 tags:Android
 tags:Mobile
-parent:Platforms/Android
-order:2
 type:quickstart
+parent:Platforms/Android
+order:1
+Checkpoint:AndroidQS
+
 
 [VAR:Topic]
 [OBJECT:Topic]
@@ -44,6 +46,8 @@ type:quickstart
 [/OBJECT]
 [/VAR]
 
+![](AndroidTopic.jpg)
+
 The Android Quick Start will walk you through setting up your computer and your Android device for deployment, as well as putting a template game on your device using both the launching and packaging features.
 
 ## Goals
@@ -58,27 +62,13 @@ The focus of the Android Quick start is to walk you through the basics of settin
 * How to deploy a level to your Android device
 * How to package your game for release.
 
-[OBJECT:TopicButtonList]
-	[PARAM:icon]
-		![](%ROOT%/build_icon.png)(convert:false)
-	[/PARAM]
-	[PARAM:title]
-		Steps
-	[/PARAM]
-	[PARAM:description]
-	[/PARAM]
-	[PARAM:links]
-	* [](Platforms/Android/GettingStarted/1 "%Platforms/Android/GettingStarted/1:description%")
-	* [](Platforms/Android/GettingStarted/2 "%Platforms/Android/GettingStarted/2:description%")
-	* [](Platforms/Android/GettingStarted/3 "%Platforms/Android/GettingStarted/3:description%")
-	* [](Platforms/Android/GettingStarted/4 "%Platforms/Android/GettingStarted/4:description%")
-	* [](Platforms/Android/GettingStarted/5 "%Platforms/Android/GettingStarted/5:description%")
-	* [](Platforms/Android/GettingStarted/6 "%Platforms/Android/GettingStarted/6:description%")
-	[/PARAM]
-[/OBJECT]
+[DIR(output:"checkpoint" parent:"Platforms/Android/GettingStarted")]
 
 [region:note]
 Unreal Engine 4 (UE4) now uses a special version of the Android SDK called **CodeWorks for Android 1R4**. This version of CodeWorks for Android was designed specifically for use with UE4 Android based projects. If you have other versions of the Android SDK installed, it is **highly recommend** that you uninstall them and then re-install CodeWorks for Android using the CodeWorks for Android 1R4 installer that is bundled with UE4. This way you can ensure that the Android SDK you are using will be compatible with UE4.
 [/region]
 
+[REGION:call_to_action]
+ [Click to Start](Platforms/Android/GettingStarted/1 "%Platforms/Android/GettingStarted/1:title%")
+[/REGION]
 
diff --git a/Engine/Documentation/Source/Platforms/Android/Reference/AndroidGettingStarted.INT.udn b/Engine/Documentation/Source/Platforms/Android/Reference/AndroidGettingStarted.INT.udn
--- a/Engine/Documentation/Source/Platforms/Android/Reference/AndroidGettingStarted.INT.udn
+++ b/Engine/Documentation/Source/Platforms/Android/Reference/AndroidGettingStarted.INT.udn
@@ -2,7 +2,7 @@
 Title: Android Development Reference
 Crumbs:%ROOT%, Platforms, Platforms/Android
 Description:How to install different Android SDKs, set environment variables, and work with texture formats.
-Version: 4.9
+Version: 4.12
 SkillLevel: Intermediate
 Parent:Platforms/Android
 tags:Mobile
@@ -13,11 +13,12 @@ type:reference
 
 %Globals:OSSelect%
 
+## Overview
+The Android Development Reference pages contains a high level explanation of the various software components and system settings you will need to make sure you have installed and setup on the computer you are using to develop your UE4 Android project.
 
-
+## Android SDK Requirements
 Here is what we currently require for Android game development:
 
-
 |Project Type | Components | Environment Variables |
 | ------------- | ---------- | --------------------- |
 | Blueprint Only Projects |  [INCLUDE:#bpcomponents] |      [INCLUDE:#bpenvvar] |
diff --git a/Engine/Documentation/Source/Platforms/GearVR/Prerequisites/GearVRPrerequisites.INT.udn b/Engine/Documentation/Source/Platforms/GearVR/Prerequisites/GearVRPrerequisites.INT.udn
--- a/Engine/Documentation/Source/Platforms/GearVR/Prerequisites/GearVRPrerequisites.INT.udn
+++ b/Engine/Documentation/Source/Platforms/GearVR/Prerequisites/GearVRPrerequisites.INT.udn
@@ -36,9 +36,10 @@ Before we can get into UE4 Gear VR development, we first must make sure that the
 The USB cord is of vital importance, as this is how you transfer data between your development PC and your Smartphone.
 
 * Only certain types of Samsung Galaxy Smartphones will work with the Gear VR Headset. 
-	The [Oculus website](https://support.oculus.com/hc/en-us/articles/204113096-Samsung-Gear-VR-Innovator-Edition-Samsung-Galaxy-Note-4-Compatibility) has a list of phone models that will work with the Gear VR Headset.
+	The [Oculus website](https://support.oculus.com/help/oculus/1099393013405406/?ref=hc_fnav) has a list of phone models that will work with the Gear VR Headset.
 	Make sure to double check that the Samsung Galaxy Note 4 Smartphone you are buying is the correct model that will work with the Gear VR Headset.
 
+
 * Make sure to download and install the specific USB drivers for your Galaxy Note 4 from the [Samsung Website](http://www.samsung.com/us/support/downloads). 
 
 * Download Unreal Engine version **4.7.4** or later from the [UE4 Launcher](https://www.unrealengine.com/) or [Github](https://www.unrealengine.com/ue4-on-github). 
diff --git a/Engine/Documentation/Source/Platforms/GearVR/QuickStart/1/GearVRQuickStart_1.INT.udn b/Engine/Documentation/Source/Platforms/GearVR/QuickStart/1/GearVRQuickStart_1.INT.udn
--- a/Engine/Documentation/Source/Platforms/GearVR/QuickStart/1/GearVRQuickStart_1.INT.udn
+++ b/Engine/Documentation/Source/Platforms/GearVR/QuickStart/1/GearVRQuickStart_1.INT.udn
@@ -3,27 +3,17 @@ Title:1. Gear VR Project Creation
 Crumbs:%ROOT%, Platforms/GearVR
 Description:How to create a UE4 C++ based project that can be used with the Gear VR Headset.
 Navigation:platform
-Skilllevel: Intermediate
 Version: 4.9
+Skilllevel: Intermediate
+tags:VR
+tags:GearVR
+type:sequential
+parent:Platforms/GearVR/QuickStart
+order:1
+Checkpoint:GearVRQS
 
+[nav]
 
-[VAR:Steps]
-[OBJECT:Navigation]
-	[PARAM:previous]
-	[/PARAM]
-	[PARAM:current]
-	[/PARAM]
-	[PARAM:home]
-		[Quick Start Home](Platforms/GearVR/QuickStart)
-	[/PARAM]
-	[PARAM:next]
-		[Next Step](Platforms/GearVR/QuickStart/2 "%Platforms/GearVR/QuickStart/2:title%")
-	[/PARAM]
-[/OBJECT]
-[/VAR]
-
-
-%Steps%
 
 With the release of Unreal Engine 4.8 and later, setting up a project that works with the Samsung Gear VR Headset is no different than setup than a normal UE4 project. 
 In the following section we will go over everything you need to know in order to create a new project that will work with the Gear VR headset.
@@ -33,6 +23,8 @@ If you are using UE4 version 4.7 you will need to create your project as a C++ p
 In UE4 version 4.8 and later you do not have to do this and can use either C++ or Blueprints for your projects.
 [/REGION]
 
+## Steps
+
 1. First make sure that you have UE4 version **4.7.4** or later downloaded and installed on your development PC. 
 You can get UE4 version **4.7.4** or later from the [Unreal Engine Launcher](http://tinyurl.com/oawb7x2) or from [Github](https://www.unrealengine.com/ue4-on-github). 
 
@@ -58,6 +50,7 @@ You can get UE4 version **4.7.4** or later from the [Unreal Engine Launcher](htt
 1. Finally give the project a location to be saved to and a name. The name for this example the project name will be **GearVR_QS**. Once everything has been completed press the green **Create Project** button to create the project.
 	![](GVR_New_Project_02.png)
 
+## End Result
+When completed you should now have a new First Person project created with all of the correct settings for the project to be used with a  Gear VR HMD.
 
-
-%Steps%
+[nav]
diff --git a/Engine/Documentation/Source/Platforms/GearVR/QuickStart/2/GearVRQuickStart_2.INT.udn b/Engine/Documentation/Source/Platforms/GearVR/QuickStart/2/GearVRQuickStart_2.INT.udn
--- a/Engine/Documentation/Source/Platforms/GearVR/QuickStart/2/GearVRQuickStart_2.INT.udn
+++ b/Engine/Documentation/Source/Platforms/GearVR/QuickStart/2/GearVRQuickStart_2.INT.udn
@@ -2,35 +2,25 @@ Availability:Public
 Title:2. Gear VR Project Setup
 Crumbs:%ROOT%, Platforms, Platforms/GearVR
 Description:How to setup your C++ based UE4 project to be used with the Gear VR Headset.
-Navigation:platform
-Skilllevel: Intermediate
 Version: 4.9
+Skilllevel: Intermediate
+tags:VR
+tags:GearVR
+type:sequential
+parent:Platforms/GearVR/QuickStart
+order:2
+Checkpoint:GearVRQS
 
-[VAR:Steps]
-[OBJECT:Navigation]
-	[PARAM:previous]
-	[Previous Step](Platforms/GearVR/QuickStart/1 "%Platforms/GearVR/QuickStart/1:title%")
-	[/PARAM]
-	[PARAM:current]
-	[/PARAM]
-	[PARAM:home]
-		[Quick Start Home](Platforms/GearVR/QuickStart)
-	[/PARAM]
-	[PARAM:next]
-		[Next Step](Platforms/GearVR/QuickStart/3 "%Platforms/GearVR/QuickStart/3:title%")
-	[/PARAM]
-[/OBJECT]
-[/VAR]
-
-%Steps%
+[nav]
 
+## Steps
 
-1. Once the project has finished being created and compiled, Visual Studio 2013 will automatically open (if already installed) along with the editor.
+1. Once the project has finished being created and compiled, Visual Studio 2015 will automatically open (if already installed) along with the editor.
 
 	![](GVR_Editor_Opened.png)
 
 	[region:note]
-	If you are using Unreal Engine 4.8 or later and a Blueprint based project, Visual Studio 2013 will not open up as it is not needed. 
+	If you are using Unreal Engine 4.8 or later and a Blueprint based project, Visual Studio 2015 will not open up as it is not needed. 
 	[/region]
 
 1. Once the UE4 Editor has loaded, from the **Main Toolbar**, open the **Windows** tab and select the **Plugins** menu. 
@@ -114,4 +104,7 @@ Version: 4.9
 	[/region]
 	-->
 
-%Steps%
+[nav]
+
+## End Result
+When completed you will now have a UE4 project that is ready to be built, deployed and viewed on your Samsung Galaxy smartphone.
\ No newline at end of file
diff --git a/Engine/Documentation/Source/Platforms/GearVR/QuickStart/3/GearVRQuickStart_3.INT.udn b/Engine/Documentation/Source/Platforms/GearVR/QuickStart/3/GearVRQuickStart_3.INT.udn
--- a/Engine/Documentation/Source/Platforms/GearVR/QuickStart/3/GearVRQuickStart_3.INT.udn
+++ b/Engine/Documentation/Source/Platforms/GearVR/QuickStart/3/GearVRQuickStart_3.INT.udn
@@ -2,32 +2,21 @@ Availability:Public
 Title:3. Gear VR Project Packaging & Deployment
 Crumbs:%ROOT%, Platforms/GearVR
 Description:Instructions on how to deploy projects form UE4 to use with the Samsung Gear VR Headset.
-Navigation:platform
-Platform:Gear VR
+Version: 4.9
+Skilllevel: Intermediate
+tags:VR
+tags:GearVR
+type:sequential
+parent:Platforms/GearVR/QuickStart
+order:3
+Checkpoint:GearVRQS
 Related: Platforms\GearVR\Prerequisites
 Related: Platforms\GearVR\BestPractices
 Related: Platforms\GearVR\Debugging
-Skilllevel: Intermediate
-Version: 4.9
 
-[VAR:Steps]
-[OBJECT:Navigation]
-	[PARAM:previous]
-	[Previous Step](Platforms/GearVR/QuickStart/2 "%Platforms/GearVR/QuickStart/2:title%")
-	[/PARAM]
-	[PARAM:current]
-	[/PARAM]
-	[PARAM:home]
-		[Quick Start Home](Platforms/GearVR/QuickStart)
-	[/PARAM]
-	[PARAM:next]
-		
-	[/PARAM]
-[/OBJECT]
-[/VAR]
-
-%Steps%
+[nav]
 
+## Steps
 
 1. First make sure that your Samsung Galaxy Note 4 Smartphone is connected to a USB cord and plugged into a USB port on the PC you are developing on.
 
@@ -63,5 +52,7 @@ Version: 4.9
 
 	![](GVR_On_Phone.png)
 
+[nav]
 
-%Steps%
+## End Result
+When completed you will be able to launch the project and view it using your  Samsung Galaxy smartphone and the Gear VR HMD.
diff --git a/Engine/Documentation/Source/Platforms/GearVR/QuickStart/4/GearVRQuickStart_4.INT.udn b/Engine/Documentation/Source/Platforms/GearVR/QuickStart/4/GearVRQuickStart_4.INT.udn
new file mode 100644
--- /dev/null
+++ b/Engine/Documentation/Source/Platforms/GearVR/QuickStart/4/GearVRQuickStart_4.INT.udn
@@ -0,0 +1,40 @@
+Availability:Public
+Title:4. On Your Own!
+Crumbs:%ROOT%, Platforms/GearVR
+Description:Using the methods provided during this guide, polish up your level with additional content and read up on related topics. 
+Navigation:platform
+Version: 4.9
+Skilllevel: Intermediate
+tags:VR
+tags:GearVR
+type:sequential
+parent:Platforms/GearVR/QuickStart
+order:4
+Checkpoint:GearVRQS
+
+[nav]
+
+Now that you can deploy a UE4 project to your Android based smartphone and view it with your Gear VR HMD try adding the following items.
+
+* Use the [Migrate Tool](Engine/Content/Browser/UserGuide/Migrate/) to move content from one of the mobile content examples into your project for use.
+
+* Adjust the settings of the [Default Device Profiles](Platforms/Mobile/Performance/TipsAndTricks) to ensure that your project will run as fast as possible when deployed to your Android based smartphone.
+
+* Use the [GPU Profiler](Engine/Performance/GPU/) to help track your projects performance while you are building it.
+
+
+
+As for specifics covered in this quick start:
+
+* For more information on setting up the Android SDK, see: [Install the Android SDK](Platforms/Android/GettingStarted/1/)
+
+* For more information on setting up your Android Device for UE4 Development, see: [Set Up Your Android Device](Platforms/Android/GettingStarted/2/)
+
+* For more information on what features Android devices support, see: [Android Device Compatibility](Platforms/Android/DeviceCompatibility/)
+
+* For more information on reducing the packaged size of your project, see: [Reducing APK Package Size](Platforms/Android/ReducingAPKSize/)
+
+
+
+
+[nav]
\ No newline at end of file
diff --git a/Engine/Documentation/Source/Platforms/GearVR/QuickStart/GearVRQuickStart.INT.udn b/Engine/Documentation/Source/Platforms/GearVR/QuickStart/GearVRQuickStart.INT.udn
--- a/Engine/Documentation/Source/Platforms/GearVR/QuickStart/GearVRQuickStart.INT.udn
+++ b/Engine/Documentation/Source/Platforms/GearVR/QuickStart/GearVRQuickStart.INT.udn
@@ -5,27 +5,29 @@ Description:Instructions on how to quickly setup and deploy projects for use wit
 Navigation:platform
 Platform:Gear VR
 Skilllevel: Intermediate
-Version: 4.9
+Version:4.9
 type:quickstart
 parent:Platforms/GearVR
 tags:VR
 tags:GearVR
+order:1
+Checkpoint:GearVRQS
+prereq:Platforms\GearVR\Prerequisites
 
+![](GVR_Headset_00.png)
+
+
+## Goals
 The Gear VR Quick Start will walk you through how to setup your development PC and Samsung Galaxy Note 4 Smartphone to work with the Samsung Gear VR headset.
 
-[OBJECT:TopicButtonList]
-	[PARAM:icon]
-		![](%ROOT%/build_icon.png)(convert:false)
-	[/PARAM]
-	[PARAM:title]
-		Steps
-	[/PARAM]
-	[PARAM:description]
-	[/PARAM]
-	[PARAM:links]
-	* [](Platforms/GearVR/QuickStart/1 "%Platforms/GearVR/QuickStart/1:description%")
-	* [](Platforms/GearVR/QuickStart/2 "%Platforms/GearVR/QuickStart/2:description%")
-	* [](Platforms/GearVR/QuickStart/3 "%Platforms/GearVR/QuickStart/3:description%")
-	[/PARAM]
-[/OBJECT]
+## Objectives
+* Create a new UE4 project that is targeted towards Gear VR development.
+* Setting the required project settings so your project can be used with Gear VR.
+* Deploying and viewing your UE4 project with the Gear VR HMD.
+
+[DIR(output:"checkpoint" parent:"Platforms/GearVR/QuickStart")]
+
+[REGION:call_to_action]
+ [Click to Start](Platforms/GearVR/QuickStart/1 "%Platforms/GearVR/QuickStart/1:title%")
+[/REGION]
 
