Capitalized 'NVIDIA' to conform to their trademark standard.

#UE4doc

diff --git a/Engine/Documentation/Source/Programming/Development/BuildingUnrealEngine/BuildingUnrealEngine.INT.udn b/Engine/Documentation/Source/Programming/Development/BuildingUnrealEngine/BuildingUnrealEngine.INT.udn
--- a/Engine/Documentation/Source/Programming/Development/BuildingUnrealEngine/BuildingUnrealEngine.INT.udn
+++ b/Engine/Documentation/Source/Programming/Development/BuildingUnrealEngine/BuildingUnrealEngine.INT.udn
@@ -204,7 +204,7 @@ Our development and support teams currently use the latest version of Ubuntu; as
 
 [REGION:note]		
 If you're unable to create an OpenGL context as you're working through these steps, you may need to update your graphics drivers (refer to Ubuntu's documentation<!-- and our list of [](Platforms/Linux/BeginnerLinuxDeveloper/SupportedGfxDrivers) -->). 
-For an optimal experience, we recommend using NVidia binary drivers (versions equal to or greater than 343.xx).
+For an optimal experience, we recommend using NVIDIA binary drivers (versions equal to or greater than 343.xx).
 [/REGION]
 
 1.	Navigate to the editor's binary path by entering `cd Engine/Binaries/Linux/` into the terminal.
