Addressed WesB's peer review comments.

#UE4doc

diff --git a/Engine/Documentation/Source/Platforms/Linux/BeginnerLinuxDeveloper/SupportedGfxDrivers/SupportedGfxDrivers.INT.udn b/Engine/Documentation/Source/Platforms/Linux/BeginnerLinuxDeveloper/SupportedGfxDrivers/SupportedGfxDrivers.INT.udn
--- a/Engine/Documentation/Source/Platforms/Linux/BeginnerLinuxDeveloper/SupportedGfxDrivers/SupportedGfxDrivers.INT.udn
+++ b/Engine/Documentation/Source/Platforms/Linux/BeginnerLinuxDeveloper/SupportedGfxDrivers/SupportedGfxDrivers.INT.udn
@@ -2,7 +2,7 @@ Availability: Docs
 Title: Supported Graphics Drivers
 Crumbs: %ROOT%, Platforms, Platforms/Linux, Platforms/Linux/BeginnerLinuxDeveloper
 Description:This page provides an overview of supported graphics drivers for Linux game developers.
-version: 4.11
+version: 4.13
 parent:Platforms/Linux/BeginnerLinuxDeveloper
 order: 1
 tags:Linux
@@ -21,9 +21,6 @@ related:GettingStarted/RecommendedSpecifications
 [PARAM:title]
 	%Platforms/Linux/BeginnerLinuxDeveloper/SupportedGfxDrivers:title%
 [/PARAM]
-[PARAM:description]
-	%Platforms/Linux/BeginnerLinuxDeveloper/SupportedGfxDrivers:description%
-[/PARAM]
 [PARAM:path]
 	[RELATIVE:Platforms/Linux/BeginnerLinuxDeveloper/SupportedGfxDrivers]
 [/PARAM]
@@ -37,9 +34,7 @@ related:GettingStarted/RecommendedSpecifications
 
 Unreal Engine requires compatible hardware and software support in order to run smoothly on your 
 machine. In an effort to help you determine your ideal development environment, we're providing 
-the following table of supported graphics drivers for Linux game development.
-
-## Supported Graphics Drivers
+the following table of supported graphics drivers for developing games on Linux.
 
 [INCLUDE:#SupportedGfxDrivers]
 
@@ -47,14 +42,14 @@ Graphics drivers are typically provided by manufacturers for testing and integra
 Engine 4 (UE4). If you want to get the most out of UE4 while maintaining hardware performance, please 
 make sure to update your graphics driver appropriately.
 
-### NVidia Graphics Drivers
+## NVIDIA Graphics Drivers
 
-If you're running an NVidia graphics driver version older than version 347.09, you should update your graphics driver 
-to prevent your system from crashing while running the engine. Also, it's important to note that NVidia's 
+If you're running an NVIDIA graphics driver version older than version 347.09, consider updating your graphics driver 
+to prevent your system from crashing while running the engine. Also, it's important to note that NVIDIA's 
 latest graphics drivers aren't guaranteed to work with UE4. Please refer back to our list of supported graphics
 drivers in order to verify that your graphics driver will work properly with the engine.
 
-### AMD Graphics Drivers
+## AMD Graphics Drivers
 
 If you're running an AMD graphics driver version older than 14.301.1001.0, you should update your graphics driver to prevent 
 your system from crashing while running the engine.
@@ -63,7 +58,7 @@ your system from crashing while running the engine.
 [EXCERPT:SupportedGfxDrivers]
 | **Manufacturer** | **Supported Graphics Driver Versions** | **Suggested Graphics Driver Version** |
 | ---------------- | -------------------------------------- | ------------------------------------- |
-| **NVidia**       | Between 347.09 and 361.43              | 361.75                                |
+| **NVIDIA**       | Between 347.09 and 361.43              | 361.75                                |
 | **AMD**          | 14.301.1001.0 or newer                 | 15.30.1025.1001                       |
 | **Intel**        | Any                                    | Any                                   |
 [/EXCERPT:SupportedGfxDrivers]
