Reorganized Landing page per JeffW's peer review comment.

#UE4doc

diff --git a/Engine/Documentation/Source/Platforms/Linux/Linux.INT.udn b/Engine/Documentation/Source/Platforms/Linux/Linux.INT.udn
--- a/Engine/Documentation/Source/Platforms/Linux/Linux.INT.udn
+++ b/Engine/Documentation/Source/Platforms/Linux/Linux.INT.udn
@@ -1,29 +1,28 @@
 Availability: Public
-Title: Linux Game Development
+Title: Linux Game Development[](Platforms/Linux/BeginnerLinuxDeveloper/SettingUpAnUnrealWorkflow/5)
 Crumbs:%ROOT%, Platforms
 Description:Developing games for the Linux platform.
 Platform:Linux
 Navigation:platform
-Version: 4.11
+Version: 4.13
 type:landing
 parent:Platforms
 tags:Linux
-related: Platforms/Linux/BeginnerLinuxDeveloper/SupportedGfxDrivers
+related: GettingStarted/RecommendedSpecifications
 related: Platforms/Linux/GettingStarted
 
 Fueling the future of Linux game development is at your fingertips! Now, more than ever, you can harness the power of
 Unreal Engine 4 (UE4) to create and share games with Linux users worldwide. If you're new to Linux game development,
-you'll want to work your way through our [](Platforms/Linux/BeginnerLinuxDeveloper) learning path before moving 
-onto more advanced topics.
+you'll want to work your way through the Quick Start Guide before moving onto more advanced topics.
 
-## Learning Paths
+## Getting Started
 
 [REGION:topics third]
-%Platforms/Linux/BeginnerLinuxDeveloper:topic%
-<!--
-%Platforms/Linux/IntermediateLinuxDeveloper:topic%
-%Platforms/Linux/AdvancedLinuxDeveloper:topic%
--->
+%Platforms/Linux/BeginnerLinuxDeveloper/SettingUpAnUnrealWorkflow:topic%
 [/REGION]
 
-If you're having trouble running UE4 on your machine, make sure you meet our [](GettingStarted/RecommendedSpecifications).
\ No newline at end of file
+## Reference
+
+[REGION:topics third]
+%Platforms/Linux/BeginnerLinuxDeveloper/SupportedGfxDrivers:topic%
+[/REGION]
\ No newline at end of file
