Converting these Directory Landing pages to being NoPublish pages per JeffW's editorial review comments.
This is part of the reorganization effort of the new Linux documentation pages.
I'm preserving these pages in case we ever decide to go back to a directed learning path paradigm... also, these pages can still be harvested for more content.

#UE4doc

diff --git a/Engine/Documentation/Source/Platforms/Linux/AdvancedLinuxDeveloper/AdvancedLinuxDeveloper.INT.udn b/Engine/Documentation/Source/Platforms/Linux/AdvancedLinuxDeveloper/AdvancedLinuxDeveloper.INT.udn
--- a/Engine/Documentation/Source/Platforms/Linux/AdvancedLinuxDeveloper/AdvancedLinuxDeveloper.INT.udn
+++ b/Engine/Documentation/Source/Platforms/Linux/AdvancedLinuxDeveloper/AdvancedLinuxDeveloper.INT.udn
@@ -1,10 +1,10 @@
-Availability: Docs
+Availability: NoPublish
 Title: Advanced Linux Developers
 Crumbs:%ROOT%, Platforms, Platforms/Linux
 Description:Learn advanced game development techniques in UE4.
 Platform:Linux
 Navigation:platform
-Version: 4.11
+Version: 4.13
 type:landing
 parent:Platforms/Linux
 order: 3
@@ -19,7 +19,7 @@ SkillLevel: Advanced
 	![](%ROOT%/tutorial_icon.png)(convert:false)
 [/PARAM]
 [PARAM:title]
-	%Platforms/Linux/AdvancedLinuxDeveloper:title%
+	%Platforms/Linux/AdvancedLinuxDeveloper:title%[](Platforms/Linux/BeginnerLinuxDeveloper)
 [/PARAM]
 [PARAM:description]
 	%Platforms/Linux/AdvancedLinuxDeveloper:description%
diff --git a/Engine/Documentation/Source/Platforms/Linux/BeginnerLinuxDeveloper/BeginnerLinuxDeveloper.INT.udn b/Engine/Documentation/Source/Platforms/Linux/BeginnerLinuxDeveloper/BeginnerLinuxDeveloper.INT.udn
--- a/Engine/Documentation/Source/Platforms/Linux/BeginnerLinuxDeveloper/BeginnerLinuxDeveloper.INT.udn
+++ b/Engine/Documentation/Source/Platforms/Linux/BeginnerLinuxDeveloper/BeginnerLinuxDeveloper.INT.udn
@@ -1,10 +1,10 @@
-Availability: Docs
+Availability: NoPublish
 Title: Beginner Linux Developers
 Crumbs:%ROOT%, Platforms, Platforms/Linux
 Description:Learn how to setup a typical workflow for developing games on Linux.
 Platform:Linux
 Navigation:platform
-Version: 4.11
+Version: 4.13
 type:landing
 parent:Platforms/Linux
 order: 1
diff --git a/Engine/Documentation/Source/Platforms/Linux/IntermediateLinuxDeveloper/IntermediateLinuxDeveloper.INT.udn b/Engine/Documentation/Source/Platforms/Linux/IntermediateLinuxDeveloper/IntermediateLinuxDeveloper.INT.udn
--- a/Engine/Documentation/Source/Platforms/Linux/IntermediateLinuxDeveloper/IntermediateLinuxDeveloper.INT.udn
+++ b/Engine/Documentation/Source/Platforms/Linux/IntermediateLinuxDeveloper/IntermediateLinuxDeveloper.INT.udn
@@ -1,10 +1,10 @@
-Availability: Docs
+Availability: NoPublish
 Title: Intermediate Linux Developers
 Crumbs:%ROOT%, Platforms, Platforms/Linux
 Description:Learn how to cook, package, and share your games.
 Platform:Linux
 Navigation:platform
-Version: 4.11
+Version: 4.13
 type:landing
 parent:Platforms/Linux
 order: 2
