Converted 'Setting up and Unreal Workflow' from being a How To Page to being a Quickstart Page per JeffW's editorial review.

#UE4doc

diff --git a/Engine/Documentation/Source/Platforms/Linux/BeginnerLinuxDeveloper/SettingUpAnUnrealWorkflow/1/SettingUpAnUnrealWorkflow_1.INT.udn b/Engine/Documentation/Source/Platforms/Linux/BeginnerLinuxDeveloper/SettingUpAnUnrealWorkflow/1/SettingUpAnUnrealWorkflow_1.INT.udn
--- a/Engine/Documentation/Source/Platforms/Linux/BeginnerLinuxDeveloper/SettingUpAnUnrealWorkflow/1/SettingUpAnUnrealWorkflow_1.INT.udn
+++ b/Engine/Documentation/Source/Platforms/Linux/BeginnerLinuxDeveloper/SettingUpAnUnrealWorkflow/1/SettingUpAnUnrealWorkflow_1.INT.udn
@@ -1,9 +1,9 @@
 Availability: Docs
-Title: 1 - Setting up Git for UE4
+Title: 1 - Required Setup
 Crumbs:%ROOT%, Platforms, Platforms/Linux, Platforms/Linux/BeginnerLinuxDeveloper, Platforms/Linux/BeginnerLinuxDeveloper/SettingUpAnUnrealWorkflow
 Description:Learn how to setup Git on Linux.
-Version: 4.11
-type:HowTo
+Version: 4.13
+type:QuickStart
 parent:Platforms/Linux/BeginnerLinuxDeveloper/SettingUpAnUnrealWorkflow
 order: 1
 tags:Linux
@@ -16,7 +16,7 @@ SkillLevel:Beginner
 	[PARAM:current]
 	[/PARAM]
 	[PARAM:home]
-		[Home](Platforms/Linux/BeginnerLinuxDeveloper/SettingUpAnUnrealWorkflow)
+		[Linux Quick Start Home](Platforms/Linux/BeginnerLinuxDeveloper/SettingUpAnUnrealWorkflow)
 	[/PARAM]
 	[PARAM:next]
 		[Next Step](Platforms/Linux/BeginnerLinuxDeveloper/SettingUpAnUnrealWorkflow/2 "%Platforms/Linux/BeginnerLinuxDeveloper/SettingUpAnUnrealWorkflow/2:title%")
@@ -30,7 +30,7 @@ Currently, we don't supply a binary installer of UE4 for Linux users. The good n
 code that we use to develop the engine. If you want to access our source code, you'll have to register for an Epic Games account, sign up for a GitHub account, 
 associate your GitHub username with your Epic Games account, and setup Git on your Linux machine.
 
-## Getting an Epic Games Account
+## Registering for an Epic Games Account
 
 If you want to create games with Unreal Engine or contribute to the future of UE4 on Linux, you'll need to have a valid Epic Games account. 
 
@@ -39,7 +39,7 @@ If you want to create games with Unreal Engine or contribute to the future of UE
 There are many benefits for having an account with us, including (but not limited to) access to the Unreal Engine community, where you'll 
 connect with other UE4 Linux developers, and where you'll be able to keep up with the latest news about Unreal Engine.
 
-## Setup Git
+## Setting up Git
 
 Now that you've registered for an Epic Games account, you're ready to work through the following steps to gain access to the UE4 source code on GitHub.
 
diff --git a/Engine/Documentation/Source/Platforms/Linux/BeginnerLinuxDeveloper/SettingUpAnUnrealWorkflow/2/SettingUpAnUnrealWorkflow_2.INT.udn b/Engine/Documentation/Source/Platforms/Linux/BeginnerLinuxDeveloper/SettingUpAnUnrealWorkflow/2/SettingUpAnUnrealWorkflow_2.INT.udn
--- a/Engine/Documentation/Source/Platforms/Linux/BeginnerLinuxDeveloper/SettingUpAnUnrealWorkflow/2/SettingUpAnUnrealWorkflow_2.INT.udn
+++ b/Engine/Documentation/Source/Platforms/Linux/BeginnerLinuxDeveloper/SettingUpAnUnrealWorkflow/2/SettingUpAnUnrealWorkflow_2.INT.udn
@@ -2,8 +2,8 @@ Availability: Docs
 Title: 2 - Downloading UE4 on Linux
 Crumbs:%ROOT%, Platforms, Platforms/Linux, Platforms/Linux/BeginnerLinuxDeveloper, Platforms/Linux/BeginnerLinuxDeveloper/SettingUpAnUnrealWorkflow
 Description:Learn how to fork and clone Unreal Engine source code.
-Version: 4.11
-type:HowTo
+Version: 4.13
+type:QuickStart
 parent:Platforms/Linux/BeginnerLinuxDeveloper/SettingUpAnUnrealWorkflow
 order: 2
 tags:Linux
@@ -17,7 +17,7 @@ SkillLevel:Beginner
 	[PARAM:current]
 	[/PARAM]
 	[PARAM:home]
-		[Home](Platforms/Linux/BeginnerLinuxDeveloper/SettingUpAnUnrealWorkflow)
+		[Linux Quick Start Home](Platforms/Linux/BeginnerLinuxDeveloper/SettingUpAnUnrealWorkflow)
 	[/PARAM]
 	[PARAM:next]
 		[Next Step](Platforms/Linux/BeginnerLinuxDeveloper/SettingUpAnUnrealWorkflow/3 "%Platforms/Linux/BeginnerLinuxDeveloper/SettingUpAnUnrealWorkflow/3:title%")
diff --git a/Engine/Documentation/Source/Platforms/Linux/BeginnerLinuxDeveloper/SettingUpAnUnrealWorkflow/3/SettingUpAnUnrealWorkflow_3.INT.udn b/Engine/Documentation/Source/Platforms/Linux/BeginnerLinuxDeveloper/SettingUpAnUnrealWorkflow/3/SettingUpAnUnrealWorkflow_3.INT.udn
--- a/Engine/Documentation/Source/Platforms/Linux/BeginnerLinuxDeveloper/SettingUpAnUnrealWorkflow/3/SettingUpAnUnrealWorkflow_3.INT.udn
+++ b/Engine/Documentation/Source/Platforms/Linux/BeginnerLinuxDeveloper/SettingUpAnUnrealWorkflow/3/SettingUpAnUnrealWorkflow_3.INT.udn
@@ -2,8 +2,8 @@ Availability: Docs
 Title: 3 - Building UE4 on Linux
 Crumbs:%ROOT%, Platforms, Platforms/Linux, Platforms/Linux/BeginnerLinuxDeveloper, Platforms/Linux/BeginnerLinuxDeveloper/SettingUpAnUnrealWorkflow
 Description:Learn how to build Unreal Engine from source.
-Version: 4.11
-type:HowTo
+Version: 4.13
+type:QuickStart
 parent:Platforms/Linux/BeginnerLinuxDeveloper/SettingUpAnUnrealWorkflow
 order: 3
 tags:Linux
@@ -17,7 +17,7 @@ SkillLevel:Beginner
 	[PARAM:current]
 	[/PARAM]
 	[PARAM:home]
-		[Home](Platforms/Linux/BeginnerLinuxDeveloper/SettingUpAnUnrealWorkflow)
+		[Linux Quick Start Home](Platforms/Linux/BeginnerLinuxDeveloper/SettingUpAnUnrealWorkflow)
 	[/PARAM]
 	[PARAM:next]
 		[Next Step](Platforms/Linux/BeginnerLinuxDeveloper/SettingUpAnUnrealWorkflow/4 "%Platforms/Linux/BeginnerLinuxDeveloper/SettingUpAnUnrealWorkflow/4:title%")
diff --git a/Engine/Documentation/Source/Platforms/Linux/BeginnerLinuxDeveloper/SettingUpAnUnrealWorkflow/4/SettingUpAnUnrealWorkflow_4.INT.udn b/Engine/Documentation/Source/Platforms/Linux/BeginnerLinuxDeveloper/SettingUpAnUnrealWorkflow/4/SettingUpAnUnrealWorkflow_4.INT.udn
--- a/Engine/Documentation/Source/Platforms/Linux/BeginnerLinuxDeveloper/SettingUpAnUnrealWorkflow/4/SettingUpAnUnrealWorkflow_4.INT.udn
+++ b/Engine/Documentation/Source/Platforms/Linux/BeginnerLinuxDeveloper/SettingUpAnUnrealWorkflow/4/SettingUpAnUnrealWorkflow_4.INT.udn
@@ -2,8 +2,8 @@ Availability: Docs
 Title: 4 - Running UE4 on Linux
 Crumbs:%ROOT%, Platforms, Platforms/Linux, Platforms/Linux/BeginnerLinuxDeveloper, Platforms/Linux/BeginnerLinuxDeveloper/SettingUpAnUnrealWorkflow
 Description:Learn how to run Unreal Engine on Linux.
-Version: 4.11
-type:HowTo
+Version: 4.13
+type:QuickStart
 parent:Platforms/Linux/BeginnerLinuxDeveloper/SettingUpAnUnrealWorkflow
 order: 4
 tags:Linux
@@ -17,9 +17,10 @@ SkillLevel:Beginner
 	[PARAM:current]
 	[/PARAM]
 	[PARAM:home]
-		[Home](Platforms/Linux/BeginnerLinuxDeveloper)
+		[Linux Quick Start Home](Platforms/Linux/BeginnerLinuxDeveloper/SettingUpAnUnrealWorkflow)
 	[/PARAM]
 	[PARAM:next]
+		[Next Step](Platforms/Linux/BeginnerLinuxDeveloper/SettingUpAnUnrealWorkflow/5 "%Platforms/Linux/BeginnerLinuxDeveloper/SettingUpAnUnrealWorkflow/5:title%")
 	[/PARAM]
 [/OBJECT]
 [/VAR]
@@ -44,7 +45,6 @@ you learned:
 &#x2713; How to build Unreal Engine from source.  
 &#x2713; How to run Unreal Engine on your Linux machine.  
 
-Throughout the course of this tutorial, you spent most of your time using the terminal. In the next tutorial, you'll learn how to setup an integrated development 
-environment (IDE), enabling you to harness the context rich features of an IDE to help you build and debug UE4 games.
+Are you ready to do some exercises on your own?
 
 %Steps%
\ No newline at end of file
diff --git a/Engine/Documentation/Source/Platforms/Linux/BeginnerLinuxDeveloper/SettingUpAnUnrealWorkflow/5/SettingUpAnUnrealWorkflow_5.INT.udn b/Engine/Documentation/Source/Platforms/Linux/BeginnerLinuxDeveloper/SettingUpAnUnrealWorkflow/5/SettingUpAnUnrealWorkflow_5.INT.udn
new file mode 100644
--- /dev/null
+++ b/Engine/Documentation/Source/Platforms/Linux/BeginnerLinuxDeveloper/SettingUpAnUnrealWorkflow/5/SettingUpAnUnrealWorkflow_5.INT.udn
@@ -0,0 +1,50 @@
+Availability: Docs
+Title: 5 - On Your Own!
+Crumbs:%ROOT%, Platforms, Platforms/Linux, Platforms/Linux/BeginnerLinuxDeveloper, Platforms/Linux/BeginnerLinuxDeveloper/SettingUpAnUnrealWorkflow
+Description:Do some exercises on your own.
+Version: 4.13
+type:QuickStart
+parent:Platforms/Linux/BeginnerLinuxDeveloper/SettingUpAnUnrealWorkflow
+order: 5
+tags:Linux
+SkillLevel:Beginner
+
+[VAR:Steps]
+[OBJECT:Navigation]
+	[PARAM:previous]
+		[Previous Step](Platforms/Linux/BeginnerLinuxDeveloper/SettingUpAnUnrealWorkflow/4 "%Platforms/Linux/BeginnerLinuxDeveloper/SettingUpAnUnrealWorkflow/4:title%")
+	[/PARAM]
+	[PARAM:current]
+	[/PARAM]
+	[PARAM:home]
+		[Linux Quick Start Home](Platforms/Linux/BeginnerLinuxDeveloper/SettingUpAnUnrealWorkflow)
+	[/PARAM]
+	[PARAM:next]
+	[/PARAM]
+[/OBJECT]
+[/VAR]
+
+%Steps%
+
+One way of increasing productivity is to write custom shell scripts that you can run from the terminal. If you know how to write shell scripts in Linux, try writing one 
+that automates the process of setting up and building UE4.
+
+![](autobuild_onyourown.png)
+
+If you want to track the latest developments with UE4 on Linux, read [this article](https://www.unrealengine.com/blog/unreal-engine-public-issues-tracker-is-now-available) 
+and head over to the [Unreal Engine Public Issues Tracker](https://issues.unrealengine.com/) to check out the latest public issues impacting UE4 on Linux.
+
+![](UE4_publicJira.png)
+
+Do you see a public issue that you can fix? If you know how to submit [pull requests](https://help.github.com/articles/using-pull-requests/) with Git, join our growing  
+community of developers, and help shape the future of UE4 on Linux.
+
+![](gitHub_pullrequests.png)
+
+If you're looking to quickly get started making projects with UE4, check-out the following tutorials:
+
+*	[](Engine/Content/QuickStart)
+*	[](Programming/QuickStart)
+*	[](Engine/QuickStart)
+
+%Steps%
\ No newline at end of file
diff --git a/Engine/Documentation/Source/Platforms/Linux/BeginnerLinuxDeveloper/SettingUpAnUnrealWorkflow/SettingUpAnUnrealWorkflow.INT.udn b/Engine/Documentation/Source/Platforms/Linux/BeginnerLinuxDeveloper/SettingUpAnUnrealWorkflow/SettingUpAnUnrealWorkflow.INT.udn
--- a/Engine/Documentation/Source/Platforms/Linux/BeginnerLinuxDeveloper/SettingUpAnUnrealWorkflow/SettingUpAnUnrealWorkflow.INT.udn
+++ b/Engine/Documentation/Source/Platforms/Linux/BeginnerLinuxDeveloper/SettingUpAnUnrealWorkflow/SettingUpAnUnrealWorkflow.INT.udn
@@ -1,10 +1,10 @@
 Availability: Docs
-Title: How to Setup your Workflow
+Title: Linux Quick Start Guide
 Crumbs:%ROOT%, Platforms, Platforms/Linux, Platforms/Linux/BeginnerLinuxDeveloper
 Description:Learn how to download, build, and run UE4 on Linux.
 Platform:Linux
-Version: 4.11
-type:HowTo
+Version: 4.13
+type:QuickStart
 parent:Platforms/Linux/BeginnerLinuxDeveloper
 order: 2
 tags:Linux
@@ -12,6 +12,27 @@ SkillLevel:Beginner
 prereq:GettingStarted/RecommendedSpecifications
 prereq:Platforms/Linux/BeginnerLinuxDeveloper/SupportedGfxDrivers
 
+[VAR:Params]
+[PARAM:image]
+	![%Platforms/Linux/BeginnerLinuxDeveloper/SettingUpAnUnrealWorkflow:title%](linuxQS_topic.png)
+[/PARAM]
+[PARAM:icon]
+	![](%ROOT%/tutorial_icon.png)(convert:false)
+[/PARAM]
+[PARAM:title]
+	%Platforms/Linux/BeginnerLinuxDeveloper/SettingUpAnUnrealWorkflow:title%
+[/PARAM]
+[PARAM:path]
+	[RELATIVE:Platforms/Linux/BeginnerLinuxDeveloper/SettingUpAnUnrealWorkflow]
+[/PARAM]
+[/VAR]
+
+[VAR:Topic]
+[OBJECT:Topic]
+%params%
+[/OBJECT]
+[/VAR]
+
 ![](CompiledEditor_Linux.png)
 
 _At the end of this tutorial, you'll have Unreal Engine 4 (UE4) running on your Linux machine._
