Fixed document title.

#UE4doc

diff --git a/Engine/Documentation/Source/Platforms/Linux/Linux.INT.udn b/Engine/Documentation/Source/Platforms/Linux/Linux.INT.udn
--- a/Engine/Documentation/Source/Platforms/Linux/Linux.INT.udn
+++ b/Engine/Documentation/Source/Platforms/Linux/Linux.INT.udn
@@ -1,5 +1,5 @@
 Availability: Public
-Title: Linux Game Development[](Platforms/Linux/BeginnerLinuxDeveloper/SettingUpAnUnrealWorkflow/5)
+Title: Linux Game Development
 Crumbs:%ROOT%, Platforms
 Description:Developing games for the Linux platform.
 Platform:Linux
