Updating docs for entering VR editing mode to mention automatic entry UEDOC-3163

diff --git a/Engine/Documentation/Source/Engine/Editor/VR/ActivateVRMode/ActivateVRMode.INT.udn b/Engine/Documentation/Source/Engine/Editor/VR/ActivateVRMode/ActivateVRMode.INT.udn
--- a/Engine/Documentation/Source/Engine/Editor/VR/ActivateVRMode/ActivateVRMode.INT.udn
+++ b/Engine/Documentation/Source/Engine/Editor/VR/ActivateVRMode/ActivateVRMode.INT.udn
@@ -64,3 +64,12 @@ skilllevel: intermediate
 At any time you can press the Home button to return to this menu and then select another help section.
 
 [/EXCERPT:main]
+
+## Automatic VR Editing Mode
+
+As of 4.13, you can also enter VR Editing Mode automatically.  As long as the editor is in the foreground, when you wear the HMD, you will automatically enter VR editing mode. When you remove the HMD, you will leave it. 
+
+There is a setting under **VR** in the **Experimental** section of **Editor Preferences** that will allow you to turn off auto-entry if you prefer.
+
+![](AutoEntryMenu_Highlight.png)
+
