#UE4 docs Lighting Channels Update, includes information on it working in Mobile

diff --git a/Engine/Documentation/Source/Engine/Rendering/LightingAndShadows/LightingChannels/LightingChannels.INT.udn b/Engine/Documentation/Source/Engine/Rendering/LightingAndShadows/LightingChannels/LightingChannels.INT.udn
new file mode 100644
--- /dev/null
+++ b/Engine/Documentation/Source/Engine/Rendering/LightingAndShadows/LightingChannels/LightingChannels.INT.udn
@@ -0,0 +1,103 @@
+Availability:Docs
+Title: Lighting Channels
+Crumbs: %ROOT%, Engine
+Description:The Lighting Channels feature overview.
+Navigation:topic
+Version: 4.12.5
+
+[VAR:Topic]
+[OBJECT:Topic]
+	[PARAM:image]
+		![%Engine/Rendering/LightingAndShadows/LightingChannels:title%](Engine/Rendering/LightingAndShadows/lighting_topic.png)
+	[/PARAM]
+	[PARAM:icon]
+		![](%ROOT%/light_icon.png)(convert:false)
+	[/PARAM]
+	[PARAM:title]
+		%Engine/Rendering/LightingAndShadows/LightingChannels:title%
+	[/PARAM]
+	[PARAM:description]
+		%Engine/Rendering/LightingAndShadows/LightingChannels:description%
+	[/PARAM]
+	[PARAM:path]
+		[RELATIVE:Engine/Rendering/LightingAndShadows/LightingChannels]
+	[/PARAM]
+[/OBJECT]
+[/VAR]
+
+[VAR:TopicCompact]
+[OBJECT:TopicCompact]
+	[PARAM:image]
+		![%Engine/Rendering/LightingAndShadows/LightingChannels:title%](Engine/Rendering/LightingAndShadows/lighting_topic.png)
+	[/PARAM]
+	[PARAM:icon]
+		![](%ROOT%/light_icon.png)(convert:false)
+	[/PARAM]
+	[PARAM:title]
+		%Engine/Rendering/LightingAndShadows/LightingChannels:title%
+	[/PARAM]
+	[PARAM:description]
+		%Engine/Rendering/LightingAndShadows/LightingChannels:description%
+	[/PARAM]
+	[PARAM:path]
+		[RELATIVE:Engine/Rendering/LightingAndShadows/LightingChannels]
+	[/PARAM]
+[/OBJECT]
+[/VAR]
+
+![](image_39.png)
+
+**Lighting Channels** allow dynamic lights to affect only objects when their lighting channels overlap. This is primarily intended for cinematic use to give users powerful control over how **Actors** are lit. At the moment, Unreal Engine supports up to 3 lighting channels.
+
+## Usage
+
+By default **Directional Lights**, **Spot Lights**, **Point Lights**, and all **Actors** that can be affected by lights (**Static Meshes**, **Skeletal Meshes**, etc...) have **Lighting Channel 0** enabled. If you desire a lightable **Actor** to be affected by another **Lighting Channel** you must enable that channel on both the **Actor** and the light.
+
+### Example
+
+![](topImage.png)
+
+In the above image the **White Directional Light** can only affect **Channel 0**, which includes the mannequin on the left and the background **Static Meshes**. While the **Purple Point Light** can only affect **Channel 1**, which **only** includes the mannequin on the right.
+
+You can see how the properties are set here:
+
+[REGION:raw]
+![](LightingChannelProp.png)
+[/REGION]
+
+For lights, the property can be found in the **Advanced** pull down of the **Light** category in the **Details Panel**. For **Actors** that can be lit, you can find the **Lighting Channels** properties in the **Lighting** category of the **Details Panel**.
+
+## Limitations
+
+Lighting channel influence is applied dynamically. This means it won't work with **Static Lights** or **Static Mesh Actors with Static Mobility** (**Static Mesh Actors** with **Mobility** set to **Movable** work just fine). You'll need to use either **Stationary** or **Movable** lights.
+
+Lighting channels will only affect direct lighting on opaque materials. So, translucent or masked materials will not work.
+
+## Performance
+
+The performance impact of using **Lighting Channels** is minimal, but also non-zero. In a scene with 1 directional light, Radeon 7870, at 1080p for example:
+
+| Lighting Channels Satus | Time in ms |
+| --- | --- |
+| **off** | 0.42ms   StandardDeferredLighting 1 draws 1 prims 3 verts |
+| **on**  | [INCLUDE:#per1] |
+
+<!--
+[EXCERPT:per1]
+0.08ms   CopyStencilToLightingChannels 1 draws 1 prims 3 verts
+
+0.45ms   StandardDeferredLighting 1 draws 1 prims 3 verts
+[/EXCERPT]
+-->
+
+## On Mobile
+
+As of Unreal Engine 4 version 4.12 **Lighting Channels** works on mobile renderers.
+
+
+
+
+
+
+
+
