#UE4 Docs: Applying all peer feedback and making new images.

diff --git a/Engine/Documentation/Source/Platforms/Mobile/Lighting/MobileLighting.INT.udn b/Engine/Documentation/Source/Platforms/Mobile/Lighting/MobileLighting.INT.udn
--- a/Engine/Documentation/Source/Platforms/Mobile/Lighting/MobileLighting.INT.udn
+++ b/Engine/Documentation/Source/Platforms/Mobile/Lighting/MobileLighting.INT.udn
@@ -8,6 +8,7 @@ Related: Platforms/Mobile/Textures
 Related: Engine/Rendering/Materials/CustomizedUVs
 Related: Engine/Rendering/Materials/MaterialProperties/BlendModes
 Related: Engine/Rendering/Materials/ExpressionReference/Utility
+Related: Engine/Rendering/LightingAndShadows/ReflectionEnvironment
 version: 4.12
 parent:Platforms/Mobile
 order:3
@@ -23,8 +24,8 @@ When using lights on mobile platforms in UE4, certain caveats and restrictions m
 The features listed below are supported on mobile platforms:
 
 * Linear HDR lighting.
-* Directional lightmaps (normal is taken into account).
-* Distance field shadows + analytical specular on sun.
+* Directional Lightmaps (Normal is taken into account).
+* Distance field shadows + analytical specular on the sun.
 * Image based specular lighting - nearest [Reflection Capture](Engine/Rendering/LightingAndShadows/ReflectionEnvironment/) is applied to each object, without parallax correction. So there can be seams between objects.
 * Dynamic objects receive lighting correctly but do not cast shadows.
 
@@ -42,34 +43,70 @@ Mobile platforms only support a subset of the light types available in Unreal En
 ## Reflections
 While reflections can help to breathe life and believability into your mobile projects, there are a few caveats you need to be aware of when using them on mobile devices. Below you will find a list of some of the things you will need to do to ensure that reflections will work on your target mobile device.
 
-*  Disable any Ambient Cubemaps you have in your post process volumes by setting the **Intensity** to **0**.
+* Disable any Ambient Cubemaps you have in your post process volumes by setting the **Intensity** to **0**.
+	
+	![](PP_Settings.png)
+
 * Place [**SphereReflectionCaptures**](Engine/Rendering/LightingAndShadows/ReflectionEnvironment/#sphereshape) Actors in the areas that you want to capture lighting for reflections. 
+
+	[OBJECT:ComparisonSlider]
+	[PARAM:before]
+	![Sphere Reflection Captures Off ](Reflection_Probes_Off.png)(w:800 h:473)
+	[/PARAM]
+	[PARAM:after]
+	![Sphere Reflection Captures On](Reflection_Probes_On.png)(w:800 h:473)
+	[/PARAM]
+	[/OBJECT]
+
 * When making Materials the support reflections keep the following information in mind. 
-	* Interesting Normal map.
-		[REGION:lightbox]
-		[![](Interesting_Normal_Map.png)(w:128 h:128)](Interesting_Normal_Map.png)
-		[/REGION]
-	* Varying Roughness.
-		[REGION:lightbox]
-		[![](Interesting_Roughness_Map.png)(w:128 h:128)](Interesting_Roughness_Map.png)
-		[/REGION]
-	* Nothing plugged into the Specular input.
+
+	* Make sure your Normal maps has lots of high and low frequency detail in it as this helps to break up the reflection making it more visually interesting.  
+
+		[OBJECT:ComparisonSlider]
+		[PARAM:before]
+		![Flat Normal Map ](Interesting_Normal_Map_00.png)(w:800 h:473)
+		[/PARAM]
+		[PARAM:after]
+		![Rough Normal Map](Interesting_Normal_Map_01.png)(w:800 h:473)
+		[/PARAM]
+		[/OBJECT]
+
+	* Make sure your Roughness maps have varying intensities of whites and blacks as this will help to give you a varying degree of shiny and dull reflections.
+
+		[OBJECT:ComparisonSlider]
+		[PARAM:before]
+		![Varying Roughness ](Interesting_Roughness_Map_00.png)(w:624 h:723)
+		[/PARAM]
+		[PARAM:after]
+		![Flat Roughness](Interesting_Roughness_Map_01.png)(w:624 h:723)
+		[/PARAM]
+		[/OBJECT]
+
 	* Depending on the surface type the Metallic input should be a 1 or 0.
 
-Right now, each mesh component gets assigned to the closest [**Reflection Capture**](Engine/Rendering/LightingAndShadows/ReflectionEnvironment/).  That means there will be seams in reflections between objects, and sometimes a large mesh will use an undesirable reflection capture that happened to be close to the mesh's center. 
+		[OBJECT:ComparisonSlider]
+		[PARAM:before]
+		![Metallic Off](Metallic_Off.png)(w:341 h:600)
+		[/PARAM]
+		[PARAM:after]
+		![Metallic On](Metallic_On.png)(w:341 h:600)
+		[/PARAM]
+		[/OBJECT]
+
+Note that each Static Mesh component gets assigned to the closest [**Reflection Capture**](Engine/Rendering/LightingAndShadows/ReflectionEnvironment/). That means there will be seams in reflections between objects, and sometimes a large mesh will use an undesirable reflection capture that happened to be close to the mesh's center. 
 
 ## Modulated Shadowing
 
-Having fully Dynamic shadows helps to add life and realism to your Unreal Engine 4(UE4) mobile projects. 
-However on some mobile devices, fully Dynamic shadowing is not an option due to hardware restrictions and the high resource demands Dynamic shadows require to render. To get around this a new, cheaper type of Dynamic shadowing was introduced to UE4 called **Modulated Shadows**. For more information on how to setup Modulated Shadows check out the [**Modulated Shadows**](Platforms/Mobile/Lighting/HowTo/ModulatedShadows/) How To.
+Fully Dynamic shadows help to add life and realism to your mobile projects. 
+However, on some mobile devices, fully Dynamic shadowing is not an option due to hardware restrictions and the high resource demands Dynamic shadows require to render. To get around this a new, cheaper type of Dynamic shadowing was introduced to UE4 called **Modulated Shadows**. For more information on how to setup Modulated Shadows check out using [**Modulated Shadows**](Platforms/Mobile/Lighting/HowTo/ModulatedShadows/).
 
 ## Dynamic Cascaded Shadow Maps
-Stationary Directional lights can be used to cast Whole-scene Dynamic Cascade Shadow Maps (CSM) shadows for just the dynamic objects in your scene. One of the benefits to using CSM shadows is that CSM shadows will blend correctly with precomputed shadows cast by the static objects in the scene. CSM shadows also do not exhibit the double shadowing you see with Modulated shadows and this technique is also faster when casting dynamic shadows for multiple objects. For more information on how to setup Cascaded Shadows check out the [**Cascaded Shadows**](Platforms\Mobile\Lighting\HowTo\CascadedShadow\) How To.
+Stationary Directional lights can be used to cast Whole-scene Dynamic Cascade Shadow Maps (CSM) shadows for just the dynamic objects in your scene. One of the benefits to using CSM shadows is that they will blend correctly with precomputed shadows cast by the static objects in the scene. CSM shadows also do not exhibit the double shadowing you see with Modulated shadows and are faster when casting dynamic shadows for multiple objects. For more information on how to setup Cascaded Shadows, check out using [**Cascaded Shadows**](Platforms\Mobile\Lighting\HowTo\CascadedShadow\).
 
 ## Modulated Shadowing VS Dynamic Shadowing
 
-While Modulated and Dynamic shadowing look and act very similar to one another, they are actually quite different. 
-This is because Modulated shadows have a number of restrictions and hard limits that Dynamic shadows do not have that allows them to be more performant on mobile devices with limited hardware capabilities.
+While Modulated and Dynamic shadowing looks and acts very similar to one another, they are quite different. 
+This is because Modulated shadows have a number of restrictions and hard limits that Dynamic shadows do not have that allows them to beperform bettert on mobile devices with limited hardware capabilities.
 In the image below you can see a comparison between Dynamic shadows and Modulated shadows.
 
 [OBJECT:ComparisonSlider]
@@ -84,7 +121,7 @@ In the image below you can see a comparison between Dynamic shadows and Modulate
 
 The following list highlights some of the features and restrictions Modulated shadows can provide if used in your UE4 project.
 
-* **Changing Shadow Color:** You can change the color of the shadow that Modulated shadows cast by adjusting the **Modulated Shadow Color** option in the **Light** section of the Directional Light. 
+* **Changing Shadow Color:** You can change the color of the shadow that Modulated shadows cast by adjusting the **Modulated Shadow Color** option in the **Light** section of the Directional Light. This is useful if you want to adjust the color of your shadows for game play or artistic reasons.
 
 	[REGION:lightbox]
 	[![](Change_Mod_Shadow_Color.png)(w:771 h:252)](Change_Mod_Shadow_Color.png)
@@ -94,7 +131,7 @@ The following list highlights some of the features and restrictions Modulated sh
 	Click for full image.
 	[/REGION]
 
-* **Shadow Blending:** Unlike Dynamic shadows, Modulated shadows can not blend with other shadows, either baked or dynamic, that you see in your projects. This means that when a Modulated shadow is on top of another shadow, both shadows will show instead of a single shadow that blends together. 
+* **Shadow Blending:** Unlike Dynamic shadows, Modulated shadows can not blend with other shadows, either baked or dynamic, that you see in your project's. This means that when a Modulated shadow is on top of another shadow, both shadows will show instead of a single shadow that blends together. 
 
 	[OBJECT:ComparisonSlider]
 		[PARAM:before]
@@ -107,33 +144,20 @@ The following list highlights some of the features and restrictions Modulated sh
 
 ## Working with Modulated and Dynamic Shadows
 
-When working with Modulated and Dynamic shadows, there are a few console and .INI settings that can be used to adjust the look and performance of the shadows. 
-In the following section, we will go over what these settings are and how you can use them in your UE4 projects to ensure that your UE4 projects dynamic shadows are not costing too much to render.
+When working with Modulated and Dynamic shadows, there are a few console and .INI settings that can be used to adjust the look and performance of the shadows. In the following section, we will take a look at these settings and how they can be applied to your UE4 project.
 
 [region:note]
 As much as possible, Modulated shadows share code with the existing shadow tech. 
 This means many of the shadow cvars and .INI settings that can be used with other shadowing methods will also work with both Modulated and Dynamic shadows.
 [/region]
 
-* **Shadow Quality:** When you first view dynamic shadows on a mobile device, the sharpness and quality of the Modulated shadows could be a bit lower than expected. 
-To address this, you can adjust the shadow quality by opening up the Unreal console, using the backtick ` key, and enter the following command, **r.shadowquality** followed by a number. 
-The higher number you input, the better the Modulated shadow will look at the expense of FPS. 
-In the following image, r.shadowquality was set to values of **1**, **2**, **3** and **4** to show what effect this has on the Modulated shadow quality.
+* **Shadow Quality:** When you first view dynamic shadows on a mobile device, the sharpness and quality of the Modulated shadows could be a bit lower than expected. To address this, you can adjust the shadow quality by opening up the Unreal console, using the backtick ` key, and enter the following command, **r.shadowquality** followed by a number. The higher number you input, the better the Modulated shadow will look at the expense of FPS. In the following image, r.shadowquality was set to values of **0**, **1**, **2**, **3**, **4** and **5** to show what effect this has on the Modulated shadow quality.
 
-	[REGION:lightbox]
-	[![](Mod_Shadow_Quality.png)(w:800 h:192)](Mod_Shadow_Quality.png)
+	[REGION:fullwidth]
+	[INCLUDE:Platforms\Mobile\Lighting\#SQ]
 	[/REGION]
 
-	[REGION:caption]
-	Click for full image.
-	[/REGION]
-
-	[region:note]
-	Be extremely careful when setting the r.shadowquality to values over 3.
-	While testing this feature using a Samsung Galaxy Note 4 setting r.shadowquality to 4 halved the frame rate on the device from 60 FPS to 30 FPS.
-	[/region]
-
-* **Self Shadowing:** Modulated shadows will provide self-shadowing on dynamic objects, like characters or pickups, however by default this feature was disabled to ensure that Modulated shadows could be as performant as possible. 
+* **Self Shadowing:** Modulated shadows will provide self-shadowing on dynamic objects, like characters or pickups, however by default this feature was disabled to ensure that Modulated shadows could perform as well as possible. 
 If your project can afford the extra cost self-shadowing requires, you can enable it by inputting **r.Shadow.EnableModulatedSelfShadow 1** into the UE4 console.
 To disable self shadowing, input **r.Shadow.EnableModulatedSelfShadow 0** into the UE4 console.
 
@@ -146,21 +170,69 @@ To disable self shadowing, input **r.Shadow.EnableModulatedSelfShadow 0** into t
 		[/PARAM]
 	[/OBJECT]
 
-* **Shadow Depth:** Using the command **r.Shadow.CSMDepthBias** you can offset the position at which the shadow starts to be rendered. 
-In the image below you can see what happens to the Modulated shadows when r.Shadow.CSMDepthBias is left at the default value of **0** and then set to values of **10**, **100**, **500** and **1000**. 
-Adjusting the Shadow Depth is a great way to fix issues like shadow penetration were the shadow that is being cast is intersecting with the object that is supposed to be casting shadows.
-
-	[REGION:lightbox]
-	[![](Shadow_Depth_Settings.png)(w:1280 h:268)](Shadow_Depth_Settings.png)
-	[/REGION]
-
-	[REGION:caption]
-	Click for full image.
+* **Shadow Depth:** Using the command **r.Shadow.CSMDepthBias** you can offset the position at which the shadow starts to render. 
+In the image below you can see what happens to the Modulated shadows when r.Shadow.CSMDepthBias is left at the default value of **0** and then set to values of **10**, **100**, **500**, **1000** and **1500**. 
+	
+	[REGION:fullwidth]
+	[INCLUDE:Platforms\Mobile\Lighting\#SD]
 	[/REGION]
 
 
 
 
+<!-----
+[EXCERPT:SQ]
+[OBJECT:SliderSequence]
+	[PARAMLITERAL:max]
+	6
+	[/PARAMLITERAL]
+	[PARAMLITERAL:ext]
+	jpg
+	[/PARAMLITERAL]
+	[PARAMLITERAL:id]
+	SQ
+	[/PARAMLITERAL]
+	[PARAM:caption]
+	**Dragging the slider will show what happens when r.shadowquality is set from a value of 0 to a value of 5**
+	[/PARAM]
+	[PARAM:image_list]
+	![](SQ_1.png)
+	![](SQ_2.png)
+	![](SQ_3.png)
+	![](SQ_4.png)
+	![](SQ_5.png)
+	![](SQ_6.png)
+	[/PARAM]
+[/OBJECT]
+[/EXCERPT:SQ]
+-->
+
+<!-----
+[EXCERPT:SD]
+[OBJECT:SliderSequence]
+	[PARAMLITERAL:max]
+	6
+	[/PARAMLITERAL]
+	[PARAMLITERAL:ext]
+	jpg
+	[/PARAMLITERAL]
+	[PARAMLITERAL:id]
+	SD
+	[/PARAMLITERAL]
+	[PARAM:caption]
+	**Dragging the slider will show what happens when r.Shadow.CSMDepthBias is set from a value of 0 to a value of 1500**
+	[/PARAM]
+	[PARAM:image_list]
+	![](SD_1.png)
+	![](SD_2.png)
+	![](SD_3.png)
+	![](SD_4.png)
+	![](SD_5.png)
+	![](SD_6.png)
+	[/PARAM]
+[/OBJECT]
+[/EXCERPT:SD]
+-->
 
 
 
