#UE4 Docs: Applying peer feedback.

diff --git a/Engine/Documentation/Source/Platforms/Mobile/Performance/MobilePerformance.INT.udn b/Engine/Documentation/Source/Platforms/Mobile/Performance/MobilePerformance.INT.udn
--- a/Engine/Documentation/Source/Platforms/Mobile/Performance/MobilePerformance.INT.udn
+++ b/Engine/Documentation/Source/Platforms/Mobile/Performance/MobilePerformance.INT.udn
@@ -95,8 +95,8 @@ In this tier, you take advantage of most of the HDR lighting features available
 
 * Consider enabling bloom to take full advantage of the HDR lighting pipeline.
 * Realistic specular reflections in combination with HDR lighting can lead to specular aliasing. You should take advantage of the normal map to roughness feature to help reduce specular aliasing due to high frequency information in your normal map. You should also consider enabling anti-aliasing in the Post Process volume to help reduce this artifact.
-* Place **Reflection Captures** intelligently to get the best results. See [Reflections: Placing Reflection Captures](Resources/Showcases/Reflections/#placingreflectioncaptures) for more information.
-* Use only static lights in your scene.
+* Taking time to consider the placement of **Reflection Captures** will help you get the best resutls. For more information check out the [Reflections: Placing Reflection Captures](Resources/Showcases/Reflections/#placingreflectioncaptures).
+* Use only Static Lights and Lightmaps in your scene for best performance.
 
 
 ### Full HDR Lighting with per-pixel lighting from the Sun
@@ -171,23 +171,23 @@ Please note that the actual resolution for any scale factor other than 0.0 will
 The iOS scale factor is directly related to Apple's scale factor system.
 [/region]
 
-**r.MobileContentScaleFactor 0.0**: This will use the native resolution of the device.
-
-**r.MobileContentScaleFactor 1.0**: On a Retina device is the non-retina resolution.
-
-**r.MobileContentScaleFactor 2.0**: Full native resolution on iPhone 5S, iPad Air, etc.
-
-**r.MobileContentScaleFactor 3.0**: Full native resolution for iPhone 6+.
+|Command| Value|Result|
+|-------|------|------|
+|r.MobileContentScaleFactor|0.0|This will use the native resolution of the device.|
+|r.MobileContentScaleFactor|1.0|On a Retina device is the non-retina resolution.|
+|r.MobileContentScaleFactor|2.0|Full native resolution on iPhone 5S, iPad Air, etc.|
+|r.MobileContentScaleFactor|3.0|ull native resolution for iPhone 6+.|
 
 ### Mobile Content Scale Factor for Android
 When using the Mobile Content Scale Factor for Android devices inputting the following number after the r.MobileContentScaleFactor will produce the following results.
 Please note that the actual resolution for any scale factor other than 0.0 will be corrected to match the aspect ratio of the screen and clamped to native resolution.
 
-**r.MobileContentScaleFactor 0.0**: This will use the native resolution of the device.
-
-**r.MobileContentScaleFactor 1.0**: This will try and give you a resolution of 1280 x 720 for Landscape and 720 x 1280 for portrait.
+|Command| Value|Result|
+|-------|------|------|
+|r.MobileContentScaleFactor|0.0|This will use the native resolution of the device.|
+|r.MobileContentScaleFactor|1.0|This will try and give you a resolution of 1280 x 720 for Landscape and 720 x 1280 for portrait.|
+|r.MobileContentScaleFactor|2.0|This will try and give you a resolution of 2560 x 1440 for Landscape and 1440 x 2560 for portrait.|
 
-**r.MobileContentScaleFactor 2.0**: This will try and give you a resolution of 2560 x 1440 for Landscape and 1440 x 2560 for portrait.
 
 ## Mobile Material Quality Settings
 When building Content for an Unreal Engine 4 (UE4) project that will run on both low and high end mobile devices you will often run into issues with certain features working on one set of devices but not working on another set of devices. You also might encounter issues with your artwork looking correct on one device but completely different on another. While there are a lot of ways to address problems like this, many of theses ways are both time and resource intensive and also extremely error prone. To address these issues UE4 has the Material Quality Level system. This system allows you to build one single Material that can then be used on a wide range of different devices, giving you full control over which devices use which features. In the following sections we will go over these systems and how you can use them in your UE4 projects:
