Fixed some minor spelling issues.

diff --git a/Engine/Documentation/Source/Platforms/Mobile/Performance/MobilePerformance.INT.udn b/Engine/Documentation/Source/Platforms/Mobile/Performance/MobilePerformance.INT.udn
--- a/Engine/Documentation/Source/Platforms/Mobile/Performance/MobilePerformance.INT.udn
+++ b/Engine/Documentation/Source/Platforms/Mobile/Performance/MobilePerformance.INT.udn
@@ -95,7 +95,7 @@ In this tier, you take advantage of most of the HDR lighting features available
 
 * Consider enabling bloom to take full advantage of the HDR lighting pipeline.
 * Realistic specular reflections in combination with HDR lighting can lead to specular aliasing. You should take advantage of the normal map to roughness feature to help reduce specular aliasing due to high frequency information in your normal map. You should also consider enabling anti-aliasing in the Post Process volume to help reduce this artifact.
-* Taking time to consider the placement of **Reflection Captures** will help you get the best resutls. For more information check out the [Reflections: Placing Reflection Captures](Resources/Showcases/Reflections/#placingreflectioncaptures).
+* Taking time to consider the placement of **Reflection Captures** will help you get the best results. For more information check out the [Reflections: Placing Reflection Captures](Resources/Showcases/Reflections/#placingreflectioncaptures).
 * Use only Static Lights and Lightmaps in your scene for best performance.
 
 
@@ -266,7 +266,7 @@ Click for full image.
 [/REGION]
 
 ### Material Quality Level Switch
-The **Quality Switch** Material Expression allows you to define different levels of complexity within a single Material. For example, if you have a complex math operation or heavy texture reads that will work on higher end mobile devices but not on lower end mobile devices you can use this node to specify simplified versions to be displayed on lower end mobile devices that do not have the power to display higher end graphics.
+The **Quality Switch** Material Expression allows you to define different levels of complexity within a single Material. For example, if you have a complex math operation, or heavy texture reads that will work on higher end mobile devices but not on lower end mobile devices you can use this node to specify simplified versions to be displayed on lower end mobile devices that do not have the power to display higher end graphics.
 
 ![](image08.png)
 
@@ -283,7 +283,7 @@ Click for full image.
 If you open up any **Material** and click on the Main Material Node then look in the Details Panel under the **Mobile** section you will find the following two properties.
 
 * **Fully Rough**: Forces a Material to be completely rough saving a number of instructions and a single texture sampler.
-* **Use Lightmap Directionality**: Uses lightmap directionally and per pixel normals. If disabled the lightmaps will be flatter but cheaper.
+* **Use Lightmap Directionality**: Uses lightmap directionally and per pixel normals. If disabled, the lightmaps will be flatter but cheaper.
 
 [REGION:lightbox]
 [![](MQS_Mobile_Material_Options.png)(w:556 h:410)](MQS_Mobile_Material_Options.png)
@@ -292,11 +292,11 @@ If you open up any **Material** and click on the Main Material Node then look in
 Click for full image.
 [/REGION]
 
-These two properties can be enabled to help reduce the rendering cost of that Material when viewed on less powerful mobile devices by completely eliminating the rendering path for those features. However this is, an all or nothing option as you either have these options enabled or disabled for this Material on all devices.
+These two properties can be enabled to help reduce the rendering cost of that Material when viewed on less powerful mobile devices by completely eliminating the rendering path for those features. However, this is, an all or nothing option as you either have these options enabled or disabled for this Material on all devices.
 
 
 ### Mobile Material Rendering Overrides
-You can override the available rendering option a platform will use by adjusting the **ES2 Quality Settings**. You can find the ES2 Quality Setting for both Android and iOS by going to your **Project Settings** > **Platforms** and then selecting either **Android ES2 Quality** or **iOS ES2 Quality Settings**. To use the overrides you must first click on the **Enable Quality Override** option and then select the option you want to override by clicking on it. Once all of the options you want have been selected, press the **Update preview shaders** button to re-compile all Materials to work with the options you have specified. 
+You can override the available rendering option a platform will use by adjusting the **ES2 Quality Settings**. You can find the ES2 Quality Setting for both Android and iOS by going to your **Project Settings** > **Platforms** and then selecting either **Android ES2 Quality** or **iOS ES2 Quality Settings**. To use the overrides, you must first click on the **Enable Quality Override** option and then select the option you want to override by clicking on it. Once all of the options you want have been selected, press the **Update preview shaders** button to re-compile all Materials to work with the options you have specified. 
 
 
 [REGION:lightbox]
