#UE4docs
Capitalized NVIDIA per the NVIDIA Trademark Guide and the NVIDIA Style Guide Requirements.

diff --git a/Engine/Documentation/Source/Engine/Physics/Apex/Apex.INT.udn b/Engine/Documentation/Source/Engine/Physics/Apex/Apex.INT.udn
--- a/Engine/Documentation/Source/Engine/Physics/Apex/Apex.INT.udn
+++ b/Engine/Documentation/Source/Engine/Physics/Apex/Apex.INT.udn
@@ -1,7 +1,7 @@
 Availability: Public
 Title:APEX
 Crumbs: %ROOT%, Engine, Engine/Physics
-Description:Creation and Import of Destructibles using the Nvidia APEX Toolset.
+Description:Creation and Import of Destructibles using the NVIDIA APEX Toolset.
 Version: 4.9
 
 [VAR:Topic]
@@ -48,8 +48,8 @@ Version: 4.9
 [/REGION]
 
 
-Nvidia produces a tool called the APEX PhysX Lab which can be used to create destructible meshes. Information on its usage can be found in the 
-Nvidia documentation over the APEX PhysX Lab tools. However, it should be noted that currently Unreal Engine 4 only supports 
+NVIDIA produces a tool called the APEX PhysX Lab which can be used to create destructible meshes. Information on its usage can be found in the 
+NVIDIA documentation over the APEX PhysX Lab tools. However, it should be noted that currently Unreal Engine 4 only supports 
 APEX Destructibles and APEX Cloth. To create APEX Cloth assets it is required to use the PhysX: 3Ds Max/Maya DCC plug-in or the standalone tool is included in the APEX SDK. 
 
 Destructible Actors can be generated within Unreal Editor 4 using Voronoi with a single depth level.  APEX Destructible assets with more than one depth 
diff --git a/Engine/Documentation/Source/Engine/Physics/Apex/ApexTypes/ApexTypes.INT.udn b/Engine/Documentation/Source/Engine/Physics/Apex/ApexTypes/ApexTypes.INT.udn
--- a/Engine/Documentation/Source/Engine/Physics/Apex/ApexTypes/ApexTypes.INT.udn
+++ b/Engine/Documentation/Source/Engine/Physics/Apex/ApexTypes/ApexTypes.INT.udn
@@ -1,7 +1,7 @@
 Availability: Public
 Title:APEX Types
 Crumbs: %ROOT%, Engine, Engine/Physics, Engine/Physics/Apex
-Description:Nvidia APEX Physics Lab APX and APB files.
+Description:The NVIDIA APEX Physics Lab APX and APB files.
 Version: 4.9
 
 [VAR:Topic]
@@ -52,7 +52,7 @@ Version: 4.9
 		![](APEXtopicIcon.png)(w:150 convert:false)
 	[/PARAM]
 	[PARAM:description]
-	Nvidia PhysX APEX has two file types that Unreal Editor 4 can import, .APX and .APB. 
+	NVIDIA PhysX APEX has two file types that Unreal Editor 4 can import, .APX and .APB. 
 	Both cloth and destructible meshes can be exported as these file types, but only destructible meshes can be import via the **Content Browser.** Cloth assets will need to be imported via **Persona**.
 	[/PARAM]
 [/OBJECT]
diff --git a/Engine/Documentation/Source/GettingStarted/RecommendedSpecifications/Specifications.INT.udn b/Engine/Documentation/Source/GettingStarted/RecommendedSpecifications/Specifications.INT.udn
--- a/Engine/Documentation/Source/GettingStarted/RecommendedSpecifications/Specifications.INT.udn
+++ b/Engine/Documentation/Source/GettingStarted/RecommendedSpecifications/Specifications.INT.udn
@@ -80,7 +80,7 @@ Unreal Engine has some specific hardware and software requirements for running t
 | **Operating System** | Ubuntu 15.04  |
 | **Processor** | Quad-core Intel or AMD, 2.5 GHz or faster |
 | **Memory** | 16 GB RAM |
-| **Video Card** | 	NVIDIA GeForce 470 GTX or higher with latest NVidia binary drivers |
+| **Video Card** | 	NVIDIA GeForce 470 GTX or higher with latest NVIDIA binary drivers |
 [/EXCERPT:LinuxMinHardware]
 -->
 
@@ -280,7 +280,7 @@ Below are the typical systems used at Epic by programmers and artists. These sho
 * 24GB RAM
 * 300GB 10K SATA (OS)
 * 2X 300GB 10K SATA (RAID0 for DATA)
-* Nvidia Geforce 470 - 580
+* NVIDIA Geforce 470 - 580
 
 
 **Artist/LD:**
@@ -290,7 +290,7 @@ Below are the typical systems used at Epic by programmers and artists. These sho
 * 16GB RAM
 * 500GB SATA (OS)
 * 2X 500GB SATA (RAID0 for DATA)
-* Nvidia Geforce 580
+* NVIDIA Geforce 580
 
 
 [REGION:note]
diff --git a/Engine/Documentation/Source/Platforms/Android/Reference/AndroidGettingStarted.INT.udn b/Engine/Documentation/Source/Platforms/Android/Reference/AndroidGettingStarted.INT.udn
--- a/Engine/Documentation/Source/Platforms/Android/Reference/AndroidGettingStarted.INT.udn
+++ b/Engine/Documentation/Source/Platforms/Android/Reference/AndroidGettingStarted.INT.udn
@@ -85,7 +85,7 @@ There are several methods for getting the components. If you have already instal
 <!--
 [EXCERPT:TADPWindowsSetup]
 
-CodeWorks for Android 1R4 (by Nvidia) is a one-stop shop for Android development needs. If you install CodeWorks for Android 1R4, all of your requirements to develop for Android will be met.
+CodeWorks for Android 1R4 (by NVIDIA) is a one-stop shop for Android development needs. If you install CodeWorks for Android 1R4, all of your requirements to develop for Android will be met.
 
 The installer for CodeWorks for Android 1R4 is included with Unreal Engine 4 for Windows. It can be found here: 
 
@@ -108,7 +108,7 @@ It is not required to compile C++ code.
 [/EXCERPT:TADPWindowsSetup]
 
 [EXCERPT:TADPMacOSSetup]
-CodeWorks for Android 1R4, by Nvidia is a one-stop shop for Android development needs. If you install CodeWorks for Android 1R4, all of your requirements to develop for Android will be met.
+CodeWorks for Android 1R4, (by NVIDIA) is a one-stop shop for Android development needs. If you install CodeWorks for Android 1R4, all of your requirements to develop for Android will be met.
 
 The installer for CodeWorks for Android 1R4 is included with Unreal Engine 4 for Mac and can be found here: 
 
diff --git a/Engine/Documentation/Source/Shared/Tutorials/InDestructibleMeshEditorTutorial/InDestructibleMeshEditorTutorial.INT.udn b/Engine/Documentation/Source/Shared/Tutorials/InDestructibleMeshEditorTutorial/InDestructibleMeshEditorTutorial.INT.udn
--- a/Engine/Documentation/Source/Shared/Tutorials/InDestructibleMeshEditorTutorial/InDestructibleMeshEditorTutorial.INT.udn
+++ b/Engine/Documentation/Source/Shared/Tutorials/InDestructibleMeshEditorTutorial/InDestructibleMeshEditorTutorial.INT.udn
@@ -155,16 +155,16 @@ Of course Destructible Meshes have many more properties for tweaking how they fr
 
 [EXCERPT:D8]
 [VAR:StageTitle]
-Nvidia PhysX Lab
+NVIDIA PhysX Lab
 [/VAR]
 
-If you desire more than one fracture level, you will need to use Nvidia's APEX PhysX Lab tool.
+If you desire more than one fracture level, you will need to use the NVIDIA APEX PhysX Lab tool.
 
 ![](APEXtopicIcon.png)
 
 [This link will take you to the APEX developer's site.](https://developer.nvidia.com/apex-download-page)
 
-You will need to register with Nvidia to be able to download the APEX Physics Lab.
+You will need to register with NVIDIA to be able to download the APEX Physics Lab.
 
 
 [/EXCERPT:D8]
diff --git a/Engine/Documentation/Source/Support/Builds/ReleaseNotes/2013/August/ReleaseNotesAugust2013.INT.udn b/Engine/Documentation/Source/Support/Builds/ReleaseNotes/2013/August/ReleaseNotesAugust2013.INT.udn
--- a/Engine/Documentation/Source/Support/Builds/ReleaseNotes/2013/August/ReleaseNotesAugust2013.INT.udn
+++ b/Engine/Documentation/Source/Support/Builds/ReleaseNotes/2013/August/ReleaseNotesAugust2013.INT.udn
@@ -219,7 +219,7 @@ Template:ReleaseNotes.html
 	* Added `IPlatformFeaturesModule`, which can return a save game class (and will be used for other things) 
 	* Added a SaveGame automation test to the Game set
 * Android builds with NDK r9 and Clang 3.3
-* Update `PhysX-3.2_EPIC` source from NVIDIAs depot
+* Update `PhysX-3.2_EPIC` source from the NVIDIA depot
 
 
 ##Upgrade Notes
diff --git a/Engine/Documentation/Source/Support/Builds/ReleaseNotes/2013/Oct/ReleaseNotesOctober2013.INT.udn b/Engine/Documentation/Source/Support/Builds/ReleaseNotes/2013/Oct/ReleaseNotesOctober2013.INT.udn
--- a/Engine/Documentation/Source/Support/Builds/ReleaseNotes/2013/Oct/ReleaseNotesOctober2013.INT.udn
+++ b/Engine/Documentation/Source/Support/Builds/ReleaseNotes/2013/Oct/ReleaseNotesOctober2013.INT.udn
@@ -175,7 +175,7 @@ Availability:Licensee
 	* The only tested VR experience currently in main is the **ElementalVR** sample project.  This is the demo that we showed at GDC, and is tuned to hit 60 FPS. Recommended as a user's first VR experience !
 	* Notes for the best experience:
 		* If your card is running the demo at a solid 60FPS, it is suggested to **turn on vsync**.  This is especially true if you use the HD kit, as the screen orientation makes tearing very noticeable!
-		* **Turn off frame buffering**.  This option is available in the nVidia control panel.  Set this to 1 to decrease the latency of the device.
+		* **Turn off frame buffering**.  This option is available in the NVIDIA control panel.  Set this to 1 to decrease the latency of the device.
 
 			![](OculusSupportUE4_02.png)
 
diff --git a/Engine/Documentation/Source/Support/Builds/ReleaseNotes/2014/4_2/ReleaseNotes_4_2.INT.udn b/Engine/Documentation/Source/Support/Builds/ReleaseNotes/2014/4_2/ReleaseNotes_4_2.INT.udn
--- a/Engine/Documentation/Source/Support/Builds/ReleaseNotes/2014/4_2/ReleaseNotes_4_2.INT.udn
+++ b/Engine/Documentation/Source/Support/Builds/ReleaseNotes/2014/4_2/ReleaseNotes_4_2.INT.udn
@@ -559,7 +559,7 @@ In the 4.2 Preview Notes it was listed that the EndPlay event would replace the
 * New: Added simple Look at control.
 * Made foliage snap to *all* modified landscape components when sculpting, not just those in the current level.
 * Moved **SpringArmComponent** tick into **PostPhysics** to resolve some camera jitter
-* Upgraded NVTT (nVidia Texture Tools) to v2.0.8
+* Upgraded NVTT (NVIDIA Texture Tools) to v2.0.8
 * Updated PVCT exporter to now export images with dimensions greater than 4096. Source image will come from 1st image in the mip chain where the largest dimension is <= 4096.
 * Reorganized a lot of vehicle tuning data for official release.
 	* Changed units from cm to m.
@@ -699,8 +699,8 @@ In the 4.2 Preview Notes it was listed that the EndPlay event would replace the
 * Improved performance when mousing over the Editor on Mac OS X.
 * Enable translucent self-shadowing on ShaderModel 4 renderers for the benefit of Mac OS X particularly.
 * Improved support for Mac <= GL 3.3 drivers that don't support rendering to 3D or 2D array textures.
-	* Older GPUs such as the ATi 4xx0 series, Nvidia 9xx0/1xx/3xx and Intel HD 3000 should now render, though performance will be lower than newer GPUs. 
-* Fixed the sky light irradiance on Nvidia cards under OS X.
+	* Older GPUs such as the ATi 4xx0 series, NVIDIA 9xx0/1xx/3xx and Intel HD 3000 should now render, though performance will be lower than newer GPUs. 
+* Fixed the sky light irradiance on NVIDIA cards under OS X.
 * Fixed (via work around) an OpenGL driver regression on AMD FirePro cards introduced in Mac OS X 10.9.3 that caused incorrect rendering of skinned models.
 	* This workaround has a small but noticeable performance penalty and will be removed after the fix for the driver bug has been released.
 * Fixed various minor rendering errors on the new Mac Pro's AMD FirePro cards on Mac OS X.
diff --git a/Engine/Documentation/Source/Support/Builds/ReleaseNotes/2014/4_3/ReleaseNotes_4_3.INT.udn b/Engine/Documentation/Source/Support/Builds/ReleaseNotes/2014/4_3/ReleaseNotes_4_3.INT.udn
--- a/Engine/Documentation/Source/Support/Builds/ReleaseNotes/2014/4_3/ReleaseNotes_4_3.INT.udn
+++ b/Engine/Documentation/Source/Support/Builds/ReleaseNotes/2014/4_3/ReleaseNotes_4_3.INT.udn
@@ -774,7 +774,7 @@ Black Jack is made completely in Blueprint and is mobile ready!
 * New: Improved device compatibility across many devices. 
 * New: The **FOnlineAchievementsGooglePlay::QueryAchievements** function now works on Android.
 * New: Changed default Android depth buffer precision to 24 bit and added Android Platform Setting to choose preference.
-* New: Allow portrait-only orientations on Android tablets, but force landscape on devices with naturally landscape output (NVidia Shield, Google TV, Kindle Fire TV, Ouya).
+* New: Allow portrait-only orientations on Android tablets, but force landscape on devices with naturally landscape output (NVIDIA Shield, Google TV, Kindle Fire TV, Ouya).
 * New: Install batch file for Android APK generated by packaging now takes an optional device serial number to pass to ADB.
 * New: The engine now choses a better EGL config at startup, which fixes performance issues on newer Android devices.
 * New: Improved error handling in Android install batch file (with some output giving suggestions).
diff --git a/Engine/Documentation/Source/Support/Builds/ReleaseNotes/2014/4_6/ReleaseNotes_4_6.INT.udn b/Engine/Documentation/Source/Support/Builds/ReleaseNotes/2014/4_6/ReleaseNotes_4_6.INT.udn
--- a/Engine/Documentation/Source/Support/Builds/ReleaseNotes/2014/4_6/ReleaseNotes_4_6.INT.udn
+++ b/Engine/Documentation/Source/Support/Builds/ReleaseNotes/2014/4_6/ReleaseNotes_4_6.INT.udn
@@ -1020,7 +1020,7 @@ Slate and UMG now support font families and composite fonts.  These allow you to
 * Disable OpenGL workarounds no longer required on Yosemite, though they will still be used on Mavericks.
 	* Combined depth-stencil clears & depth-fetch-during-depth-test now work on AMD 7xx0 & Dxx0 GPUs.
 	* On Intel cards we no longer need to flush changes to the blend state, significantly improving performance on those GPUs.
-* Fix broken skeletal mesh rendering on Nvidia cards under Yosemite.
+* Fix broken skeletal mesh rendering on NVIDIA cards under Yosemite.
 	* This rebinds the SRV whenever the underlying buffer is modified which should not have a performance impact.
 * Fix crash reporting under Yosemite due to a bad interaction between our operator-new override & a system library.
 * Amend launching processes on OS X to report failures rather than crashing.
diff --git a/Engine/Documentation/Source/Support/Builds/ReleaseNotes/2014/January/ReleaseNotesJanuary2014.INT.udn b/Engine/Documentation/Source/Support/Builds/ReleaseNotes/2014/January/ReleaseNotesJanuary2014.INT.udn
--- a/Engine/Documentation/Source/Support/Builds/ReleaseNotes/2014/January/ReleaseNotesJanuary2014.INT.udn
+++ b/Engine/Documentation/Source/Support/Builds/ReleaseNotes/2014/January/ReleaseNotesJanuary2014.INT.udn
@@ -693,7 +693,7 @@ QA_APPROVED_UE4_BUILD_MINIMUM
 	* Made the default merge tool p4merge.app on Mac OS X. 
 	* Make the editor services 'Launch Game' use the default application to open (the project the same as double-clicking in Finder).
 	* Prevent blue flickering on AMD cards under Mac OS X when more than one UE4 instance is running.
-	* Worked around an Nvidia driver bug to avoid crashes loading some samples on Mac OS X.
+	* Worked around an NVIDIA driver bug to avoid crashes loading some samples on Mac OS X.
 * **Moblie**
 	* Added support for texture streaming on iOS.
 	* Cleaned up OpenGL TexStorage2D handling and added support for texture storage on iOS.
diff --git a/Engine/Documentation/Source/Support/Builds/ReleaseNotes/2015/4_8/ReleaseNotes_4_8.INT.udn b/Engine/Documentation/Source/Support/Builds/ReleaseNotes/2015/4_8/ReleaseNotes_4_8.INT.udn
--- a/Engine/Documentation/Source/Support/Builds/ReleaseNotes/2015/4_8/ReleaseNotes_4_8.INT.udn
+++ b/Engine/Documentation/Source/Support/Builds/ReleaseNotes/2015/4_8/ReleaseNotes_4_8.INT.udn
@@ -3927,7 +3927,7 @@ Cooked builds can use a dedicated thread for package streaming instead of using
 
 * New: Makefile generator will now generate targets for projects outside of the engine directory.
 
-* New: Minor: a lot of performance warnings issued by NVidia drivers have been fixed.
+* New: Minor: a lot of performance warnings issued by NVIDIA drivers have been fixed.
 
 * New: Minor: Steam shared objects are now being loaded dynamically.
 
diff --git a/Engine/Documentation/Source/Support/Builds/ReleaseNotes/2015/4_9/ReleaseNotes_4_9.INT.udn b/Engine/Documentation/Source/Support/Builds/ReleaseNotes/2015/4_9/ReleaseNotes_4_9.INT.udn
--- a/Engine/Documentation/Source/Support/Builds/ReleaseNotes/2015/4_9/ReleaseNotes_4_9.INT.udn
+++ b/Engine/Documentation/Source/Support/Builds/ReleaseNotes/2015/4_9/ReleaseNotes_4_9.INT.udn
@@ -2924,7 +2924,7 @@ Other improvements include:
 
 * Changed the way the CrashReportClient is launched on Mac OS X to avoid a potential deadlock.
 
-* Disable OpenGL GPU timing on Nvidia GPUs and entirely on OS X 10.9 (Mavericks) to avoid crashes in the OpenGL drivers.
+* Disable OpenGL GPU timing on NVIDIA GPUs and entirely on OS X 10.9 (Mavericks) to avoid crashes in the OpenGL drivers.
 
 * Fix a crash when launching an external process on Mac OS X caused by over-releasing the path to the executable.
 
diff --git a/Engine/Documentation/Source/Support/Builds/ReleaseNotes/2016/4_11/ReleaseNotes_4_11.INT.udn b/Engine/Documentation/Source/Support/Builds/ReleaseNotes/2016/4_11/ReleaseNotes_4_11.INT.udn
--- a/Engine/Documentation/Source/Support/Builds/ReleaseNotes/2016/4_11/ReleaseNotes_4_11.INT.udn
+++ b/Engine/Documentation/Source/Support/Builds/ReleaseNotes/2016/4_11/ReleaseNotes_4_11.INT.udn
@@ -294,7 +294,7 @@ Metal is now the default graphics API on Mac OS X El Capitan!
 ![image alt text](image_23.png)
 [/REGION]
 
-Epic have worked closely with Apple, AMD, Nvidia, and Intel to integrate Metal for Mac and in 4.11 it replaces OpenGL as the primary graphics API for OS X El Capitan. The 4.11 release provides the same rendering features across Metal and OpenGL by default. Metal provides a streamlined, low-overhead API with precompiled shaders and efficient multi-threading support to maximise the processing power of the GPU. We will continue to improve and extend support for Mac Metal and look for ways to leverage new API features in upcoming versions of the engine.
+Epic have worked closely with Apple, AMD, NVIDIA, and Intel to integrate Metal for Mac and in 4.11 it replaces OpenGL as the primary graphics API for OS X El Capitan. The 4.11 release provides the same rendering features across Metal and OpenGL by default. Metal provides a streamlined, low-overhead API with precompiled shaders and efficient multi-threading support to maximise the processing power of the GPU. We will continue to improve and extend support for Mac Metal and look for ways to leverage new API features in upcoming versions of the engine.
 
 There is also experimental Metal support for Shader Model 5 features, try it out using the "-metalsm5" command-line switch.
 
@@ -3358,7 +3358,7 @@ Stats are automatically broken into the threads they execute on and can appear i
 
 * New: Added an experimental option to enable OpenCL acceleration for the convex decomposition tool on Mac which defaults to off.
 
-    * This avoids an OpenCL crash on Nvidia's OpenCL drivers but still allows users of AMD & Intel GPUs to enable it if they wish.
+    * This avoids an OpenCL crash on the NVIDIA OpenCL drivers but still allows users of AMD & Intel GPUs to enable it if they wish.
 
 * Bugfix: Change the ShaderCache to report when resource references are invalid and prevent it from caching states that use invalid references to avoid errors on subsequent runs.
 
@@ -5388,7 +5388,7 @@ Stats are automatically broken into the threads they execute on and can appear i
 
 * Added implementation of FPlatformProcess::UserTempDir()
 
-* Bind the index buffer as vertex attribute 0 when performing an attribute-less draw call to work around incorrect gl_VertexID values on Nvidia drivers rather than using different index buffers.
+* Bind the index buffer as vertex attribute 0 when performing an attribute-less draw call to work around incorrect gl_VertexID values on NVIDIA drivers rather than using different index buffers.
 
 * Detect when UE4 Mac is running under apitrace and fallback to using NSView rather than CoreAnimation to display so that tracing OpenGL works.
 
diff --git a/Engine/Documentation/Source/Support/Builds/ReleaseNotes/2016/4_12/ReleaseNotes_4_12.INT.udn b/Engine/Documentation/Source/Support/Builds/ReleaseNotes/2016/4_12/ReleaseNotes_4_12.INT.udn
--- a/Engine/Documentation/Source/Support/Builds/ReleaseNotes/2016/4_12/ReleaseNotes_4_12.INT.udn
+++ b/Engine/Documentation/Source/Support/Builds/ReleaseNotes/2016/4_12/ReleaseNotes_4_12.INT.udn
@@ -1981,7 +1981,7 @@ Collision can now be generated for Spline meshes created at runtime.
 
 * New: Added an explicit, experimental option for using OpenCL to accelerate the convex decomposition tool on Mac rather than assuming that it is desirable (defaults to off).
 
-    * This may prevent a crash on Nvidia's OpenCL drivers while enabling AMD & Intel users to run it if desired.
+    * This may prevent a crash on the NVIDIA OpenCL drivers while enabling AMD & Intel users to run it if desired.
 
 * New: Added implementations of unpackHalf2x16, countbits, firstbitlow & firstbithigh for Metal.
 
@@ -2089,13 +2089,13 @@ Collision can now be generated for Spline meshes created at runtime.
 
 * Slightly improved Mac Metal performance on AMD by only creating texture-views when absolutely necessary.
 
-* Switch uniform buffers to managed memory on Mac as this is more appropriate for AMD & Nvidia GPUs.
+* Switch uniform buffers to managed memory on Mac as this is more appropriate for AMD & NVIDIA GPUs.
 
 * Used a dispatch_semaphore, not an FEvent, for Metal free-list synchronisation as the dispatch_worker threads can't be properly set up for FStats.
 
 * We now warn the user when trying to run on OpenGL 3.2 GPUs that they are running unsupported and there may be rendering errors and performance problems.
 
-* Added a workaround for a Metal shader compiler problem on Nvidia GPUs - sample_compare doesn't like taking level(0) as its LOD-option argument and may not currently support any LOD-option.
+* Added a workaround for a Metal shader compiler problem on NVIDIA GPUs - sample_compare doesn't like taking level(0) as its LOD-option argument and may not currently support any LOD-option.
 
 #### PlayStation 4
 
