Addressed peer review comments made by IanS and an editorial review comment made by JeffW.

#UE4doc

diff --git a/Engine/Documentation/Source/Platforms/Linux/BeginnerLinuxDeveloper/SettingUpAnUnrealWorkflow/3/SettingUpAnUnrealWorkflow_3.INT.udn b/Engine/Documentation/Source/Platforms/Linux/BeginnerLinuxDeveloper/SettingUpAnUnrealWorkflow/3/SettingUpAnUnrealWorkflow_3.INT.udn
--- a/Engine/Documentation/Source/Platforms/Linux/BeginnerLinuxDeveloper/SettingUpAnUnrealWorkflow/3/SettingUpAnUnrealWorkflow_3.INT.udn
+++ b/Engine/Documentation/Source/Platforms/Linux/BeginnerLinuxDeveloper/SettingUpAnUnrealWorkflow/3/SettingUpAnUnrealWorkflow_3.INT.udn
@@ -37,8 +37,6 @@ on your Linux machine.
 
 ![](built_engine.png)
 
-_Some lower-end Linux systems can take over an hour to build the engine._
-
-At this point, you should have a fully built binary of UE4 on your hard disk. Alas, the time has come to run the engine on your Linux system!
+At this point, you should have a fully built binary of UE4 on your hard disk. The time has come to run the engine on your Linux system!
 
 %Steps%
\ No newline at end of file
diff --git a/Engine/Documentation/Source/Platforms/Linux/BeginnerLinuxDeveloper/SettingUpAnUnrealWorkflow/SettingUpAnUnrealWorkflow.INT.udn b/Engine/Documentation/Source/Platforms/Linux/BeginnerLinuxDeveloper/SettingUpAnUnrealWorkflow/SettingUpAnUnrealWorkflow.INT.udn
--- a/Engine/Documentation/Source/Platforms/Linux/BeginnerLinuxDeveloper/SettingUpAnUnrealWorkflow/SettingUpAnUnrealWorkflow.INT.udn
+++ b/Engine/Documentation/Source/Platforms/Linux/BeginnerLinuxDeveloper/SettingUpAnUnrealWorkflow/SettingUpAnUnrealWorkflow.INT.udn
@@ -33,9 +33,10 @@ prereq:Platforms/Linux/BeginnerLinuxDeveloper/SupportedGfxDrivers
 [/OBJECT]
 [/VAR]
 
-![](CompiledEditor_Linux.png)
-
-_At the end of this tutorial, you'll have Unreal Engine 4 (UE4) running on your Linux machine._
+[REGION:imagetable]
+![](CompiledEditor_Linux.png)(w:920)
+At the end of this tutorial, you'll have Unreal Engine 4 (UE4) running on your Linux machine.
+[/REGION]
 
 ## Goal
 
