Minor grammar fix.

#UE4doc

diff --git a/Engine/Documentation/Source/Platforms/Linux/BeginnerLinuxDeveloper/SettingUpAnUnrealWorkflow/5/SettingUpAnUnrealWorkflow_5.INT.udn b/Engine/Documentation/Source/Platforms/Linux/BeginnerLinuxDeveloper/SettingUpAnUnrealWorkflow/5/SettingUpAnUnrealWorkflow_5.INT.udn
--- a/Engine/Documentation/Source/Platforms/Linux/BeginnerLinuxDeveloper/SettingUpAnUnrealWorkflow/5/SettingUpAnUnrealWorkflow_5.INT.udn
+++ b/Engine/Documentation/Source/Platforms/Linux/BeginnerLinuxDeveloper/SettingUpAnUnrealWorkflow/5/SettingUpAnUnrealWorkflow_5.INT.udn
@@ -37,7 +37,7 @@ and head over to the [Unreal Engine Public Issues Tracker](https://issues.unreal
 ![](UE4_publicJira.png)
 
 Do you see a public issue that you can fix? If you know how to submit [pull requests](https://help.github.com/articles/using-pull-requests/) with Git, join our growing  
-community of developers, and help shape the future of UE4 on Linux.
+community of developers and help shape the future of UE4 on Linux.
 
 ![](gitHub_pullrequests.png)
 
