#ue4 docs

Applying Peer Feedback.

diff --git a/Engine/Documentation/Source/Engine/UMG/UserGuide/WidgetInteraction/WidgetInteraction.INT.udn b/Engine/Documentation/Source/Engine/UMG/UserGuide/WidgetInteraction/WidgetInteraction.INT.udn
--- a/Engine/Documentation/Source/Engine/UMG/UserGuide/WidgetInteraction/WidgetInteraction.INT.udn
+++ b/Engine/Documentation/Source/Engine/UMG/UserGuide/WidgetInteraction/WidgetInteraction.INT.udn
@@ -18,7 +18,7 @@ This feature is considered experimental and some aspects may not work as expecte
 If you are using a [Widget Component](Engine/Components/Widget) to display UI that exists in 3D in your game world and you want to allow players to interact with that widget, the **Widget Interaction** component provides a way for that interaction to occur. 
 
 The Widget Interaction component performs a [Raycast](Gameplay/HowTo/UseRaycasts/) to see if it hits a Widget Component placed in the world and if it hits one, you can set up rules to determine how to interact with it. 
-The interaction is performed by simulating a defined key press, for example if a Button can be clicked on with the Left Mouse Button, you could tell other forms of input to simulate a Left Mouse Button click (controller button press, motion controller trigger press, etc.)
+The interaction is performed by simulating a defined key press, for example if a Button can be clicked on with the Left Mouse Button, you could tell other forms of input to simulate a Left Mouse Button click (controller button press, motion controller trigger press, etc.).
 
 [REGION:note]
 For step-by-step examples, see [](Engine\UMG\HowTo\InWorldWidgetInteraction\) and [](Engine\UMG\HowTo\VirtualKeyboards). 
@@ -30,22 +30,22 @@ Typically you will want to add the Widget Interaction component to your player [
 
 ![](WidgetInteractionComponent.png)
 
-Above we add and attach it to a gun our character is carrying so that where the gun points, the Widget Interaction component points. 
+Above, we add and attach it to a gun our character is carrying so that where the gun points, the Widget Interaction component points. 
 
 
 ## Widget Interaction Property Reference
 
-Once you have added a Widget Interaction component, you can adjust several properties forom the **Details** panel to define how it functions. 
-Aside from common component properties like **Transform**, **Rendering** or **Sockets**, the properties below are specifically used with the Widget Interaction Component. 
+Once you have added a Widget Interaction component, you can adjust several properties from the **Details** panel to define how it functions. 
+Aside from common component properties like **Transform**, **Rendering**, or **Sockets**, the properties below are specifically used with the Widget Interaction Component. 
 
 | Option | Description |
 | --- | ----------- |
 | [REGION:tablesection]**Interaction**[/REGION] ||
-| **Virtual User Index** | Represents the Virtual User Index. Widget Interaction Components operate by using this Virtual User Index which captures and handles focus states separately. Each Virtual User should be represented by a different Index to ensure that they will maintain separate capture and focus states for each. As each Widget Interaction component comes online it informs Slate of the Virtual User Index it has been assigned and can send events as if it were a real Slate user. |
-| **Pointer Index** | Each user Virtual Controller or virtual Finger Tips being simulated should use a different Pointer Index.|
+| **Virtual User Index** | Represents the Virtual User Index. Widget Interaction Components operate by using this Virtual User Index which captures and handles focus states separately. Each Virtual User should be represented by a different Index to ensure that they will maintain separate capture and focus states. As each Widget Interaction component comes online, it informs Slate of the Virtual User Index it has been assigned and can send events as if it were a real Slate user. |
+| **Pointer Index** | Each user's simulated Virtual Controller or Virtual Finger Tips should use a different Pointer Index.|
 | **Interaction Distance** | The distance in game units the component should be able to interact with a widget component. |
 | **Interaction Source** | Determines where we should project from and start the trace (World, Mouse, Center Screen or Custom). If you set this to Custom, you'll need to call `SetCustomHitResult()` and provide the result of a custom hit test from whatever location you wish. |
-| **Enable Hit Testing** | Should the interaction component perform hit testing (Automatic or Custom) and attempt to simulate hover - if you were going to emulate a keyboard you would want to turn this option off if the virtual keyboard was separate from the virtual pointer device and used a secon interaction component.|
+| **Enable Hit Testing** | Should the interaction component perform hit testing (Automatic or Custom) and attempt to simulate hover? If you were going to emulate a keyboard and the virtual keyboard is separate from the virtual pointer device, you should turn this option off and use a second interaction component for the pointer device.|
 | [REGION:tablesection]**Debugging**[/REGION] ||
 | **Show Debug** | Shows debugging lines and a hit sphere to help debug interaction. |
 | **Debug Color** | Determines the color of the debug lines when **Show Debug** is enabled. |
@@ -64,19 +64,18 @@ You can also get information such as the last "hit" position registered by the W
 | --- | ----------- |
 | **Get Hovered Widget Component** | Gets the currently hovered Widget Component. |
 | **Get Last Hit Result** | Gets the last hit result generated by the component. Returns the custom hit result if that was set. |
-| **Is Over Focusable Widget** | Returns true if a widget under the hit result is focusable. For example, Slate widgets that return true for `SupportsKeyboardFocus()`.|
-| **Is Over Hit Test Visible Widget** | Returns true if a widget under the hit result has a visibility setting that makes it hit test visible. For example, Slate widgets that return true for `GetVisibility().IsHitTestVisible()`.|
-| **Is Over Interactable Widget** | Returns true if a widget under the hit result is interactive. For example, Slate widgets that return true for `IsInteractble()`.|
+| **Is Over Focusable Widget** | Returns true if a widget under the hit result is focusable. Slate widgets that return true for `SupportsKeyboardFocus()`. |
+| **Is Over Hit Test Visible Widget** | Returns true if a widget under the hit result has a visibility setting that makes it hit test visible. Slate widgets that return true for `GetVisibility().IsHitTestVisible()`. |
+| **Is Over Interactable Widget** | Returns true if a widget under the hit result is interactive. Slate widgets that return true for `IsInteractble()`. |
 | **Press and Release Key** | Performs both the press and release of a simulated keyboard key.| 
-| **Press Key** | Performs a press of a simulated keyboard key. Avoid using this for `a-z|A-Z`, things like the Editable Textbox in Slate expect **OnKeyChar** to be called to signal a specific character being sent to the widget so for those cases you would use **SendKeyChar** instead. |
-| **Press Pointer Key** | Performs a press of a simulated mouse/pointer event (pointer could be mouse, could be touch, etc). Normally you would just use Left/Right Mouse Button for the Key, however advanced usage could be imagined where you send other keys to signal widgets to take special actions if they're under the cursor.  |
-| **Release Key** | Releases a key as if it had been released from the keyboard.  |
-| **Release Pointer Key** | Releases a key as if the mouse/pointer where the source of it. Similar to Press Pointer Key, normally you would use Left/Right Mouse Button for the key however advanced use cases are possible. |
+| **Press Key** | Performs a press of a simulated keyboard key. Avoid using this for `a-z|A-Z`, since things like the Editable Textbox in Slate expect **OnKeyChar** to be called to signal a specific character being sent to the widget. For those cases, use **SendKeyChar** instead. |
+| **Press Pointer Key** | Simulates pressing a key on the mouse/pointer. In most cases, the key supplied would be a mouse button, but any key can be used. Advanced use cases, such as sending other keys to signal widgets to take special actions, are supported. |
+| **Release Key** | Releases a key as if it had been released from the keyboard. |
+| **Release Pointer Key** | Simulates releasing a key on the mouse/pointer, similar to Press Pointer Key. In most cases, the key supplied would be a mouse button, but any key can be used. |
 | **Scroll Wheel** | Sends a scroll wheel event to the widget under the last hit result. |
-| **Send Key Char** | Transmits a list of characters to a widget by simulating a OnKeyChar event for each key listed in the string. |
-| **Set Custom Hit Result** | Sets a custom hit result. This is only take into account if the Interaction Source is set to Custom.  |
+| **Send Key Char** | Transmits a list of characters to a widget by simulating an OnKeyChar event for each key listed in the string. |
+| **Set Custom Hit Result** | Sets a custom hit result. This is only taken into account if the Interaction Source is set to Custom.  |
 
 The properties for the Widget Interaction component can also be reterived in the form of a getter or set via setter nodes. 
 
 ![](WidgetInteractionVariables.png)
-
