Quick Menu Updates - UEDOC-3162

diff --git a/Engine/Documentation/Source/Engine/Editor/VR/QuickSelect/QuickSelectMenu.INT.udn b/Engine/Documentation/Source/Engine/Editor/VR/QuickSelect/QuickSelectMenu.INT.udn
--- a/Engine/Documentation/Source/Engine/Editor/VR/QuickSelect/QuickSelectMenu.INT.udn
+++ b/Engine/Documentation/Source/Engine/Editor/VR/QuickSelect/QuickSelectMenu.INT.udn
@@ -6,7 +6,7 @@ parent: Engine/Editor/VR
 type: reference
 order: 2
 tags: VR
-version: 4.11
+version: 4.13
 skilllevel: intermediate
 
 [VAR:Topic]
@@ -43,13 +43,13 @@ skilllevel: intermediate
 [/OBJECT]
 [/VAR]
 
-The **Quick Select** menu nables you to quickly adjust grid, rotation and snap settings as well as summon the Content Browser, Details panel, and more.
+The **Quick Select** menu enables you to quickly adjust grid, rotation and snap settings as well as summon the Content Browser, Details panel, and more.
 
-![](quickselectmenu.png)(w:480)
+![](QuickMenuVR.png)(w:350)
 
 ## Using the Quick Select Menu
 
-1. Aim on of the controllers at the other controller to display the Quick Select menu.
+1. Aim one of the controllers at the other controller to display the Quick Select menu.
 1. Aim the free controller at the options on the menu to highlight them.
 1. Squeeze the Trigger to select the highlighted option.
 
@@ -61,13 +61,19 @@ The **Quick Select** menu nables you to quickly adjust grid, rotation and snap s
 | Move | Toggles grid snapping on or off when moving Actors. The snapping value is modified by aiming at the numerical box below, squeezing the trigger, and moving the controller. |
 | Rotate | Toggles grid snapping on or off when rotating Actors. The snapping value is modified by aiming at the numerical box below, squeezing the trigger, and moving the controller. |
 | Scale | Toggles grid snapping on or off when scaling Actors. The snapping value is modified by aiming at the numerical box below, squeezing the trigger, and moving the controller. |
-| [REGION:tablesection]Gizmo[/REGION] ||
-| Coordinate System | Sets the coordinate system [World Space/Local Space] to use when transforming Actors. |
-| Transform Selector | Cycles the Universal Gizmo through the available transform modes: Translate, Rotate, Scale. Only the manipulators for the selected mode are displayed. |
 | [REGION:tablesection]Interface[/REGION] ||
 | Content | Opens the Content Browser |
 | Details | Opens the Details panel for viewing and modifying properties of the selected Actor(s). |
 | Modes | Opens the Modes panel providing access to Actor placement tools and more. |
 | Outliner | Opens theWorld Outliner which displays a hierarchical list of all Actors in the world enabling quick selection, visibility toggling, organization, and Actor-specific actions to be performed. |
 | Asset | Opens the Asset window which acts as a container for any asset editors - Blueprint Editor, Material Editor, etc. - that may be open. |
-| Tutorial | Opens the VR Mode Tutorial which provides access to a collection of instructional walkthroughs for common features of the VR Editor. |
\ No newline at end of file
+| Tutorial | Opens the VR Mode Tutorial which provides access to a collection of instructional walkthroughs for common features of the VR Editor. |
+| Light | Adds a flashlight to the controller used to click on this button.  |
+| Screenshot | Takes a screenshot of the current scene. |
+| World | Opens the World Settings panel providing access to properties of the current Level. |
+| [REGION:tablesection]Gizmo[/REGION] ||
+| Coordinate System | Sets the coordinate system [World Space/Local Space] to use when transforming Actors. |
+| Transform Selector | Cycles the Universal Gizmo through the available transform modes: Translate, Rotate, Scale. Only the manipulators for the selected mode are displayed. |
+| [REGION:tablesection] [/REGION] ||
+| Play | Enables Play In Editor mode.  |
+
