Transform Gizmo Updates - UEDOC-3161

diff --git a/Engine/Documentation/Source/Engine/Editor/VR/UniversalGizmo/UniversalGizmo.INT.udn b/Engine/Documentation/Source/Engine/Editor/VR/UniversalGizmo/UniversalGizmo.INT.udn
--- a/Engine/Documentation/Source/Engine/Editor/VR/UniversalGizmo/UniversalGizmo.INT.udn
+++ b/Engine/Documentation/Source/Engine/Editor/VR/UniversalGizmo/UniversalGizmo.INT.udn
@@ -45,15 +45,16 @@ skilllevel: intermediate
 
 The **Universal Gizmo** is a transform widget designed for moving, rotating, and scaling Actors when working in VR mode. 
 
-![](universalgizmo.png)(w:400)
+![](TransformGizmo.png)(w:400)
 
-The elements of the Universal Gizmo are colored Red, Green, and Blue representing the X, Y, and Z axes, respectively.
+
+The elements of the Universal Gizmo are colored red, green, and blue representing the X, Y, and Z axes, respectively. 
 
 | Action | Manipulator | Usage |
 | --- | --- | --- |
-| **Translate** | ![](universalgizmo_move.png)(w:192) | Aim the controller at one of the translation manipulators and squeeze the **Trigger** halfway to grab it. Move the controller to move the selected Actor(s) along the axis. |
-| **Translate Plane** | ![](universalgizmo_moveplane.png)(w:192) | Aim the controller at one of the translation plane manipulators and squeeze the **Trigger** halfway to grab it. Move the controller to move the selected Actor(s) along the plane (in 2 axes). |
-| **Rotate** | ![](universalgizmo_rotate.png)(w:192) | Aim the controller at one of the rotation manipulators and squeeze the **Trigger** halfway to grab it. Move the controller to rotate the selected Actor(s) around the axis. |
-| **Scale** | ![](universalgizmo_scale.png)(w:192) | Aim the controller at one of the scale manipulators and squeeze the **Trigger** halfway to grab it. Move the controller to scale the selected Actor(s) along the Y-axis. |
-| **Uniform Scale** | ![](universalgizmo_uniformscale.png)(w:192) | Aim the controller at one of the stretch manipulators and squeeze the **Trigger** halfway to grab it. Move the controller to scale the selected Actor(s) along all axes uniformly. |
-| **Stretch** | ![](universalgizmo_stretch.png)(w:192) | Aim the controller at one of the stretch manipulators and squeeze the **Trigger** halfway to grab it. Move the controller to scale the selected Actor(s) non-uniformly along all axes while remaining anchored to the opposite corner. |
+| **Translate** | ![](Translate.png)(w:300) | Aim the controller at one of the translation manipulators and squeeze the **Trigger** halfway to grab it. Move the controller to move the selected Actor(s) along the axis. |
+| **Translate Plane** | ![](TranslatePlane.png)(w:300) | Aim the controller at one of the translation plane manipulators and squeeze the **Trigger** halfway to grab it. Move the controller to move the selected Actor(s) along the plane (in 2 axes). |
+| **Rotate** | ![](Rotate.png)(w:300) | Aim the controller at one of the rotation manipulators and squeeze the **Trigger** halfway to grab it. Move the controller to rotate the selected Actor(s) around the axis. |
+| **Uniform Scale** | ![](UniformScale.png)(w:300) | Aim the controller at the white cube in the center of the gizmo and squeeze the **Trigger** halfway to grab it. Move the controller to scale the selected Actor(s) along all axes uniformly. |
+| **Scale** | ![](NonUniformScale.png)(w:300) | Aim the controller at one of the scale manipulators and squeeze the **Trigger** halfway to grab it. Move the controller to scale the selected Actor(s) along the Y-axis. Note that the individual scale manipulators only show if you have set **Local Space** using the [Quick Select Menu](Engine/Editor/VR/QuickSelect).|
+
