#UE4 docs HLOD corrections

diff --git a/Engine/Documentation/Source/Engine/HLOD/HLODLanding.INT.udn b/Engine/Documentation/Source/Engine/HLOD/HLODLanding.INT.udn
--- a/Engine/Documentation/Source/Engine/HLOD/HLODLanding.INT.udn
+++ b/Engine/Documentation/Source/Engine/HLOD/HLODLanding.INT.udn
@@ -4,7 +4,7 @@ crumbs:
 availability:docs
 version:4.12
 
-**Hierarchical Level of Detail** meshes (HLODs) have been developed as a way to increase performance and keep the quality bar high for projects using Unreal Engine 4.11 or later. HLODs can replace multiple static mesh actors with a group of static mesh actors at long view distances. This helps reduce the number of actors that need to be rendered for the scene, increasing performance by lowering the number of draw calls per frame. 
+**Hierarchical Level of Detail** meshes (HLODs) have been developed as a way to increase performance and keep the quality bar high for projects using Unreal Engine 4.11 or later. HLODs can replace multiple **Static Mesh Actors** with single, combined **Static Mesh Actor** at long view distances. This helps reduce the number of **Actors** that need to be rendered for the scene, increasing performance by lowering the number of draw calls per frame. 
 
 [REGION:warning]
 HLOD doesn't currently work well with dynamic shadows (casting or recieving them), and as such, you should use static lighting and shadowing for the meshes that the HLOD system generates.
diff --git a/Engine/Documentation/Source/Engine/HLOD/HowTo/BuildingHLODs/BuildingHLODs.INT.udn b/Engine/Documentation/Source/Engine/HLOD/HowTo/BuildingHLODs/BuildingHLODs.INT.udn
--- a/Engine/Documentation/Source/Engine/HLOD/HowTo/BuildingHLODs/BuildingHLODs.INT.udn
+++ b/Engine/Documentation/Source/Engine/HLOD/HowTo/BuildingHLODs/BuildingHLODs.INT.udn
@@ -6,23 +6,21 @@ version:4.12
 related:Engine/HLOD/Overview
 related:Engine/HLOD/Reference
 
-To build your Hierarchical Level of Detail (HLOD) Meshes, there are 2 major steps. Firstly, you must **Generate Clusters**, which groups Actors in your level together based on the setting you specify in the **Cluster Generation Settings**. This process is fairly . Then you have to Generate Proxy Meshes, which can take a very long time based on scene complexity and settings specified in the **Mesh Generation Settings**.
-
-For each LOD level you'll need to specify the settings for that LOD (such as 
+There are two major steps in building your **Hierarchical Level of Detail** (HLOD) Meshes. Firstly, you must **Generate Clusters**, which groups **Actors** in your level together based on the setting you specify in the **Cluster Generation Settings**. Then you have to **Generate Proxy Meshes**, which can take a very long time based on scene complexity and settings specified in the **Mesh Generation Settings**.
 
 [REGION:warning]
-HLOD doesn't currently work well with dynamic shadows (casting or recieving them), and as such, you should use static lighting and shadowing for the meshes that the HLOD system generates.
+HLOD doesn't currently work well with dynamic shadows (casting or receiving them), and as such, you should use static lighting and shadowing for the meshes that the HLOD system generates.
 [/REGION]
 
 ## Steps
 
-Below are the steps to generate your HLOD meshes broken into 2 major sections: Generate Clusters and Generate Proxy Meshes.
+Below are the steps to generate your **HLOD Meshes** broken into two major sections: **Generate Clusters** and **Generate Proxy Meshes**.
 
 ### Generate Clusters
 
-Cluster generation uses the settings from the individual HLOD Levels to decide how it should group static mesh actors in the scene together. This generation process can take some time depending on the settings use, the number of actors being groups, whether materials are generated, and generally on your hardware’s specs. 
+Cluster generation uses the settings from the individual HLOD Levels to decide how it should group **Static Mesh Actors** in the scene together. This generation process can take some time depending on the settings use, the number of **Actors** being groups, whether materials are generated, and generally on your hardware’s specs. 
 
-1. Once you’ve setup your specific settings for the individual HLOD levels you need, click the Generate Clusters button.
+1. Once you’ve setup your specific settings for the individual HLOD levels you need, click the **Generate Clusters** button.
 	
 	![image alt text](Engine/HLOD/Overview/image_3.png)(h:28 w:112 convert:false)
 
@@ -30,31 +28,25 @@ Cluster generation uses the settings from the individual HLOD Levels to decide h
 
 	![image alt text](Engine/HLOD/Overview/image_4.png)(h:58 w:398 convert:false)
 
-1. Now that the process has completed you’ll see the HLOD Outliner populated with all the clustered actors. 
+1. Now that the process has completed you’ll see the **HLOD Outliner** populated with all the clustered **Actors**. 
 
 	![image alt text](Engine/HLOD/Overview/image_5.png)(h:350 w:624 convert:false)
 
-1. You can expand the individual LODActors to see what static meshes make up this cluster by clicking the arrow button to the left of the name.
+1. You can expand the individual **LODActors** to see what **Static Meshes** make up this cluster by clicking the arrow button to the left of the name.
 
 	![image alt text](Engine/HLOD/Overview/image_6.png)(h:350 w:624 convert:false)
 
-1. You can visualize the clusters in the level by selecting a LODActor from the HLOD Outliner and locating it in the editor viewport.
+1. You can visualize the clusters in the level by selecting a **LODActor** from the **HLOD Outliner** and locating it in the editor viewport.
 
 	![image alt text](Engine/HLOD/Overview/image_7.png)(h:481 w:624 convert:false)
 
-To visualize how the generated clusters work in the editor you can use the "Forced Viewing Level" slider to see the HLOD in action without the need to have it transition at a specific screen size. This is helpful for troubleshooting any issues appearing on screen that may be part of the generated cluster.
+To visualize how the generated clusters work in the editor you can use the **Forced Viewing Level** slider to see the HLOD in action without the need to have it transition at a specific screen size. This is helpful for troubleshooting any issues appearing on screen that may be part of the generated cluster.
 
 ![image alt text](Engine/HLOD/Overview/image_8.png)(h:17 w:624 convert:false)
 
-Should you run into any issues you don’t like with the generated cluster you can expand the specified cluster and select the offending static mesh actor. You can then click and drag it to another cluster or by right-clicking on the actor you have the option to **Remove** or **Exclude** it from the generation of the cluster. 
-
-![image alt text](Engine/HLOD/Overview/image_9.png)(h:101 w:392 convert:false)
-
-You also have some Right-Click options that work for a single LODActor or the selection.
+Should you run into any issues you don’t like with the generated cluster you can expand the specified cluster and select the offending **Static Mesh Actor**. You can then click and drag it to another cluster, or you have the option to **Remove** or **Exclude** it from the generation of the cluster by right-clicking on the **Actor's** name in the list. 
 
-![image alt text](Engine/HLOD/Overview/image_10.png)(h:132 w:234 convert:false)
-
-Actors can also be excluded on a per-instance basis by selecting them in the level and in their Details panel set the option for **Enable Auto LODGeneration** to false.
+**Actors** can also be excluded on a per-instance basis by selecting them in the level and in their **Details Panel** set the option for **Enable Auto LODGeneration** to false.
 
 ![image alt text](Engine/HLOD/Overview/image_11.png)(h:236 w:269 convert:false)
 
@@ -62,7 +54,7 @@ Repeat this process while adjusting the values in the HLOD's level **Cluster Gen
 
 ### Generate Proxy Meshes
 
-After you’re happy with the results of the generated cluster you can move on to the option to Build the clusters into a Proxy Mesh. This proxy mesh will be a newly created static mesh actor that combines materials (if enabled), has it’s own lightmaps, and it’s own editable static mesh that can be opened in the static mesh editor. 
+After you’re happy with the results of the generated cluster you can move on to the option to Build the clusters into a Proxy Mesh. This proxy mesh will be a newly created **Static Mesh Actor** that combines materials (if enabled), has it’s own lightmaps, and it’s own editable static mesh that can be opened in the static mesh editor. 
 
 1. If you’re ready to build the proxy meshes you can now click the button  "Generate Proxy Mesh" button to start the process. 
 
@@ -75,5 +67,7 @@ After you’re happy with the results of the generated cluster you can move on t
 	Depending on your HLOD Level settings, the number of proxy meshes being created, and your system specs this process can take a while, even for high-end machines!
 
 	[REGION:note]
-	As a reference the Proxy Mesh Generation process for my machine (12-core i7 processor, GTX-980, and 64GB of RAM) took ~10-12 minutes with default settings for both HLOD Levels and ~100+ LODActors per HLOD level.
+	As a reference the Proxy Mesh Generation process for my machine (12-core i7 processor, GTX-980, and 64GB of RAM) took ~10-12 minutes with default settings for both HLOD Levels and ~100+ **LODActors** per HLOD level.
 	[/REGION]
+
+
diff --git a/Engine/Documentation/Source/Engine/HLOD/HowTo/Enabling/EnablingHLODs.INT.udn b/Engine/Documentation/Source/Engine/HLOD/HowTo/Enabling/EnablingHLODs.INT.udn
--- a/Engine/Documentation/Source/Engine/HLOD/HowTo/Enabling/EnablingHLODs.INT.udn
+++ b/Engine/Documentation/Source/Engine/HLOD/HowTo/Enabling/EnablingHLODs.INT.udn
@@ -9,7 +9,7 @@ related:Engine/HLOD/Reference
 The **Hierarchical Level of Detail System** (HLOD) must be enabled on each Level that wants to make use of it. Follow the steps below to enable the Hierarchical LOD System.
 
 [REGION:warning]
-HLOD doesn't currently work well with dynamic shadows (casting or recieving them), and as such, you should use static lighting and shadowing for the meshes that the HLOD system generates.
+HLOD doesn't currently work well with dynamic shadows (casting or receiving them), and as such, you should use static lighting and shadowing for the meshes that the HLOD system generates.
 [/REGION]
 
 ## Steps
diff --git a/Engine/Documentation/Source/Engine/HLOD/HowTo/HLODHowTos.INT.udn b/Engine/Documentation/Source/Engine/HLOD/HowTo/HLODHowTos.INT.udn
--- a/Engine/Documentation/Source/Engine/HLOD/HowTo/HLODHowTos.INT.udn
+++ b/Engine/Documentation/Source/Engine/HLOD/HowTo/HLODHowTos.INT.udn
@@ -1,8 +1,10 @@
 title:Hierarchical Level of Detail How-To's
 description:Words
-crumbs:
+crumbs:This page contains links to the Hierarchical Level of Detail How-To's
 availability:docs
 
+This page contains links to the **Hierarchical Level of Detail** (HLOD) mesh's How-To's. They will walk you through enabling the HLOD system and setting up HLOD in your game.
+
 [OBJECT:TopicButtonList]
 	[PARAM:icon]
 		![](%ROOT%/tech_icon.png)(convert:false)
@@ -13,8 +15,8 @@ availability:docs
 	[PARAM:description]
 	[/PARAM]
 	[PARAM:links]
-		* [](Engine/HLOD/HowTo/BuildingHLODs "%Engine/HLOD/HowTo/BuildingHLODs:description%")
 		* [](Engine/HLOD/HowTo/Enabling "%Engine/HLOD/HowTo/Enabling:description%")
+		* [](Engine/HLOD/HowTo/BuildingHLODs "%Engine/HLOD/HowTo/BuildingHLODs:description%")
 	[/PARAM]
 [/OBJECT]
 []()
\ No newline at end of file
diff --git a/Engine/Documentation/Source/Engine/HLOD/Overview/HLODOverview.INT.udn b/Engine/Documentation/Source/Engine/HLOD/Overview/HLODOverview.INT.udn
--- a/Engine/Documentation/Source/Engine/HLOD/Overview/HLODOverview.INT.udn
+++ b/Engine/Documentation/Source/Engine/HLOD/Overview/HLODOverview.INT.udn
@@ -3,7 +3,7 @@ description:Overview of the Hierarchical Level of Detail system in Unreal Engine
 crumbs:%ROOT%
 availability:docs
 
-HLODs, in their simplest form, combine pre-existing **Static Mesh Actors** into a single Static Mesh Actor and material, with atlased textures, to reduce draw calls which in-turn helps with your performance of your game. There is a robust system in place to help you set specifically what needs to be generated when you create these HLOD actors for your scenes and covers more options that is normally available for LODs for a single actor. 
+HLODs, in their simplest form, combine pre-existing **Static Mesh Actors** into a single **Static Mesh Actor** and material, with atlased textures, to reduce draw calls which in-turn helps with your performance of your game. There is a robust system that enables you to set what needs to be generated when you create these **HLOD Actors**.
 
 [REGION:warning]
 HLOD doesn't currently work well with dynamic shadows (casting or recieving them), and as such, you should use static lighting and shadowing for the meshes that the HLOD system generates.
@@ -11,19 +11,20 @@ HLOD doesn't currently work well with dynamic shadows (casting or recieving them
 
 ## Proxy Meshes
 
-Once your proxy meshes have been generated you can find these located under the HLOD folder generated in your project. This will include the LODActors referenced in the HLOD Outliner. 
+Once your proxy meshes have been generated, you can find them located under the `HLOD` folder generated in your project. This will include the **LODActors** referenced in the **HLOD Outliner**. 
 
 ![image alt text](image_15.png)(h:320 w:373 convert:false)
 
-These proxy meshes can be opened independently and settings adjusted. If you decided to merge the materials you’ll have access to these as well to adjust the texture source and any settings herein. 
+These proxy meshes can be opened independently and settings adjusted. If you decided to merge the materials, you’ll have access to them to adjust the texture source and any material settings. 
 
 ![image alt text](image_16.png)(h:357 w:404 convert:false)
 
-The proxy mesh above contains several different elements that originally all had their own texture that have now been combined into a single texture below.
+The proxy mesh above contains several different elements that originally all had their own textures that have now been combined into a single texture.
+
 
 [REGION:note]
 
-This can be problematic for anything using Mask or Opacity. These alpha channels are not passed to these asset for this merged texture. If opacity or masked items are need you would want to consider unchecking the option for Merge Textures. When doing this the Proxy Mesh will assign the original material as a Material Element ID instead of combining. This has it’s own advantages and disadvantages.
+This can be problematic for anything using **Mask** or **Opacity**. These alpha channels are not passed to these assets for this merged texture. If opacity or masked items are need you would want to consider unchecking the option for **Merge Textures**. When doing this the Proxy Mesh will assign the original material as a **Material Element ID** instead of combining.
 
 [/REGION]
 
@@ -31,13 +32,15 @@ This can be problematic for anything using Mask or Opacity. These alpha channels
 
 ## HLOD Volume
 
-**The HLOD Volume** is a volume used to manually define/create a HLOD cluster. This volume can be added from the **Modes** panel **Place Mode** and dragged into your level.  
+**The HLOD Volume** is used to manually define/create a HLOD cluster. This volume can be added from the **Modes** panel **Place Mode** and dragged into your level.  
 
 ![image alt text](image_18.png)(h:266 w:229 convert:false)
 
-Place this volume around actors you want to be grouped into the same cluster. It can be helpful to make this slightly larger than the actors and should not be tightly bounding them. Once you are have placed the volume around the actors you can use the option **Generate Clusters** from the HLOD Outliner to create a new cluster with these grouped actors.
+Place this volume around **Actors** you want to be grouped into the same cluster. It can be helpful to make this slightly larger than the **Actors** and should not be tightly bounding them. Once you are have placed the volume around the **Actors** you can use the option **Generate Clusters** from the HLOD Outliner to create a new cluster with these grouped **Actors**.
+
+### Example
 
-Example for Comparison: Generated Clusters before adding HLOD Volume
+Below is an example of Generated Clusters before and after adding HLOD Volume.
 
 ![image alt text](image_19.png)(h:264 w:354 convert:false)
 
@@ -45,21 +48,21 @@ HLOD Level Desired Bound Radius: 500
 
 ![image alt text](image_20.png)(h:288 w:372 convert:false)
 
-Drag in the Volume and scale accordingly to cover your actors.
+Drag in the Volume and scale accordingly to cover your **Actors**.
 
  ![image alt text](image_21.png)(h:273 w:370 convert:false)
 
-Selecting the generated LODActor from the HLOD Outliner shows the cluster created and the Cluster bounds. 
+Selecting the generated **LODActor** from the HLOD Outliner shows the cluster created and the Cluster bounds. 
 
-Right-click on the LODActor in the HLOD Outliner and click on **Select Contained Actors** to see the actors used for this particular LODActor in your scene.
+Right-click on the **LODActor** in the HLOD Outliner and click on **Select Contained Actors** to see the **Actors** used for this particular **LODActor** in your scene.
 
 ![image alt text](image_22.png)(h:93 w:212 convert:false)
 
 ![image alt text](image_23.png)(h:268 w:357 convert:false)
 
-## Details Panel for LODActor
+## Details Panel for **LODActor**
 
-By selecting a LODActor in the level you have access to some settings in the Details panel to override some aspects of your generated HLOD along with details to which static meshes have been included with this cluster. This is not unlike any other actors in your scene where you have the option to override or change some values to work within the confines of your design. 
+By selecting a **LODActor** in the level you have access to some settings in the Details panel to override some aspects of your generated HLOD along with details to which static meshes have been included with this cluster. This is not unlike any other **Actors** in your scene where you have the option to override or change some values to work within the confines of your design. 
 
 ![image alt text](image_24.png)(h:237 w:270 convert:false)
 
diff --git a/Engine/Documentation/Source/Engine/HLOD/Reference/HLODReference.INT.udn b/Engine/Documentation/Source/Engine/HLOD/Reference/HLODReference.INT.udn
--- a/Engine/Documentation/Source/Engine/HLOD/Reference/HLODReference.INT.udn
+++ b/Engine/Documentation/Source/Engine/HLOD/Reference/HLODReference.INT.udn
@@ -3,7 +3,7 @@ description:A reference of the interface elements and properties found in the Hi
 crumbs:
 availability:docs
 
-This page details the buttons and properties that exist on the Hierarchical LOD Outliner.
+This page details the buttons and properties that exist on the **Hierarchical LOD Outliner**.
 
 [REGION:warning]
 HLOD doesn't currently work well with dynamic shadows (casting or recieving them), and as such, you should use static lighting and shadowing for the meshes that the HLOD system generates.
@@ -19,7 +19,7 @@ HLOD doesn't currently work well with dynamic shadows (casting or recieving them
 Below are the descriptions for the 3 buttons and slider that exist in the **Hierarchical LOD Outliner**.
 
 | --- | :---: | :--- |
-| **Generate Clusters** || This will generate spherical clusters of static mesh actors in the scene that can be grouped. The way this is handled is by using the settings for each LOD settings specified under the LODSystem section of the HLOD Outliner. |
+| **Generate Clusters** || This will generate spherical clusters of **Static Mesh Actors** in the scene that can be grouped. The way this is handled is by using the settings for each LOD settings specified under the LODSystem section of the HLOD Outliner. |
 | **Delete Clusters** || This will delete all HLODs that have been previously setup in this level. This includes all generated proxy meshes as well as the clusters that have been setup. This will not wipe the previously set values or properties for the HLOD Levels. |
 | **Generate Proxy Meshes** || [INCLUDE:#gpm] |
 | **Forced Viewing Level** || This will allow you to force a specific LOD level to be visible in the editor. This is helpful for debugging any LOD issues you may have. |
@@ -38,19 +38,19 @@ Below, broken out by major section, are the properties that can be found in a **
 | --- | :---: | :--- |
 | **Desired Bound Radius** || This is not guaranteed but used to calculate filling factor for auto clustering. |
 | **Desired Filling Percentage** || This is not guaranteed but used to calculate filling factor for auto clustering. |
-| **Min Number of Actors to Build** || Min number of actors to build LODActor. |
+| **Min Number of Actors to Build** || Min number of **Actors** to build LODActor. |
 
 ### Mesh Generation Settings
 
 | Property | %spacer% | Description |
 | --- | :---: | :--- |
-| **Transition Screen Size** || The screen radius a mesh object should reach before swapping to the LOD actor. When the parent LOD displays it won’t display any of the children. |
+| **Transition Screen Size** || The screen radius a mesh object should reach before swapping to the LOD **Actor**. When the parent LOD displays it won’t display any of the children. |
 | [REGION:tablesection]**Merge Settings**[/REGION] |||
 | **Generate Light Map UV** || Whether to generate a lightmap UV for a merged mesh. |
 | **Target Light Map UVChannel** || The target UV channel to assign the lightmap UV. |
 | **Target Light Map Resolution** || The lightmap resolution set for the merged mesh. |
-| **Pivot Point at Zero** || Whether a merged mesh should have a pivot at the World Origin, or at the first merged actor from a generated cluster. |
-| **Merge Physics Data:** || Whether to merge the collision primitives when combining the actors. |
+| **Pivot Point at Zero** || Whether a merged mesh should have a pivot at the World Origin, or at the first merged **Actor** from a generated cluster. |
+| **Merge Physics Data:** || Whether to merge the collision primitives when combining the **Actors**. |
 | **Merge Materials** || Whether to merge source materials into one flat material. |
 | [REGION:tablesection]**[Material Settings See Below](#materialsettings)**[/REGION] |||
 | **Bake Vertex Data ** || Incorporates painted vertex colours into model/baked out material. |
@@ -85,7 +85,7 @@ Below, broken out by major section, are the properties that can be found in a **
 
 <!--
 [EXCERPT:GPM]
-When you generate the proxy meshes this will use the clustered LODs to merge these actors together into a single asset. Depending on what settings you changed for the HLOD levels your static meshes will generate based on this. This can affect your materials, lightmaps, and various other properties as well. 
+When you generate the proxy meshes this will use the clustered LODs to merge these **Actors** together into a single asset. Depending on what settings you changed for the HLOD levels your static meshes will generate based on this. This can affect your materials, lightmaps, and various other properties as well. 
 
 [REGION:note]
 Be aware that this process can take a **very long time** depending on your scene complexity and settings for your HLOD levels.
