Peer review edits for Lighting Channels documentation

diff --git a/Engine/Documentation/Source/Engine/Rendering/LightingAndShadows/LightingChannels/LightingChannels.INT.udn b/Engine/Documentation/Source/Engine/Rendering/LightingAndShadows/LightingChannels/LightingChannels.INT.udn
--- a/Engine/Documentation/Source/Engine/Rendering/LightingAndShadows/LightingChannels/LightingChannels.INT.udn
+++ b/Engine/Documentation/Source/Engine/Rendering/LightingAndShadows/LightingChannels/LightingChannels.INT.udn
@@ -3,7 +3,7 @@ Title: Lighting Channels
 Crumbs: %ROOT%, Engine
 Description:The Lighting Channels feature overview.
 Navigation:topic
-Version: 4.12.5
+Version: 4.13
 
 [VAR:Topic]
 [OBJECT:Topic]
@@ -47,17 +47,17 @@ Version: 4.12.5
 
 ![](image_39.png)
 
-**Lighting Channels** allow dynamic lights to affect only objects when their lighting channels overlap. This is primarily intended for cinematic use to give users powerful control over how **Actors** are lit. At the moment, Unreal Engine supports up to 3 lighting channels.
+**Lighting Channels** allow dynamic lights to only affect objects when their lighting channels overlap. This is primarily intended for cinematic use to give users powerful control over how **Actors** are lit. At the moment, Unreal Engine supports up to 3 lighting channels.
 
 ## Usage
 
-By default **Directional Lights**, **Spot Lights**, **Point Lights**, and all **Actors** that can be affected by lights (**Static Meshes**, **Skeletal Meshes**, etc...) have **Lighting Channel 0** enabled. If you desire a lightable **Actor** to be affected by another **Lighting Channel** you must enable that channel on both the **Actor** and the light.
+By default, Directional Lights, Spot Lights, Point Lights, and all Actors that can be affected by lights (Static Meshes, Skeletal Meshes, etc.) have **Lighting Channel 0** enabled. If you would like a lightable **Actor** to be affected by another **Lighting Channel**, you must enable that channel on both the Actor and the light.
 
 ### Example
 
 ![](topImage.png)
 
-In the above image the **White Directional Light** can only affect **Channel 0**, which includes the mannequin on the left and the background **Static Meshes**. While the **Purple Point Light** can only affect **Channel 1**, which **only** includes the mannequin on the right.
+In the above image, the white Directional Light can only affect **Channel 0**, which includes the mannequin on the left and the background Static Meshes, while the purple Point Light can only affect **Channel 1**, which only includes the mannequin on the right.
 
 You can see how the properties are set here:
 
@@ -65,11 +65,11 @@ You can see how the properties are set here:
 ![](LightingChannelProp.png)
 [/REGION]
 
-For lights, the property can be found in the **Advanced** pull down of the **Light** category in the **Details Panel**. For **Actors** that can be lit, you can find the **Lighting Channels** properties in the **Lighting** category of the **Details Panel**.
+For lights, the property can be found in the **Advanced** pulldown menu of the **Light** category in the **Details Panel**. For Actors that can be lit, you can find the **Lighting Channels** properties in the **Lighting** category of the **Details Panel**.
 
 ## Limitations
 
-Lighting channel influence is applied dynamically. This means it won't work with **Static Lights** or **Static Mesh Actors with Static Mobility** (**Static Mesh Actors** with **Mobility** set to **Movable** work just fine). You'll need to use either **Stationary** or **Movable** lights.
+Lighting channel influence is applied dynamically. This means it won't work with Static Lights or Static Mesh Actors with Static Mobility. Static Mesh Actors with **Mobility** set to **Movable** do work, however. You'll need to use either **Stationary** or **Movable** lights.
 
 Lighting channels will only affect direct lighting on opaque materials. So, translucent or masked materials will not work.
 
@@ -92,7 +92,7 @@ The performance impact of using **Lighting Channels** is minimal, but also non-z
 
 ## On Mobile
 
-As of Unreal Engine 4 version 4.12 **Lighting Channels** works on mobile renderers.
+As of Unreal Engine 4.12, **Lighting Channels** works on mobile renderers.
 
 
 
