Added config file info to Android IAP page, crosslinked IAP pages UEDOC-3189

diff --git a/Engine/Documentation/Source/Platforms/Android/InAppPurchases/In-AppPurchases.INT.udn b/Engine/Documentation/Source/Platforms/Android/InAppPurchases/In-AppPurchases.INT.udn
--- a/Engine/Documentation/Source/Platforms/Android/InAppPurchases/In-AppPurchases.INT.udn
+++ b/Engine/Documentation/Source/Platforms/Android/InAppPurchases/In-AppPurchases.INT.udn
@@ -7,7 +7,7 @@ Parent:Platforms/Android
 tags:Mobile
 tags:Android
 type:howto
-
+Related:Platforms/Mobile/InAppPurchases/
 
 ## Configuration
 
@@ -35,6 +35,11 @@ type:howto
 	
 	![image alt text](image_2.png)
 
+1. Edit [ProjectName]/Config/IOS/IOSEngine.ini:
+
+		[OnlineSubsystemIOS.Store]
+		bSupportsInAppPurchasing=True
+
 [INCLUDE:Platforms/Mobile/InAppPurchases#usage_android]
 
 ## Testing
diff --git a/Engine/Documentation/Source/Platforms/Mobile/InAppPurchases/In-AppPurchases.INT.udn b/Engine/Documentation/Source/Platforms/Mobile/InAppPurchases/In-AppPurchases.INT.udn
--- a/Engine/Documentation/Source/Platforms/Mobile/InAppPurchases/In-AppPurchases.INT.udn
+++ b/Engine/Documentation/Source/Platforms/Mobile/InAppPurchases/In-AppPurchases.INT.udn
@@ -76,9 +76,7 @@ There is a non-latent version of this function (the Blueprint node will display
 
 ![image alt text](image_5.png)
 
-To restrore purchases, use the **Restore In-App Purchases** Blueprint node (or associated C++ function call). It only takes in the Player Controller, and returns an array of all the purchase information associated with that Player Controller. You can then process that array to compare against the particular Identifier (or Identifiers) that your gameplay logic needs. 
+To restore purchases, use the **Restore In-App Purchases** Blueprint node (or associated C++ function call). It only takes in the Player Controller, and returns an array of all the purchase information associated with that Player Controller. You can then process that array to compare against the particular Identifier (or Identifiers) that your gameplay logic needs. 
 
 [/EXCERPT:usage]
-[REGION:warning]
-As of 4.10, Restore In-App Purchases is a function that only works on iOS. 
-[/REGION]
+
diff --git a/Engine/Documentation/Source/Platforms/iOS/InAppPurchases/In-AppPurchases.INT.udn b/Engine/Documentation/Source/Platforms/iOS/InAppPurchases/In-AppPurchases.INT.udn
--- a/Engine/Documentation/Source/Platforms/iOS/InAppPurchases/In-AppPurchases.INT.udn
+++ b/Engine/Documentation/Source/Platforms/iOS/InAppPurchases/In-AppPurchases.INT.udn
@@ -5,7 +5,7 @@ Availability:Public
 Parent:Platforms/iOS
 order:1
 type:howto
-
+Related:Platforms/Mobile/InAppPurchases/
 
 ## Configuration
 
