Fixing swapped config file sections

diff --git a/Engine/Documentation/Source/Platforms/Android/InAppPurchases/In-AppPurchases.INT.udn b/Engine/Documentation/Source/Platforms/Android/InAppPurchases/In-AppPurchases.INT.udn
--- a/Engine/Documentation/Source/Platforms/Android/InAppPurchases/In-AppPurchases.INT.udn
+++ b/Engine/Documentation/Source/Platforms/Android/InAppPurchases/In-AppPurchases.INT.udn
@@ -7,6 +7,7 @@ Parent:Platforms/Android
 tags:Mobile
 tags:Android
 type:howto
+Prereq:P
 Related:Platforms/Mobile/InAppPurchases/
 
 ## Configuration
@@ -35,9 +36,12 @@ Related:Platforms/Mobile/InAppPurchases/
 	
 	![image alt text](image_2.png)
 
-1. Edit [ProjectName]/Config/IOS/IOSEngine.ini:
+1. Edit [ProjectName]/Config/Android/AndroidEngine.ini:
 
-		[OnlineSubsystemIOS.Store]
+		[OnlineSubsystem]
+		DefaultPlatformService=GooglePlay
+
+		[OnlineSubsystemGooglePlay.Store]
 		bSupportsInAppPurchasing=True
 
 [INCLUDE:Platforms/Mobile/InAppPurchases#usage_android]
diff --git a/Engine/Documentation/Source/Platforms/iOS/InAppPurchases/In-AppPurchases.INT.udn b/Engine/Documentation/Source/Platforms/iOS/InAppPurchases/In-AppPurchases.INT.udn
--- a/Engine/Documentation/Source/Platforms/iOS/InAppPurchases/In-AppPurchases.INT.udn
+++ b/Engine/Documentation/Source/Platforms/iOS/InAppPurchases/In-AppPurchases.INT.udn
@@ -5,7 +5,7 @@ Availability:Public
 Parent:Platforms/iOS
 order:1
 type:howto
-Related:Platforms/Mobile/InAppPurchases/
+Related:
 
 ## Configuration
 
