#UE4 Docs: Adding new part for Full Prescision on mobile.

diff --git a/Engine/Documentation/Source/Platforms/Mobile/Materials/MobileMaterials.INT.udn b/Engine/Documentation/Source/Platforms/Mobile/Materials/MobileMaterials.INT.udn
--- a/Engine/Documentation/Source/Platforms/Mobile/Materials/MobileMaterials.INT.udn
+++ b/Engine/Documentation/Source/Platforms/Mobile/Materials/MobileMaterials.INT.udn
@@ -34,9 +34,7 @@ There are some considerations you must take into account when creating Materials
 * Only Default and Unlit shading models are available.
 * [](Engine/Rendering/Materials/CustomizedUVs) should be used to avoid dependent texture fetches (no math on texture UVs).
 * Translucent and Masked Materials are extremely expensive. It is recommended that you use [Opaque Materials](Engine/Rendering/Materials/MaterialProperties/BlendModes#opaque) wherever possible.
-* Depth fade can be used in translucent Materials for **iOS platforms**. It is not supported on platforms where the hardware does not support fetching from the depth buffer and 
-	emulating this would incur an unacceptable performance cost.
-
+* Depth fade can be used in translucent Materials for **iOS platforms**. It is not supported on platforms where the hardware does not support fetching from the depth buffer and emulating this would incur an unacceptable performance cost.
 
 In addition, some features are simply not available on mobile platforms.
 
@@ -53,15 +51,37 @@ To maintain compatibility with older devices, use the [Feature Level](Engine/Ren
 
 
 ## Mobile-Specific Material Settings
+The following Material properties are to help optimize and or reduce the rendering cost of the Material when used on a mobile device. 
 
 | Property | Description |
 | -------- | ---------- |
 | bFullyRough | When enabled, saves a number of shader ALU instructions and one sampler.  This effectively disables Specular while maintaining overall reflective energy. |
+| High Quality Reflections |Enables high quality reflections in the forward renderer. Enabling this setting reduces the number of samplers available to the Material as two more samplers will be used for reflections cubemaps. |
 | bUseLightmapDirectionality | When disabled, lighting from lightmaps will be flat but cheaper. |
+| Use Full Precision | [INCLUDE:Platforms\Mobile\Materials\#UFP]|
 
 ## Troubleshooting Materials for Mobile
-
+Due to hardware differences between PC's and mobile devices you could run into issues with certain options working on the PC but not working when viewed on a mobile device. The following list contains some of the more common issues you could run into and ways to fix them.
 | Issue | Solution |
 | ----- | -------- |
 | Default Material is displayed in the Mobile Previewer. | Open the [log file](Engine/Basics/DirectoryStructure) in **Engine\Saved** and search for error messages to know how to fix the Material. |
 | Material Animation with Time Expressions looks choppy but framerate is good elsewhere. | Make sure the Period property is enabled in the **Time** expression along with a Period. This will move the computation to the CPU and run at full-precision instead of half. |
+| Materials using World Coordinates computations or very small values are not being displayed correctly. | Try enabling the Use Full Precision option on the Material that is not displaying correctly on a mobile device. The precision level of Materials on Mobile devices is often truncated to one or two decimal places to help with performance. Enabling Use Full Precision will make sure that the Material's precision kept to the values that are set in the Material Editor  |
+
+
+
+<!--
+[EXCERPT:UFP]
+When enabled, Use Full Precision will make sure that the Material does not automatically truncate data when the Material used on a mobile device to help increase performance.
+
+[OBJECT:ComparisonSlider]
+	[PARAM:before]
+	![Full Precision Off ](Full_Precision_Off.png)
+	[/PARAM]
+	[PARAM:after]
+	![Full Precision On](Full_Precision_On.png)
+	[/PARAM]
+[/OBJECT]
+[/EXCERPT:UFP]
+-->
+
