Initial Draft of the Installed Build Reference per MatthewG's write-up.

#UE4doc

diff --git a/Engine/Documentation/Source/Programming/Development/InstalledBuildReference/InstalledBuildReference.INT.udn b/Engine/Documentation/Source/Programming/Development/InstalledBuildReference/InstalledBuildReference.INT.udn
new file mode 100644
--- /dev/null
+++ b/Engine/Documentation/Source/Programming/Development/InstalledBuildReference/InstalledBuildReference.INT.udn
@@ -0,0 +1,166 @@
+Availability: Docs
+Title: Installed Build Reference
+Crumbs: %ROOT%, Programming, Programming/Development
+Description:This page describes the Installed Build process, including an overview of how to write Installed Build scripts.
+version: 4.13
+parent:Programming/Development
+type:Reference
+related:Programming/Development/BuildConfigurations
+related:Programming/UnrealBuildSystem
+
+[TOC (start:1 end:4)]
+
+%Globals:osselect%
+
+<!--Write an Overview!-->
+
+## What's an Installed Build?
+
+An **Installed Build** of Unreal Engine 4 (UE4) is a version that's distributed using the [Epic Games Launcher](GettingStarted/Installation).
+
+![](InstalledBuilds.png)
+
+This version of the engine is typically shipped with each release of UE4 and as such, the Installed Build is fully featured, but focused by an 
+Installed Build process that quickly gets users up and running with the engine. Additionally, the Unreal Editor is distributed as 
+precompiled binaries with static libraries that are built for targeted platforms in development and shipping configurations.
+
+## About the Installed Build Process
+
+The Installed Build process automates build packages and tools that are needed to make an Installed Build version of UE4. 
+In summary, making an Installed Build is the process of setting up default target platforms, compiling UE4 
+(including the Unreal Header Tool, the UE4 Editor, etc.), generating embedded documentation, running tests, 
+and deploying the Installed Build version of UE4 for targeted platforms.
+
+## Making an Installed Build
+
+The Installed Build process is made with our BuildGraph Scripting System. BuildGraph scripts create custom sequences 
+of UE4 build commands for build automation. You can declare graphs by scripting an XML file (using [XML Schema v1.0](https://www.w3.org/TR/NOTE-dcd])) 
+with syntax similar to MSBuild, Apache ANT, and NAnt. Scripted XML files define a network of node dependencies containing tasks that 
+will be executed during the Installed Build process.
+
+If you're interested in reading through example BuildGraph scripts, navigate to the `[UE4Root]/Engine/Build/Graph/Examples` folder, where you'll find: 
+
+* `AllExamples.xml`
+* `Building.xml`
+* `TagsAndFiles.xml`
+
+If you're interested in learning more about our BuildGraph Scripting System<!--TO DO: Add cross-link to pending BuildGraph documentation-->, head over to 
+our BuildGraph Scripting System reference to learn how you can customize and write your own scripts.
+
+### Running the Installed Build Script
+
+The Installed Build script is located at `[UE4Root]/Engine/Build/InstalledEngineBuild.xml`. You can run the Installed Build script by invoking the 
+**AutomationTool** with the following command line:
+
+[OBJECT:ToggleButtonContent]
+	[PARAMLITERAL:category]
+	OS
+	[/PARAMLITERAL]
+	[PARAMLITERAL:id]
+	windows
+	[/PARAMLITERAL]
+	[PARAMLITERAL:active]
+	active_button_content
+	[/PARAMLITERAL]
+	[PARAM:content]
+	`BuildGraph -target=”Make Installed Build Win64” -script=Engine/Build/InstalledEngineBuild.xml -clean`
+	[/PARAM]
+[/OBJECT]
+
+[OBJECT:ToggleButtonContent]
+	[PARAMLITERAL:category]
+	OS
+	[/PARAMLITERAL]
+	[PARAMLITERAL:id]
+	mac
+	[/PARAMLITERAL]
+	[PARAMLITERAL:active]
+	active_button_content
+	[/PARAMLITERAL]
+	[PARAM:content]
+	`BuildGraph -target=”Make Installed Build Mac” -script=Engine/Build/InstalledEngineBuild.xml -clean`
+	[/PARAM]
+[/OBJECT]
+
+In summary, running `InstalledEngineBuild.xml` makes an Installed Build with all of the default options and platforms enabled. 
+
+#### Installed Build Script Options
+
+If you want to see a list of available options for customizing the build process (along with a list of nodes that'll be built), 
+you'll want to append the `-listonly` option before running the script. The list should look something like this: 
+
+	Options:
+	
+	    -set:PublishDir=...            Root directory to contain published builds and symbols (Default: $(RootDir)/LocalBuilds/Engine) 
+	    -set:WithWin64=...             Include the Win64 target platform (Default: true) 
+	    -set:WithWin32=...             Include the Win32 target platform (Default: true) 
+	    -set:WithMac=...               Include the Mac target platform (Default: true) 
+	    -set:WithAndroid=...           Include the Android target platform (Default: true) 
+	    -set:WithIOS=...               Include the iOS target platform (Default: true) 
+	    -set:WithTVOS=...              Include the tvOS target platform (Default: true) 
+	    -set:WithLinux=...             Include the Linux target platform (Default: true) 
+	    -set:WithHTML5=...             Include the HTML5 target platform (Default: true) 
+	    -set:WithPS4=...               Include the PS4 target platform (Default: false) 
+	    -set:WithXboxOne=...           Include the XboxOne target platform (Default: false) 
+	    -set:WithDDC=...               Build a stand-alone derived-data cache for the engine content and templates (Default: true) 
+	    -set:AnalyticsTypeOverride=... Identifier for analytic events to send 
+	
+	Graph:
+	    Trigger: None
+	        Agent: Editor Win64 (CompileWin64)
+	            Node: Update Version Files
+	            Node: Compile UnrealHeaderTool Win64
+	            Node: Compile UE4Editor Win64
+	        Agent: Editor Mac (CompileMac)
+	            Node: Compile UnrealHeaderTool Mac
+	        Agent: Target Platforms Win64 (CompileWin64)
+	            Node: Compile UE4Game Win64
+	            Node: Compile UE4Game Win32
+	            Node: Compile UE4Game Linux
+	            Node: Compile UE4Game Android
+	            Node: Compile UE4Game HTML5
+	        Agent: Target Platforms Mac (CompileMac)
+	            Node: Compile UE4Game IOS
+	            Node: Compile UE4Game TVOS
+	        Agent: Tools Group Win64 (CompileWin64)
+	            Node: Build Tools Win64
+	            Node: Build Tools CS
+	            Node: Make Feature Packs
+	        Agent: DDC Group Win64 (Win64)
+	            Node: Build DDC Win64
+	        Agent: Installed Build Group Win64 (Win64)
+	            Node: Make Installed Build Win64
+
+Some portions of the node graph can be ignored because they're only relevant to Epic's internal build system. For example, the Agent descriptions, namely `CompileWin64` 
+and `CompileMac`, refer to the type of machine they're meant to run on; however, the script can still be run from a single build machine. 
+
+### Some Additional Commands
+
+Depending on your targeted platform and available Software Development Kit (SDK), we're providing some useful commands that should make your life easier.
+
+#### IOS and TVOS
+
+If you want to include IOS and TVOS in your Installed Build, you'll have to have a Mac setup for remote building. Otherwise, append the following to your command line:
+`-set:WithIOS=false -setWithTVOS=false`.
+
+#### Playstation 4 and XboxOne
+
+If you have the required SDKs, you can add Playstation 4 (PS4) and XboxOne as target platforms (although, this is untested). To add PS4 and XboxOne 
+as target platforms, append `-set:WithPS4=true -set:WithXboxOne=true` to the command line.
+
+#### Amending the Derived Data Cache
+
+Building a stand-alone [Derived Data Cache (DDC)](Engine/Basics/DerivedDataCache) for the Engine and Template content can be one of the slowest parts of the build. If you don't need a stand-alone DDC, you can 
+skip this step by appending `-set:WithDDC=false` to the command line.
+
+#### Not Signing Executables
+
+Typically, you don't have to setup your machine to sign executables, as this can cause the Install Build process to fail while attempting to sign the executables being built. 
+To avoid this problem, append `-NoSign` to the command line.
+
+### Including Additional Files
+
+RuntimeDependencies (set in `build.cs` files for each module) automatically collate and determine which files are included in an Installed Build. 
+However, there are some required files that cannot be included this way, so they're defined in `[UE4Root]/Engine/Build/InstalledEngineFilters.xml`. 
+The `InstalledEngineFilters` XML file also lists patterns of files that should be excluded from the build, determining the types of files that need to be stripped or signed, and finally, deciding 
+what projects to build the DDC for. If you need to add any additional files to the Installed Build, the `InstalledEngineFilters` XML file is a great place to start.
\ No newline at end of file
