#UE4 Docs: Fixing the meta data so this is not nested in the table of contents.

diff --git a/Engine/Documentation/Source/Platforms/Mobile/Performance/TipsAndTricks/MobilePerfTipsAndTricks.INT.udn b/Engine/Documentation/Source/Platforms/Mobile/Performance/TipsAndTricks/MobilePerfTipsAndTricks.INT.udn
--- a/Engine/Documentation/Source/Platforms/Mobile/Performance/TipsAndTricks/MobilePerfTipsAndTricks.INT.udn
+++ b/Engine/Documentation/Source/Platforms/Mobile/Performance/TipsAndTricks/MobilePerfTipsAndTricks.INT.udn
@@ -1,8 +1,13 @@
 Title: Mobile Performance Tips and Tricks
 Description:Optimizations and techniques utilized in the Unreal Match 3 game to ensure optimal performance on supported mobile devices.
-Crumbs:
-Parent:Platforms/Mobile/Performance
+Crumbs:%ROOT%, Platforms, Platforms/Mobile
 Availability:Public
+Version: 4.9
+Platform:Mobile
+parent:Platforms/Mobile
+order:6
+type:overview
+tags:mobile
 
 [VAR:Topic]
 [OBJECT:Topic]
