Made some minor grammatical revisions.
Wrote an Overview for the Reference page.

#UE4doc

diff --git a/Engine/Documentation/Source/Programming/Development/InstalledBuildReference/InstalledBuildReference.INT.udn b/Engine/Documentation/Source/Programming/Development/InstalledBuildReference/InstalledBuildReference.INT.udn
--- a/Engine/Documentation/Source/Programming/Development/InstalledBuildReference/InstalledBuildReference.INT.udn
+++ b/Engine/Documentation/Source/Programming/Development/InstalledBuildReference/InstalledBuildReference.INT.udn
@@ -10,9 +10,10 @@ related:Programming/UnrealBuildSystem
 
 [TOC (start:1 end:4)]
 
-%Globals:osselect%
-
-<!--Write an Overview!-->
+Developers are always looking for ways to efficiently deliver build solutions for a variety of hardware profiles and operating systems. 
+With the release of our BuildGraph scripting system, you'll be able to create custom sequences of UE4 build commands, harnessing the Installed 
+Build process that we've developed for your use. As you read through the following reference page, you'll learn about the Installed Build process, 
+including how to run scripts that make Installed Builds for targeted build solutions.
 
 ## What's an Installed Build?
 
@@ -20,23 +21,21 @@ An **Installed Build** of Unreal Engine 4 (UE4) is a version that's distributed
 
 ![](InstalledBuilds.png)
 
-This version of the engine is typically shipped with each release of UE4 and as such, the Installed Build is fully featured, but focused by an 
-Installed Build process that quickly gets users up and running with the engine. Additionally, the Unreal Editor is distributed as 
-precompiled binaries with static libraries that are built for targeted platforms in development and shipping configurations.
+This version of the engine is shipped with each release of UE4, being fully featured, but focused by an Installed Build process 
+that gets users up and running with the engine. Additionally, the Unreal Editor is distributed as precompiled binaries with static 
+libraries that are built for targeted platforms in development and shipping configurations.
 
 ## About the Installed Build Process
 
 The Installed Build process automates build packages and tools that are needed to make an Installed Build version of UE4. 
-In summary, making an Installed Build is the process of setting up default target platforms, compiling UE4 
-(including the Unreal Header Tool, the UE4 Editor, etc.), generating embedded documentation, running tests, 
-and deploying the Installed Build version of UE4 for targeted platforms.
+In summary, making an Installed Build is the process of setting up default target platforms, compiling UE4 with its tools and 
+editor, generating embedded documentation, running tests, and deploying the Installed Build version of UE4 for targeted platforms.
 
 ## Making an Installed Build
 
-The Installed Build process is made with our BuildGraph Scripting System. BuildGraph scripts create custom sequences 
-of UE4 build commands for build automation. You can declare graphs by scripting an XML file (using [XML Schema v1.0](https://www.w3.org/TR/NOTE-dcd])) 
-with syntax similar to MSBuild, Apache ANT, and NAnt. Scripted XML files define a network of node dependencies containing tasks that 
-will be executed during the Installed Build process.
+The Installed Build process is made with our BuildGraph scripting system. BuildGraph scripts create custom sequences 
+of UE4 build commands for build automation. You can declare graphs by scripting an XML file with syntax (using [XML Schema v1.0](https://www.w3.org/TR/NOTE-dcd])) 
+similar to MSBuild, Apache ANT, and NAnt. Scripted XML files define a network of node dependencies containing tasks that will be executed during the Installed Build process.
 
 If you're interested in reading through example BuildGraph scripts, navigate to the `[UE4Root]/Engine/Build/Graph/Examples` folder, where you'll find: 
 
@@ -44,14 +43,18 @@ If you're interested in reading through example BuildGraph scripts, navigate to
 * `Building.xml`
 * `TagsAndFiles.xml`
 
-If you're interested in learning more about our BuildGraph Scripting System<!--TO DO: Add cross-link to pending BuildGraph documentation-->, head over to 
-our BuildGraph Scripting System reference to learn how you can customize and write your own scripts.
+<!--TO DO: Add cross-link to pending BuildGraph Documentation
+If you're interested in learning more about our BuildGraph Scripting System, head over to our BuildGraph Scripting System reference to learn how you can 
+customize your own scripts.
+-->
 
 ### Running the Installed Build Script
 
 The Installed Build script is located at `[UE4Root]/Engine/Build/InstalledEngineBuild.xml`. You can run the Installed Build script by invoking the 
 **AutomationTool** with the following command line:
 
+%Globals:osselect%
+
 [OBJECT:ToggleButtonContent]
 	[PARAMLITERAL:category]
 	OS
@@ -87,7 +90,7 @@ In summary, running `InstalledEngineBuild.xml` makes an Installed Build with all
 #### Installed Build Script Options
 
 If you want to see a list of available options for customizing the build process (along with a list of nodes that'll be built), 
-you'll want to append the `-listonly` option before running the script. The list should look something like this: 
+you'll want to append the `-listonly` option before running the script. Subsequently, the list output should look something like this: 
 
 	Options:
 	
@@ -131,12 +134,13 @@ you'll want to append the `-listonly` option before running the script. The list
 	        Agent: Installed Build Group Win64 (Win64)
 	            Node: Make Installed Build Win64
 
-Some portions of the node graph can be ignored because they're only relevant to Epic's internal build system. For example, the Agent descriptions, namely `CompileWin64` 
-and `CompileMac`, refer to the type of machine they're meant to run on; however, the script can still be run from a single build machine. 
+Please note that some portions of the node graph can be ignored because they're only relevant to Epic's internal build system. For example, 
+the Agent descriptions, namely `CompileWin64` and `CompileMac`, refer to the type of machine they're meant to run on; however, the script 
+can still be run from a single build machine. 
 
 ### Some Additional Commands
 
-Depending on your targeted platform and available Software Development Kit (SDK), we're providing some useful commands that should make your life easier.
+Depending on your targeted platform and available Software Development Kit (SDK), we're providing some useful commands that you should take note of.
 
 #### IOS and TVOS
 
@@ -150,8 +154,8 @@ as target platforms, append `-set:WithPS4=true -set:WithXboxOne=true` to the com
 
 #### Amending the Derived Data Cache
 
-Building a stand-alone [Derived Data Cache (DDC)](Engine/Basics/DerivedDataCache) for the Engine and Template content can be one of the slowest parts of the build. If you don't need a stand-alone DDC, you can 
-skip this step by appending `-set:WithDDC=false` to the command line.
+Building a stand-alone [Derived Data Cache (DDC)](Engine/Basics/DerivedDataCache) for the Engine and Template content can be one of the slowest aspects of the build. 
+If you don't need a stand-alone DDC, you can skip this step by appending `-set:WithDDC=false` to the command line.
 
 #### Not Signing Executables
 
