Accepted peer review fixes, made a few other changes to the text to clarify based on feedback.

diff --git a/Engine/Documentation/Source/Engine/Editor/VR/QuickSelect/QuickSelectMenu.INT.udn b/Engine/Documentation/Source/Engine/Editor/VR/QuickSelect/QuickSelectMenu.INT.udn
--- a/Engine/Documentation/Source/Engine/Editor/VR/QuickSelect/QuickSelectMenu.INT.udn
+++ b/Engine/Documentation/Source/Engine/Editor/VR/QuickSelect/QuickSelectMenu.INT.udn
@@ -62,18 +62,18 @@ The **Quick Select** menu enables you to quickly adjust grid, rotation and snap
 | Rotate | Toggles grid snapping on or off when rotating Actors. The snapping value is modified by aiming at the numerical box below, squeezing the trigger, and moving the controller. |
 | Scale | Toggles grid snapping on or off when scaling Actors. The snapping value is modified by aiming at the numerical box below, squeezing the trigger, and moving the controller. |
 | [REGION:tablesection]Interface[/REGION] ||
-| Content | Opens the Content Browser |
+| Content | Opens the Content Browser. |
 | Details | Opens the Details panel for viewing and modifying properties of the selected Actor(s). |
-| Modes | Opens the Modes panel providing access to Actor placement tools and more. |
-| Outliner | Opens theWorld Outliner which displays a hierarchical list of all Actors in the world enabling quick selection, visibility toggling, organization, and Actor-specific actions to be performed. |
-| Asset | Opens the Asset window which acts as a container for any asset editors - Blueprint Editor, Material Editor, etc. - that may be open. |
-| Tutorial | Opens the VR Mode Tutorial which provides access to a collection of instructional walkthroughs for common features of the VR Editor. |
-| Light | Adds a flashlight to the controller used to click on this button.  |
+| Modes | Opens the Modes panel, providing access to Actor placement tools, foliage editing, and mesh painting. |
+| Outliner | Opens the World Outliner which displays a hierarchical list of all Actors in the world enabling quick selection, visibility toggling, organization, and Actor-specific actions to be performed. |
+| Asset | Opens the Asset window, which acts as a container for any asset editors - Blueprint Editor, Material Editor, etc. - that may be open. |
+| Tutorial | Opens the VR Mode Tutorial, which provides access to a collection of instructional walkthroughs for common features of the VR Editor. |
+| Light | Adds a flashlight to, or removes the flashlight from, the controller that was used to click the button.  |
 | Screenshot | Takes a screenshot of the current scene. |
-| World | Opens the World Settings panel providing access to properties of the current Level. |
+| World | Opens the World Settings panel, providing access to properties of the current Level. |
 | [REGION:tablesection]Gizmo[/REGION] ||
 | Coordinate System | Sets the coordinate system [World Space/Local Space] to use when transforming Actors. |
 | Transform Selector | Cycles the Universal Gizmo through the available transform modes: Translate, Rotate, Scale. Only the manipulators for the selected mode are displayed. |
 | [REGION:tablesection] [/REGION] ||
-| Play | Enables Play In Editor mode.  |
+| Play | Starts your project from within the VR Editor so that you can test your gameplay quickly. To instantly go back to the VR Editor from your test session, press both triggers and a grip button on each motion controller.  |
 
