Minor typo fix found during peer review

diff --git a/Engine/Documentation/Source/Platforms/Mobile/Materials/MobileMaterials.INT.udn b/Engine/Documentation/Source/Platforms/Mobile/Materials/MobileMaterials.INT.udn
--- a/Engine/Documentation/Source/Platforms/Mobile/Materials/MobileMaterials.INT.udn
+++ b/Engine/Documentation/Source/Platforms/Mobile/Materials/MobileMaterials.INT.udn
@@ -61,7 +61,7 @@ The following Material properties are to help optimize and or reduce the renderi
 | Use Full Precision | [INCLUDE:Platforms\Mobile\Materials\#UFP]|
 
 ## Troubleshooting Materials for Mobile
-Due to hardware differences between PC's and mobile devices you could run into issues with certain options working on the PC but not working when viewed on a mobile device. The following list contains some of the more common issues you could run into and ways to fix them.
+Due to hardware differences between PCs and mobile devices you could run into issues with certain options working on the PC but not working when viewed on a mobile device. The following list contains some of the more common issues you could run into and ways to fix them.
 | Issue | Solution |
 | ----- | -------- |
 | Default Material is displayed in the Mobile Previewer. | Open the [log file](Engine/Basics/DirectoryStructure) in **Engine\Saved** and search for error messages to know how to fix the Material. |
@@ -72,7 +72,7 @@ Due to hardware differences between PC's and mobile devices you could run into i
 
 <!--
 [EXCERPT:UFP]
-When enabled, Use Full Precision will make sure that the Material does not automatically truncate data when the Material used on a mobile device to help increase performance.
+When enabled, Use Full Precision will make sure that the Material does not automatically truncate data when the Material is used on a mobile device to help increase performance.
 
 [OBJECT:ComparisonSlider]
 	[PARAM:before]
