Addressing RichardH's rather excellent peer review comments.

#UE4doc

diff --git a/Engine/Documentation/Source/Programming/Development/InstalledBuildReference/InstalledBuildReference.INT.udn b/Engine/Documentation/Source/Programming/Development/InstalledBuildReference/InstalledBuildReference.INT.udn
--- a/Engine/Documentation/Source/Programming/Development/InstalledBuildReference/InstalledBuildReference.INT.udn
+++ b/Engine/Documentation/Source/Programming/Development/InstalledBuildReference/InstalledBuildReference.INT.udn
@@ -21,9 +21,9 @@ An **Installed Build** of Unreal Engine 4 (UE4) is a version that's distributed
 
 ![](InstalledBuilds.png)
 
-This version of the engine is shipped with each release of UE4, being fully featured, but focused by an Installed Build process 
-that gets users up and running with the engine. Additionally, the Unreal Editor is distributed as precompiled binaries with static 
-libraries that are built for targeted platforms in development and shipping configurations.
+This version of the engine is shipped with each release of UE4. Although the Installed Build is fully featured, its features are streamlined by an 
+Installed Build process that quickly gets users up and running with the engine. Additionally, the Unreal Editor is distributed as 
+precompiled binaries with a static library for each target platform in development and shipping configurations.
 
 ## About the Installed Build Process
 
@@ -144,7 +144,7 @@ Depending on your targeted platform and available Software Development Kit (SDK)
 
 #### IOS and TVOS
 
-If you want to include IOS and TVOS in your Installed Build, you'll have to have a Mac setup for remote building. Otherwise, append the following to your command line:
+If you want to include IOS and TVOS in your Installed Build, you'll have to have a Mac set up for remote building. Otherwise, append the following to your command line:
 `-set:WithIOS=false -setWithTVOS=false`.
 
 #### Playstation 4 and XboxOne
@@ -152,14 +152,14 @@ If you want to include IOS and TVOS in your Installed Build, you'll have to have
 If you have the required SDKs, you can add Playstation 4 (PS4) and XboxOne as target platforms (although, this is untested). To add PS4 and XboxOne 
 as target platforms, append `-set:WithPS4=true -set:WithXboxOne=true` to the command line.
 
-#### Amending the Derived Data Cache
+#### Excluding the Derived Data Cache
 
 Building a stand-alone [Derived Data Cache (DDC)](Engine/Basics/DerivedDataCache) for the Engine and Template content can be one of the slowest aspects of the build. 
 If you don't need a stand-alone DDC, you can skip this step by appending `-set:WithDDC=false` to the command line.
 
 #### Not Signing Executables
 
-Typically, you don't have to setup your machine to sign executables, as this can cause the Install Build process to fail while attempting to sign the executables being built. 
+Typically, you don't have to set up your machine to sign executables, as this can cause the Install Build process to fail while attempting to sign the executables being built. 
 To avoid this problem, append `-NoSign` to the command line.
 
 ### Including Additional Files
