Minor acronym fix.

#UE4doc

diff --git a/Engine/Documentation/Source/Programming/Development/InstalledBuildReference/InstalledBuildReference.INT.udn b/Engine/Documentation/Source/Programming/Development/InstalledBuildReference/InstalledBuildReference.INT.udn
--- a/Engine/Documentation/Source/Programming/Development/InstalledBuildReference/InstalledBuildReference.INT.udn
+++ b/Engine/Documentation/Source/Programming/Development/InstalledBuildReference/InstalledBuildReference.INT.udn
@@ -11,13 +11,13 @@ related:Programming/UnrealBuildSystem
 [TOC (start:1 end:4)]
 
 Developers are always looking for ways to efficiently deliver build solutions for a variety of hardware profiles and operating systems. 
-With the release of our BuildGraph scripting system, you'll be able to create custom sequences of UE4 build commands, harnessing the Installed 
-Build process that we've developed for your use. As you read through the following reference page, you'll learn about the Installed Build process, 
-including how to run scripts that make Installed Builds for targeted build solutions.
+With the release of our BuildGraph scripting system, you'll be able to create custom sequences of Unreal Engine 4 (UE4) build commands, 
+harnessing the Installed Build process that we've developed for your use. As you read through the following reference page, you'll learn 
+about the Installed Build process, including how to run scripts that make Installed Builds for targeted build solutions.
 
 ## What's an Installed Build?
 
-An **Installed Build** of Unreal Engine 4 (UE4) is a version that's distributed using the [Epic Games Launcher](GettingStarted/Installation).
+An **Installed Build** of UE4 is a version that's distributed using the [Epic Games Launcher](GettingStarted/Installation).
 
 ![](InstalledBuilds.png)
 
