#UE4 Docs: Applying peer review feedback to the doc.

diff --git a/Engine/Documentation/Source/Platforms/Mobile/Materials/MobileMaterials.INT.udn b/Engine/Documentation/Source/Platforms/Mobile/Materials/MobileMaterials.INT.udn
--- a/Engine/Documentation/Source/Platforms/Mobile/Materials/MobileMaterials.INT.udn
+++ b/Engine/Documentation/Source/Platforms/Mobile/Materials/MobileMaterials.INT.udn
@@ -3,6 +3,7 @@ Title:Materials for Mobile Platforms
 Crumbs:%ROOT%, Platforms, Platforms/Mobile
 Description:Setting up Materials for mobile games.
 Platform:Mobile
+Related: Platforms/Mobile/Previewer/
 Related: Platforms/Mobile/Performance
 Related: Platforms/Mobile/Textures
 Related: Engine/Rendering/Materials/CustomizedUVs
@@ -55,10 +56,10 @@ The following Material properties are to help optimize and or reduce the renderi
 
 | Property | Description |
 | -------- | ---------- |
-| bFullyRough | When enabled, saves a number of shader ALU instructions and one sampler.  This effectively disables Specular while maintaining overall reflective energy. |
-| High Quality Reflections |Enables high quality reflections in the forward renderer. Enabling this setting reduces the number of samplers available to the Material as two more samplers will be used for reflections cubemaps. |
-| bUseLightmapDirectionality | When disabled, lighting from lightmaps will be flat but cheaper. |
-| Use Full Precision | [INCLUDE:Platforms\Mobile\Materials\#UFP]|
+| **bFullyRough** | When enabled, saves a number of shader ALU instructions and one sampler.  This effectively disables Specular while maintaining overall reflective energy. |
+| **High Quality Reflections** |Enables high quality reflections in the forward renderer. Enabling this setting reduces the number of samplers available to the Material as two more samplers will be used for reflections cubemaps. |
+| **bUseLightmapDirectionality** | When disabled, lighting from lightmaps will be flat but cheaper. |
+| **Use Full Precision** | [INCLUDE:Platforms\Mobile\Materials\#UFP]|
 
 ## Troubleshooting Materials for Mobile
 Due to hardware differences between PCs and mobile devices you could run into issues with certain options working on the PC but not working when viewed on a mobile device. The following list contains some of the more common issues you could run into and ways to fix them.
@@ -66,13 +67,14 @@ Due to hardware differences between PCs and mobile devices you could run into is
 | ----- | -------- |
 | Default Material is displayed in the Mobile Previewer. | Open the [log file](Engine/Basics/DirectoryStructure) in **Engine\Saved** and search for error messages to know how to fix the Material. |
 | Material Animation with Time Expressions looks choppy but framerate is good elsewhere. | Make sure the Period property is enabled in the **Time** expression along with a Period. This will move the computation to the CPU and run at full-precision instead of half. |
-| Materials using World Coordinates computations or very small values are not being displayed correctly. | Try enabling the Use Full Precision option on the Material that is not displaying correctly on a mobile device. The precision level of Materials on Mobile devices is often truncated to one or two decimal places to help with performance. Enabling Use Full Precision will make sure that the Material's precision kept to the values that are set in the Material Editor  |
+| Materials using **World Coordinates** computations or very small values are not being displayed correctly. | Try enabling the Use Full Precision option on the Material that is not displaying correctly on a mobile device. The precision level of Materials on Mobile devices is often truncated to one or two decimal places to help with performance. Enabling **Use Full Precision** will make sure that the Material's precision kept to the values that are set in the Material Editor. |
+| Materials look fine in the editor but not on device. | To get a better idea of what Materials will look like when viewed on a mobile device, you can change the **Preview Rendering Level** of the UE4 editor to simulate the rendering capabilities of the device you are targeting.  Check out the [Mobile Previewer](Platforms/Mobile/Previewer/)  page for more information on how to do this. |
 
 
 
 <!--
 [EXCERPT:UFP]
-When enabled, Use Full Precision will make sure that the Material does not automatically truncate data when the Material is used on a mobile device to help increase performance.
+When enabled, Use Full Precision will make sure that the Material does not automatically truncate data when the Material is used on a mobile device. Keep in mind that this will make the Material **more** expensive to render as you telling UE4 to disreguard any automatic optmizations that help reduce the cost of the Material.
 
 [OBJECT:ComparisonSlider]
 	[PARAM:before]
