Adding new document for Twinblast Bust project under Engine/Documentation/Source/Resources/TwinblastBust

diff --git a/Engine/Documentation/Source/Resources/Showcases/TwinblastBust/TwinblastBust.INT.udn b/Engine/Documentation/Source/Resources/Showcases/TwinblastBust/TwinblastBust.INT.udn
new file mode 100644
--- /dev/null
+++ b/Engine/Documentation/Source/Resources/Showcases/TwinblastBust/TwinblastBust.INT.udn
@@ -0,0 +1,311 @@
+Availability: Public
+Title:Twinblast Bust
+Crumbs:%ROOT%, Resources
+Description:An overview of the Realistic Rendering included with UE4.
+Related: Engine/Rendering/Materials
+Related: Engine/Rendering/Materials/PhysicallyBased
+Related: Engine/Rendering/Materials/LayeredMaterials
+Version: 4.13
+
+[TOC(start:2 end:3)]
+
+The purpose of this showcase is to demonstrate the use of high-quality character shader techniques, as demonstrated on a character from Epic's MOBA game Paragon. To view this showcase, simply open the project and press Play In Editor to see the cinematic pan-around. 
+
+For more information on the tech used to produce this character, please refer to this [Unreal Engine Livestream - Tech & Techniques Behind Creating the Characters for Paragon](https://www.youtube.com/watch?v=toLJh5nnJA8).
+
+## Skin Shading
+
+The character's skin is created primarily through the use of UE4's [Subsurface Profile shading model](Engine/Rendering/Materials/LightingModels/SubSurfaceProfile/). 
+
+![](TwinblastSkinMat.png)(w:700 convert:false)
+
+Note the use of [Material Functions](Engine/Rendering/Materials/Functions/) to set up the base for the skin Material. This is done as a reusable authoring approach to creating Materials for Paragon. In this way, artists can standardize approaches to generating certain types of surfaces, knowing that a fundamental change to a single Function will update all instances where that Function is being used. 
+
+This skin setup utilizes three textures: 
+
+* A diffuse texture - *TwinBlast_UE4_Demo_Head_D*
+* A specular texture -  *CH2_exp04_merged_spec_f_FC_FINAL*
+* A scatter texture -  *TwinBlast_UE4_Demo_Head_BackScatter*
+
+The diffuse texture is responsible for supplying the Base Color for the Material, but also note that there is a Roughness map currently stored within the Alpha channel of the texture. This is a common technique used to minimize the amount of textures in use. The textures were originally produced from facial scans and then cleaned up by artists at Epic. 
+
+## Hair Shading
+
+Hair shading is handled by way of UE4's Hair shader model. This shader is physically based shading model based on research by [REDACTED_PERSON] and currently used by Weta Digital. The shader approximates light's reflection from the surface of the hair with anisotropic specularity, as well as its refraction through it and its scatter throughout multiple strands. 
+
+To utilize the UE4 Hair shader, set the **Shading Model** property of your Material to **Hair**.
+
+
+![](HairSM.png)
+
+### Hair and Specularity
+
+In the real world, hair tends to have multiple specular highlights: one representing the color of the light and another that is a mix of hair color and light color. For sake of this document, we will refer to these as primary and secondary specular highlights. The Hair shader approximates the same effects with highly realistic results.
+
+![](TwinblastHair_Specs_Diagram.png)(w:700 convert:false)
+[REGION:caption]1. The lighter hair represents the primary specular highlight. 2. The brighter red hair represents the secondary specular highlight.[/REGION]
+
+The approximation algorithm used in UE4's hair shader creates these effects in a similar manner to how they are made in the real world. As light strikes the surface of a hair follicle, it does not simply bounce off. Hair is translucent, allowing some light to pass *through* it, to potentially bounce around inside it, then exit. The hair shader approximates this with three possible paths for light to travel as it interacts with hair, as shown in the animated GIF diagram below:
+
+![](HairDiagram.gif)(w:700 convert:false)
+[REGION:caption]Cross section of single hair follicle showing how the hair shader approximates interaction with light, demonstrating the three primary types of paths light will take. See the table below for descriptions of each part of the process.[/REGION] 
+
+| Number | Description |
+|--------|-------------|
+|0. |The growth (root/tip) direction of the hair follicle. |
+|1. |Reflection only path, with light bouncing off the surface of the hair. This produces the primary specular highlight.|
+|2. |Transmission-transmission path, in which light passes into the hair follicle and back out the other side. This is how light scatters within a volume of hair.|
+|3. |Transmission-reflection-transmission path, in which light enters the hair follicle, reflects off the inside boundary of the surface, and then exits . This produces the secondary specular highlight.|
+
+As shown in the diagram above, a strand of hair is not a perfect cylinder or tube. In actuality hair appears more as a series of stacked cones. This means light bouncing off the surface of hair will be scattered differently than it would if hair were perfectly smooth. Further, because each strand of hair is generally pointing in a different direction, the specular highlight will not be unified, but will be independently placed based on the direction the hair is pointing. This is often referred to as *anisotropic* specularity, which is also supported by the UE4 hair shader.
+
+### Hair and Transparency
+
+The Hair Shader uses the Masked blend mode instead of translucency. However, Masked transparency results in a binary result - surfaces are either fully opaque or fully transparent. A noisy dither pattern is constantly moved across the surface, becoming more dense in areas that need to be more solid. Dither is used as a way to blend Masked transparency, but only works when TemporalAA is active. 
+
+[OBJECT:ComparisonSlider]
+[PARAM:before]
+![Without TemporalAA](Hair_AAOff.png)(w:700 convert:false)
+[/PARAM]
+[PARAM:after]
+![With TemporalAA](Hair_AAOn.png)(w:700 convert:false)
+[/PARAM]
+[/OBJECT]
+
+[REGION:note]Using an animated dither with TemporalAA requires several frames to resolve the blend. This can lead to some artifacting on hair while it is in motion. This is an expected side effect of the technique.[/REGION]
+
+#### Edge Masking
+
+![](EdgeMaskGraph.png)(w:700 convert:false)
+
+Though not specifically part of the shader itself, it is worth noting that the Material used on Twinblast's hair - as well as the other characters of Paragon - uses what is called an "Edge Mask" to fade the hair away as it becomes edge-on to the camera. Since hair is rendered as a series of planes - see [Hair and Geometry](#hairgeometry) below - it is likely that from certain camera angles, one would see the edge of each sheet, breaking the hair effect.
+
+To help alleviate this problem, the Material calculates the camera's vector against the vertex normal such that as the surface turns perpendicularly away from the camera, it will also fade away to complete transparency. This is also why many characters with thick heads of hair will often have a hair texture painted on the scalp, as seen in the image below.
+
+![](TwinblastScalp.png)(w:700 convert:false)
+
+[REGION:tip]It should also be pointed out that this effect is tied to a Quality Switch node, so that on lower-powered platforms it can easily be switched off to trade back some shader performance. As always, it's best to test on target platform.[/REGION]
+
+
+
+
+
+### Hair Authoring
+
+Creating hair using the techniques shown in this showcase requires a bit of understand of how Epic constructed this character's hair.
+
+#### Hair Geometry
+
+Hair geometry using UE4's hair shader is generally going to be constructed using a series of planar surfaces, which is a common approach in many game hair solutions. These can be authored in the DCC app of your choice.
+
+![](HairGeo.png)(w:700 convert:false)
+
+#### Hair Textures
+
+In this use of the UE4 Hair shader, the final result is driven by five primary textures: *Diffuse*, *Alpha*, *Root*, *Depth*, and a unique-per-strand *ID* texture. At Epic, these textures are typically generated using 3ds Max's Hair system, projecting the simulated hair results onto a piece of geometry. However, there are many available options for similar results. 
+
+
+| Texture | Name | Description |
+|---------|------|-------------|
+|![](Diffuse.png)(w:200 convert:false)|**Diffuse**|The Diffuse texture is precisely what one might expect in that supplies the primary diffuse - or *base* - color of the hair itself. It is sometimes desirable to leave this texture uncolored and drive its color via parameters, especially in those cases where a character's hair may take on many different colors.|
+|![](Alpha.png)(w:200 convert:false)|**Alpha**|The Alpha texture supplies the areas of opacity for the hair, hiding geometry outside the body of strands.|
+|![](Root.png)(w:200 convert:false)|**Root**|The Root texture provides a mask which is used to change the color of the hair from the root to the tip, useful for lightening or darkening hair along its length.|
+|![](Depth.png)(w:200 convert:false)|**Depth**|The Depth texture is used by Pixel Depth offset to give the illusion of pushing the hair deeper into the hair volume. It can also be used as a basis to change color or shading values for hair at different depths, such as reducing overall specularity as hair falls further toward the scalp.|
+|![](ID.png)(w:200 convert:false)|**Unique ID**| The Unique ID texture simply provides a unique 0 to 1 (black to white) value for each strand of hair on a given piece of hair geometry. This is used to provide subtle variation within the hair itself. |
+
+
+
+### Hair Shader Properties
+
+When using the Hair Shader you will notice some new properties available on the main Material node: **Scatter**, **Tangent**, and **Backlit**.
+
+![](HairProperties.png)
+
+[REGION:warning]At the time of this writing (4.13), the Backlit property only exists as a part of an early version of the shader. Connections to it are not used internally by the Hair shader and the property can therefore be ignored.[/REGION]
+
+#### Scatter
+
+One of the primary reasons we refer to the Hair shader as an *approximation*, rather than a *simulation*, is that it is not actually simulating every individual strand of hair, nor is it simulating perfectly accurate lighting behavior. In the real world, as light bounces off of and transmits through follicles of hair, it will often encounter *other* strands of hair, repeating the same process potentially many times in succession. At this time there is no computing system capable of accurately producing the effects of such a simulation in realtime.
+
+Nevertheless, the way in which light scatters through a body of hair is still crucial to how realistic that hair appears to be, in a game just as much as in the real world. To control this, the Hair shader provides the Scatter property, which replaces Metallic on your primary shader node and is limited to values between 0.0 and 0.1. Scatter controls how much light can pass into and scatter through the hair *as a whole surface*, rather than per follicle. In simpler terms, this controls how much light passes through the entire body of your character's hair. 
+
+An important point to note about Scatter is that it will tend to be higher for lighter colored hair and darker for darker hair. This follows the physcal rules of the natural world, as darker hair will tend to absorb more light. In practical terms, if you are trying to create a blonde character, you will find that changing the diffuse texture or color is not enough; you will also have to increase the Scatter value as well.
+
+
+[OBJECT:ComparisonSlider]
+[PARAM:before]
+![Scatter value of 0.0](LowScatter.png)(w:700 convert:false)
+[/PARAM]
+[PARAM:after]
+![Scatter value of 1.0](HighScatter.png)(w:700 convert:false)
+[/PARAM]
+[/OBJECT]
+[REGION:caption]In this example the Diffuse texture has not been modified from its default dark brown color. This illustrates how various tones of hair can be generated just by changing the Scatter value.[/REGION]
+
+#### Tangent
+
+The Tangent property replaces the Normal property on the Hair shader. Instead of a typical normal map, the hair shader calculates a tangent. The tangent is calculated as a vector that runs parallel to each strand of hair, pointed back toward the root. In Twinblast's hair shader, this vector is given a random offset in the Z-Axis between 0.3 and -0.3, using the Unique ID texture. This produces a vector which has a random direction within an arc, and aids in providing variation within the anisotropic specular highlights.
+
+![](HairTangent.png)(w:400 convert:false)
+[REGION:caption]In this image the yellow line represents the tangent along a strand of hair, pointing back toward the root.[/REGION]
+
+Tangents can be mapped in one of two ways: automatically or with the help of a flow map. The automatic method simply requires that that the textures for each sheet of hair be oriented such that the root is toward the top and the tips are toward the bottom. So long as your hair is short, and no sheets are bending around and twisting too much, this approach is adequate. This is the method used on Twinblast's hair.
+
+##### Flow Maps for Tangency
+
+The other method requires the creation of a flow map. This is useful if your character's hair is long and is bending and needs to appear to bend and curl more than the actual geometry does, or if the individual texture portions for the hair are oriented in different ways (not top-to-bottom as described above). The flow map will represent the direction the hair is moving in tangent space, or along the surface. Within the Twinblast Bust project, you can find an unused flow map named *T_SparrowHair_Flow*. Below is a comparison of that flowmap and the result of the specularity. 
+
+[OBJECT:ComparisonSlider]
+[PARAM:before]
+![Final Result](SparowFlowMap.png)(w:600 convert:false)
+[/PARAM]
+[PARAM:after]
+![Flow Map Texturing](SparowFlowMap2.png)(w:600 convert:false)
+[/PARAM]
+[/OBJECT]
+[REGION:caption]Here you can see how the flow map is textured along the surface of Sparrow's hair. Note that the flowmap is only in use on some of the hair sheets, not the entire body of hair. Notice how different values along the flow map subtly shift the specularity along the hair.[/REGION]
+
+
+
+
+#### Using Pixel Depth Offset with the Hair Shader
+
+Pixel Depth Offset (PDO) is not a property unique to the Hair shader. In layman's terms, PDO causes pixels to look recessed away from the camera, creating an artificial sense of depth along the surface. Since hair is made up of simple sheets of geometry, as explained in [Hair and Geometry](#hairgeometry) below, the use of PDO can provide a tangible sense of depth to the hair shape as a whole. It also breaks up the intersection point where a plane of hair geometry contacts the scalp, as shown below.
+
+[OBJECT:ComparisonSlider]
+[PARAM:before]
+![Without Pixel Depth Offset](PDO_off.png)(w:700 convert:false)
+[/PARAM]
+[PARAM:after]
+![With Pixel Depth Offset](PDO_On.png)(w:700 convert:false)
+[/PARAM]
+[/OBJECT]
+
+
+
+
+## Eye Shading
+
+![](EyeShading.png)(w:700 convert:false)
+
+Eyes are often referred to as "the windows to the soul." Few aspects of a character have as much influence on the viewer's perception of a character as a real person as that character's eyes. The Eye shader in UE4 is designed to bring realistic rendering of an eye surface, exposing artistic control over each of the biological parts of the eye itself.
+
+[REGION:warning]The Eye shader in UE4 is highly technical and has been developed such that there are **very strong dependencies** between the shader code, the Material, the shape of the geometry, and the geometry's UV layout. Building an eye from scratch using this shader is not recommended without extensive experience. To that end, we strongly recommend, for those creating realistic humanoid eyes, that you extract the eye geometry from this example and use the eye Material **as-is**, replacing the necessary textures to suit your purpose. [/REGION]
+
+### Eye Biology
+
+In order to fully understand the features of the eye it is best to know a little bit of eye biology. Each of the parts you see here are represented in this eye setup, and all can be modified in some way using the exposed parameters in the Material Instance applied to Twinblast's eyes (*MI_TwinBlast_EyeRefractive_Bust*).
+
+Here is a quick refresher:
+
+![](EyeDiagram.png)(w:700 convert:false)
+|Number|Name|Description|
+|------|----|-----------|
+|1|Sclera|The sclera is also known as the "white of the eye." This area of the eye is generally very wet and has a slight amount of tactile texture.|
+|2|Limbus|The limbus is the dark ring that is present between the iris and the sclera. It is more dominant in some eyes than others, and will tend to fade away when viewed from edge-on.|
+|3|Iris|The iris is the ring of color around the center of the eye. When one is said to have "green" eyes, it is because their iris is predominantly green. In a real eye, the iris is a ring of muscle-like fibers that expand and contract, allowing more light into - or closing light out of - the pupil. It is also worth noting that in the real world, the iris is actually more of a disc or cone shape, and does not bulge outward with the rest of the eye.|
+|4|Pupil|The pupil is the dark spot at the center of the eye. It is the hole through which light passes so it can be picked up by the rods and cones of the retina.|
+|5|Cornea|The cornea is the clear, fluid-filled dome that rests over the surface of the iris. |
+
+### Refraction within the Eye
+
+Because the eyeball is filled with fluid, it is given to refracting any light that passes through it. This can be seen in the real world when looking at an eye from multiple angles. The iris and pupil will be distorted by refraction as they are viewed through the cornea. Traditional approaches to solving this for games and film was to create two separate eye surfaces: one that supplied the sclera, iris, and pupil; and another surface on top that supplied the cornea and overall wetness of the eye. This allowed for refraction as the underlying surface was viewed through the wet layer shell. Such an approach was used in Epic's *A Boy and His Kite* tech demo on the boy's eyes. The diagram below shows a representation of the two surfaces.
+
+![](KiteBoyEyes.png)(w:700 convert:false)
+
+Using our new Eye shader, however, an even more realistic effect can be achieved with a *single surface*. The refraction of light through the surface of the cornea is now handled entirely within the shader, alleviating the need for underlying geometry. You can experiment with this yourself by opening the *MI_TwinBlast_EyeRefractive_Bust* Material Instance and adjusting the **Depth Scale** property.
+
+[OBJECT:ComparisonSlider]
+[PARAM:before]
+![Without Refraction (Depth Scale at 0.0)](EyeNotRefracting.png)(w:700 convert:false)
+[/PARAM]
+[PARAM:after]
+![With Refraction (Depth Scale at 1.2)](EyeRefracting.png)(w:700 convert:false)
+[/PARAM]
+[/OBJECT]
+
+
+### Eye Authoring
+
+Earlier in this section we provided a warning that there are strong interdependencies between the code of the Eye shader, the setup of the material, and the geometry of the eye mesh. The following is a high-level overview of what makes the eye setup so particular, which will be useful if you should try replacing the geometry with your own eye mesh for your characters.
+
+When building your own eyes, we **strongly** recommend that you start using the eye geometry and Material provided in this example. As such, most of the information provided here will be based upon that assumption.
+
+#### Shape and UV Layout
+
+Much like the human eye of the real world, the eye mesh is not modeled to be perfectly spherical. Rather, it has a distinct dome at the front to represent the cornea. This gives the eye almost an egg shape.
+
+![](EyeGeometry.png)(w:400 convert:false)
+[REGION:caption]In this image you can see the geometry of the eye viewed from the side in Maya. Note the raised area of the cornea at the front.[/REGION]
+
+The UV layout for the eye is also critical, as it is from this that the proportions of each eye texture are produced. Fortunately, the UV layout that works best with the Eye shader is relatively basic: a planar projection along the front axis of the eye. Inclusion of backfaces is not extremely important (since they'll be behind the eye anyway), which can be seen on the Twinblast example.
+ 
+ The UV layout will look like this:
+
+![](EyeUVs.png)(w:400 convert:false)
+[REGION:caption]Notice that the UV layout is centered on the 0-1 space and that the projection looks directly into the cornea. [/REGION]
+
+
+### Eye Textures
+
+Many of the textures you'll create for your eyes will be based on the UV layout created above, with the exception of the iris map. The list of primary editable textures includes:
+
+#### Sclera Map
+The sclera map controls the colors of the eye whites, as well as the vein and tissue color that will be present at the edges of the eye. If you wanted your character to be able to have bloodshot eyes, for example, this is the texture you would modify.
+
+![](ScleraMap.png)(w:200 convert:false) 
+[REGION:caption]*TwinBlast_UE4_Demo_Eye_Sclera_Color*[/REGION]  
+
+
+
+#### Mid Plane Displacement Map
+
+The mid plane displacement map is used to lock down a plane that cuts through the center of the eye, and is used as a basis for offsetting the depth of the iris.
+
+![](DispMap.png)(w:200 convert:false) 
+[REGION:caption]*EyeMidPlaneDisplacement_Example*[/REGION]
+
+#### Normal Map
+
+The normal map does the same job as normal maps in all shaders - it supplies tactile texture variation. In the case of the Twinblast eye setup, it provides tiny lumps and bumps along the wet surface of the eye.
+
+![](WetNormal.png)(w:200 convert:false) 
+[REGION:caption]*T_Eye_Wet_Normal'*[/REGION]
+
+#### Tangent Map
+
+The tangent map is used to control the direction the surface is flowing in tangent space. In the case of Twinblast's eye setup, this is used to accentuate the change in direction where the cornea meets the sclera. Note that this out put is **not** fed directly into the Normal Map input. Rather, it is connected to a **Tangent Output** node, which must be created separately, as seen in the master material (*M_HairSheet_Master2*).
+
+![](Normal.png)(w:200 convert:false) 
+[REGION:caption]*EYE_NORMALS.EYE_NORMALS* [/REGION] 
+
+
+![](TangentOutput.png)(w:200 convert:false) 
+[REGION:caption]A Tangent Input node is used as a way to supply the base tangent map for the eye bulge, rather than trying to blend it with a typical normal map. [/REGION] 
+
+#### Iris Map
+
+The iris texture is unique in that it does not match the UV layout. Instead, the texure fills the entirety of the resolution. UVs and alpha masks are employed to control the overall iris size, as well as the size of the pupil. If making your own iris maps, try to keep the pupil size reasonably close to the proportions in the supplied texture.
+
+![](IrisTexture.png)
+[REGION:caption]*TwinBlast_UE4_Demo_Eye_Iris'* [/REGION] 
+
+
+### Eye Ambient Occlusion
+
+The final part of the eye setup for Twinblast is the use of a translucent helper surface to to simulate ambient occlusion around the eye. In the real world - wherever your eye comes in contact with other tissues such as the eyelid or the tear duct - you will see very soft shadowing. The softness is due to the fact that the eye itself is translucent; light will scatter through it. To help simulate this realistically, we have placed a sheet of geometry that covers the eye and provides approximated ambient occlusion, creating a richer degree of realism across the eye.
+
+The Material applied to this sheet is named *MI_TwinBlast_Eye_Occlusion_Bust* and can be found in Element 3 of the Twinblast model.
+
+![](AOSheetMaya.png)(w:700 convert:false) 
+[REGION:caption]This image shows the ambient occlusion sheet geometry in Maya, without the geometry of the face. The border edge of the sheet will perfectly fill the opening of the eyelid. [/REGION]
+
+
+[OBJECT:ComparisonSlider]
+[PARAM:before]
+![Without AO Sheet](EyeAO_off.png)(w:700 convert:false)
+[/PARAM]
+[PARAM:after]
+![With AO Sheet](EyeAO_on.png)(w:700 convert:false)
+[/PARAM]
+[/OBJECT]
