Making some changes to the Twinblast Bust doc for clarity and continuity

diff --git a/Engine/Documentation/Source/Resources/Showcases/TwinblastBust/TwinblastBust.INT.udn b/Engine/Documentation/Source/Resources/Showcases/TwinblastBust/TwinblastBust.INT.udn
--- a/Engine/Documentation/Source/Resources/Showcases/TwinblastBust/TwinblastBust.INT.udn
+++ b/Engine/Documentation/Source/Resources/Showcases/TwinblastBust/TwinblastBust.INT.udn
@@ -21,17 +21,22 @@ The character's skin is created primarily through the use of UE4's [Subsurface P
 
 Note the use of [Material Functions](Engine/Rendering/Materials/Functions/) to set up the base for the skin Material. This is done as a reusable authoring approach to creating Materials for Paragon. In this way, artists can standardize approaches to generating certain types of surfaces, knowing that a fundamental change to a single Function will update all instances where that Function is being used. 
 
-This skin setup utilizes three textures: 
+### Skin Shader Textures 
 
-* A diffuse texture - *TwinBlast_UE4_Demo_Head_D*
-* A specular texture -  *CH2_exp04_merged_spec_f_FC_FINAL*
-* A scatter texture -  *TwinBlast_UE4_Demo_Head_BackScatter*
+The textures used on Twinblast's skin are all at 4K resolution and were originally produced from facial scans of the actor. The textures were then cleaned up and tweaked by artists at Epic. This skin setup utilizes five total texture maps: diffuse, roughness, specularity, scatter, and normal.
+
+| Texture | Name | Description |
+|---------|------|-------------|
+|![](SkinDiffuse.png)(w:200 convert:false)|**Diffuse**|The diffuse map supplies the Base Color for the Material. At 4K, you can see the tiny capillaries just underneath the surface. Any darkening of wrinkles will accentuate the texture supplied by the normal map. (*TwinBlast_UE4_Demo_Head_D*)|
+|![](SkinRough.png)(w:200 convert:false)|**Roughness**|The roughness map is stored within the alpha channel of the diffuse texture. This is a common technique used to minimize the amount of textures in use. Note that the roughness increases wihtin pores and wrinkles. This causes those areas to appear less shiny, accentuating the look of depth provided by the diffuse and normal map. Also note that the hair texture on the head is pushed out to fully rough (1.0) this prevents any stray specular highlights from the scalp, which will give a much more significant sense of depth to the hair. (*TwinBlast_UE4_Demo_Head_D*)|
+|![](SkinSpec.png)(w:200 convert:false)|**Specular**|The specular map scales the amount of specular highlight visible across the surface of the skin. It is important to note that the default value for specularity is 0.5. This map boosts specularity at areas where the skin may be stretched a bit tighter, and dampens it in areas where we don't want to see reflection, such as the center of pores and within wrinkles. (*CH2_exp04_merged_spec_f_FC_FINAL*)|
+|![](SkinScatter.png)(w:200 convert:false)|**Scatter**|The scatter map controls how much light is going to be scattered through the surface of the skin. Areas that are dark will exhibit very little scatter, such as the cheeks, while lighter areas will have higher amounts of perceived scatter, such as the nose and ears. The color of the scatter is managed by the [Subsurface Profile asset.](Engine/Rendering/Materials/LightingModels/SubSurfaceProfile/) (*TwinBlast_UE4_Demo_Head_BackScatter*)|
+|![](SkinNormal.png)(w:200 convert:false)|**Normal**|The normal map works as expected in typical material setups, supplying the tactile texture to the surface by perturbing pixel normals. There is nothing out of the ordinary in its setup for this example.|
 
-The diffuse texture is responsible for supplying the Base Color for the Material, but also note that there is a Roughness map currently stored within the Alpha channel of the texture. This is a common technique used to minimize the amount of textures in use. The textures were originally produced from facial scans and then cleaned up by artists at Epic. 
 
 ## Hair Shading
 
-Hair shading is handled by way of UE4's Hair shader model. This shader is physically based shading model based on research by [REDACTED_PERSON] and currently used by Weta Digital. The shader approximates light's reflection from the surface of the hair with anisotropic specularity, as well as its refraction through it and its scatter throughout multiple strands. 
+Hair shading is handled by way of UE4's Hair shader model. This shader is physically based shading model based on research by [[REDACTED_PERSON], [REDACTED_PERSON] and [REDACTED_PERSON]]([REDACTED_PROFILE_URL]) and currently used by [Weta Digital](http://www.wetafx.co.nz/research). The shader approximates light's reflection from the surface of the hair with anisotropic specularity, as well as its refraction through it and its scatter throughout multiple strands. 
 
 To utilize the UE4 Hair shader, set the **Shading Model** property of your Material to **Hair**.
 
@@ -107,11 +112,11 @@ In this use of the UE4 Hair shader, the final result is driven by five primary t
 
 | Texture | Name | Description |
 |---------|------|-------------|
-|![](Diffuse.png)(w:200 convert:false)|**Diffuse**|The Diffuse texture is precisely what one might expect in that supplies the primary diffuse - or *base* - color of the hair itself. It is sometimes desirable to leave this texture uncolored and drive its color via parameters, especially in those cases where a character's hair may take on many different colors.|
-|![](Alpha.png)(w:200 convert:false)|**Alpha**|The Alpha texture supplies the areas of opacity for the hair, hiding geometry outside the body of strands.|
-|![](Root.png)(w:200 convert:false)|**Root**|The Root texture provides a mask which is used to change the color of the hair from the root to the tip, useful for lightening or darkening hair along its length.|
-|![](Depth.png)(w:200 convert:false)|**Depth**|The Depth texture is used by Pixel Depth offset to give the illusion of pushing the hair deeper into the hair volume. It can also be used as a basis to change color or shading values for hair at different depths, such as reducing overall specularity as hair falls further toward the scalp.|
-|![](ID.png)(w:200 convert:false)|**Unique ID**| The Unique ID texture simply provides a unique 0 to 1 (black to white) value for each strand of hair on a given piece of hair geometry. This is used to provide subtle variation within the hair itself. |
+|![](Diffuse.png)(w:200 convert:false)|**Diffuse**|The Diffuse texture is precisely what one might expect in that supplies the primary diffuse - or *base* - color of the hair itself. It is sometimes desirable to leave this texture uncolored and drive its color via parameters, especially in those cases where a character's hair may take on many different colors. (*TwinBlast_UE4_Demo_Hair_D*)|
+|![](Alpha.png)(w:200 convert:false)|**Alpha**|The Alpha texture supplies the areas of opacity for the hair, hiding geometry outside the body of strands. (*TwinBlast_UE4_Demo_Hair_A*)|
+|![](Root.png)(w:200 convert:false)|**Root**|The Root texture provides a mask which is used to change the color of the hair from the root to the tip, useful for lightening or darkening hair along its length. (*TwinBlast_UE4_Demo_Hair_Roots*)|
+|![](Depth.png)(w:200 convert:false)|**Depth**|The Depth texture is used by Pixel Depth offset to give the illusion of pushing the hair deeper into the hair volume. It can also be used as a basis to change color or shading values for hair at different depths, such as reducing overall specularity as hair falls further toward the scalp. (*TwinBlast_UE4_Demo_Hair_Depth*)|
+|![](ID.png)(w:200 convert:false)|**Unique ID**| The Unique ID texture simply provides a unique 0 to 1 (black to white) value for each strand of hair on a given piece of hair geometry. This is used to provide subtle variation within the hair itself. (*TwinBlast_UE4_Demo_Hair_ID*)|
 
 
 
@@ -250,45 +255,17 @@ The UV layout for the eye is also critical, as it is from this that the proporti
 
 Many of the textures you'll create for your eyes will be based on the UV layout created above, with the exception of the iris map. The list of primary editable textures includes:
 
-#### Sclera Map
-The sclera map controls the colors of the eye whites, as well as the vein and tissue color that will be present at the edges of the eye. If you wanted your character to be able to have bloodshot eyes, for example, this is the texture you would modify.
-
-![](ScleraMap.png)(w:200 convert:false) 
-[REGION:caption]*TwinBlast_UE4_Demo_Eye_Sclera_Color*[/REGION]  
-
-
-
-#### Mid Plane Displacement Map
-
-The mid plane displacement map is used to lock down a plane that cuts through the center of the eye, and is used as a basis for offsetting the depth of the iris.
-
-![](DispMap.png)(w:200 convert:false) 
-[REGION:caption]*EyeMidPlaneDisplacement_Example*[/REGION]
-
-#### Normal Map
-
-The normal map does the same job as normal maps in all shaders - it supplies tactile texture variation. In the case of the Twinblast eye setup, it provides tiny lumps and bumps along the wet surface of the eye.
-
-![](WetNormal.png)(w:200 convert:false) 
-[REGION:caption]*T_Eye_Wet_Normal'*[/REGION]
-
-#### Tangent Map
-
-The tangent map is used to control the direction the surface is flowing in tangent space. In the case of Twinblast's eye setup, this is used to accentuate the change in direction where the cornea meets the sclera. Note that this out put is **not** fed directly into the Normal Map input. Rather, it is connected to a **Tangent Output** node, which must be created separately, as seen in the master material (*M_HairSheet_Master2*).
-
-![](Normal.png)(w:200 convert:false) 
-[REGION:caption]*EYE_NORMALS.EYE_NORMALS* [/REGION] 
-
-
-![](TangentOutput.png)(w:200 convert:false) 
-[REGION:caption]A Tangent Input node is used as a way to supply the base tangent map for the eye bulge, rather than trying to blend it with a typical normal map. [/REGION] 
-
-#### Iris Map
+| Texture | Name | Description |
+|---------|------|-------------|
+|![](ScleraMap.png)(w:200 convert:false) |**Sclera Map**|The sclera map controls the colors of the eye whites, as well as the vein and tissue color that will be present at the edges of the eye. If you wanted your character to be able to have bloodshot eyes, for example, this is the texture you would modify. (*TwinBlast_UE4_Demo_Eye_Sclera_Color*)|
+|![](DispMap.png)(w:200 convert:false) |**Mid Plane Displacement Map**|The mid plane displacement map is used to lock down a plane that cuts through the center of the eye, and is used as a basis for offsetting the depth of the iris. (*EyeMidPlaneDisplacement_Example*)|
+|![](WetNormal.png)(w:200 convert:false) |**Normal Map**|The normal map does the same job as normal maps in all shaders - it supplies tactile texture variation. In the case of the Twinblast eye setup, it provides tiny lumps and bumps along the wet surface of the eye. (*T_Eye_Wet_Normal*)|
+|![](Normal.png)(w:200 convert:false) |**Tangent Map**|The tangent map is used to control the direction the surface is flowing in tangent space. In the case of Twinblast's eye setup, this is used to accentuate the change in direction where the cornea meets the sclera. (*EYE_NORMALS.EYE_NORMALS*)|
+|![](IrisTexture.png)(w:200 convert:false)|**Iris Map**|The iris texture is unique in that it does not match the UV layout. Instead, the texure fills the entirety of the resolution. UVs and alpha masks are employed to control the overall iris size, as well as the size of the pupil. If making your own iris maps, try to keep the pupil size reasonably close to the proportions in the supplied texture. (*TwinBlast_UE4_Demo_Eye_Iris*)|
 
-The iris texture is unique in that it does not match the UV layout. Instead, the texure fills the entirety of the resolution. UVs and alpha masks are employed to control the overall iris size, as well as the size of the pupil. If making your own iris maps, try to keep the pupil size reasonably close to the proportions in the supplied texture.
+[REGION:warning]Note that the Tangent Map output is **not** fed directly into the Normal Map input. Rather, it is connected to a **Tangent Output** node, which must be created separately, as seen in the master material **M_HairSheet_Master2**. [/REGION]
+![](TangentOutput.png)
 
-![](IrisTexture.png)
-[REGION:caption]*TwinBlast_UE4_Demo_Eye_Iris'* [/REGION] 
 
 
 ### Eye Ambient Occlusion
@@ -309,3 +286,7 @@ The Material applied to this sheet is named *MI_TwinBlast_Eye_Occlusion_Bust* an
 ![With AO Sheet](EyeAO_on.png)(w:700 convert:false)
 [/PARAM]
 [/OBJECT]
+
+## Special Thanks
+
+Special thanks to [[REDACTED_PERSON], [REDACTED_PERSON] and [REDACTED_PERSON]]([REDACTED_PROFILE_URL]) for releasing their work [*The importance of sampling for physically-based hair fiber models*](http://dl.acm.org/authorize?6967840), upon which some of the work of UE4's hair shader is based.
