Minor changes and typos on the Twinblast doc. Setting availability to Docs.

diff --git a/Engine/Documentation/Source/Resources/Showcases/TwinblastBust/TwinblastBust.INT.udn b/Engine/Documentation/Source/Resources/Showcases/TwinblastBust/TwinblastBust.INT.udn
--- a/Engine/Documentation/Source/Resources/Showcases/TwinblastBust/TwinblastBust.INT.udn
+++ b/Engine/Documentation/Source/Resources/Showcases/TwinblastBust/TwinblastBust.INT.udn
@@ -1,7 +1,7 @@
-Availability: Public
+Availability: Docs
 Title:Twinblast Bust
 Crumbs:%ROOT%, Resources
-Description:An overview of the Realistic Rendering included with UE4.
+Description:An overview of the Twinblast Bust showcase.
 Related: Engine/Rendering/Materials
 Related: Engine/Rendering/Materials/PhysicallyBased
 Related: Engine/Rendering/Materials/LayeredMaterials
@@ -31,7 +31,7 @@ The textures used on Twinblast's skin are all at 4K resolution and were original
 |![](SkinRough.png)(w:200 convert:false)|**Roughness**|The roughness map is stored within the alpha channel of the diffuse texture. This is a common technique used to minimize the amount of textures in use. Note that the roughness increases wihtin pores and wrinkles. This causes those areas to appear less shiny, accentuating the look of depth provided by the diffuse and normal map. Also note that the hair texture on the head is pushed out to fully rough (1.0) this prevents any stray specular highlights from the scalp, which will give a much more significant sense of depth to the hair. (*TwinBlast_UE4_Demo_Head_D*)|
 |![](SkinSpec.png)(w:200 convert:false)|**Specular**|The specular map scales the amount of specular highlight visible across the surface of the skin. It is important to note that the default value for specularity is 0.5. This map boosts specularity at areas where the skin may be stretched a bit tighter, and dampens it in areas where we don't want to see reflection, such as the center of pores and within wrinkles. (*CH2_exp04_merged_spec_f_FC_FINAL*)|
 |![](SkinScatter.png)(w:200 convert:false)|**Scatter**|The scatter map controls how much light is going to be scattered through the surface of the skin. Areas that are dark will exhibit very little scatter, such as the cheeks, while lighter areas will have higher amounts of perceived scatter, such as the nose and ears. The color of the scatter is managed by the [Subsurface Profile asset.](Engine/Rendering/Materials/LightingModels/SubSurfaceProfile/) (*TwinBlast_UE4_Demo_Head_BackScatter*)|
-|![](SkinNormal.png)(w:200 convert:false)|**Normal**|The normal map works as expected in typical material setups, supplying the tactile texture to the surface by perturbing pixel normals. There is nothing out of the ordinary in its setup for this example.|
+|![](SkinNormal.png)(w:200 convert:false)|**Normal**|The normal map works as expected in typical material setups, supplying the tactile texture to the surface by perturbing pixel normals. There is nothing out of the ordinary in its setup for this example. (*TwinBlast_UE4_Demo_Head__normals*)|
 
 
 ## Hair Shading
@@ -151,7 +151,7 @@ An important point to note about Scatter is that it will tend to be higher for l
 
 The Tangent property replaces the Normal property on the Hair shader. Instead of a typical normal map, the hair shader calculates a tangent. The tangent is calculated as a vector that runs parallel to each strand of hair, pointed back toward the root. In Twinblast's hair shader, this vector is given a random offset in the Z-Axis between 0.3 and -0.3, using the Unique ID texture. This produces a vector which has a random direction within an arc, and aids in providing variation within the anisotropic specular highlights.
 
-![](HairTangent.png)(w:400 convert:false)
+![](HairTangent.png)(w:700 convert:false)
 [REGION:caption]In this image the yellow line represents the tangent along a strand of hair, pointing back toward the root.[/REGION]
 
 Tangents can be mapped in one of two ways: automatically or with the help of a flow map. The automatic method simply requires that that the textures for each sheet of hair be oriented such that the root is toward the top and the tips are toward the bottom. So long as your hair is short, and no sheets are bending around and twisting too much, this approach is adequate. This is the method used on Twinblast's hair.
