Adding a header image to Twinblast Bust doc

diff --git a/Engine/Documentation/Source/Resources/Showcases/TwinblastBust/TwinblastBust.INT.udn b/Engine/Documentation/Source/Resources/Showcases/TwinblastBust/TwinblastBust.INT.udn
--- a/Engine/Documentation/Source/Resources/Showcases/TwinblastBust/TwinblastBust.INT.udn
+++ b/Engine/Documentation/Source/Resources/Showcases/TwinblastBust/TwinblastBust.INT.udn
@@ -9,13 +9,16 @@ Version: 4.13
 
 [TOC(start:2 end:3)]
 
+![](TwinblastHeader.png)(w:700 convert:false)
+
+
 The purpose of this showcase is to demonstrate the use of high-quality character shader techniques, as demonstrated on a character from Epic's MOBA game Paragon. To view this showcase, simply open the project and press Play In Editor to see the cinematic pan-around. 
 
 For more information on the tech used to produce this character, please refer to this [Unreal Engine Livestream - Tech & Techniques Behind Creating the Characters for Paragon](https://www.youtube.com/watch?v=toLJh5nnJA8).
 
 ## Skin Shading
 
-The character's skin is created primarily through the use of UE4's [Subsurface Profile shading model](Engine/Rendering/Materials/LightingModels/SubSurfaceProfile/). 
+The character's skin is created through the use of UE4's [Subsurface Profile shading model](Engine/Rendering/Materials/LightingModels/SubSurfaceProfile/). 
 
 ![](TwinblastSkinMat.png)(w:700 convert:false)
 
