#UE4 Docs: Removed the mention of level name in the 3rd step of Enabling Vulkan Preview Rendering in Editor as the users level name could be different and is not important to getting this feature to work.

diff --git a/Engine/Documentation/Source/Platforms/Android/VulkanMobileRenderer/VulkanMobileRender.INT.udn b/Engine/Documentation/Source/Platforms/Android/VulkanMobileRenderer/VulkanMobileRender.INT.udn
--- a/Engine/Documentation/Source/Platforms/Android/VulkanMobileRenderer/VulkanMobileRender.INT.udn
+++ b/Engine/Documentation/Source/Platforms/Android/VulkanMobileRenderer/VulkanMobileRender.INT.udn
@@ -157,7 +157,7 @@ You can enable the ability to use the Vulkan render while inside of the UE4 Edit
 	After doing this, the Editor will need to recompile the entire shader cache to include the required Vulkan options. Depending on the size of your project and power of your development computer, this process could take anywhere from a few minutes to an hour or more to complete.
 	[/region]
 
-1. If you see the following two items listed in the lower right-hand corner of the Level preview **Feature Level: ES3_1** and **Level: Vulkan_Rendering(Persistent)** then everything is setup correctly, and you are now seeing what your project will look like when viewed on a device that supports the Vulkan API.
+1. In the lower right-hand conner of the UE4 Editor window you should see the text, **Feature Level: ES3_1**. If you see this text then everything is setup correctly and you are now seeing what your project will look like when viewed on a device that supports the Vulkan API.
 
 	[REGION:lightbox]
 	[![](VMP_04.png)(w:447 h:308)](VMP_04.png)
