Twinblast doc typos. Added note disclaimer at beginning about it only concerning skin, hair, and eyes.

diff --git a/Engine/Documentation/Source/Resources/Showcases/TwinblastBust/TwinblastBust.INT.udn b/Engine/Documentation/Source/Resources/Showcases/TwinblastBust/TwinblastBust.INT.udn
--- a/Engine/Documentation/Source/Resources/Showcases/TwinblastBust/TwinblastBust.INT.udn
+++ b/Engine/Documentation/Source/Resources/Showcases/TwinblastBust/TwinblastBust.INT.udn
@@ -14,6 +14,8 @@ Version: 4.13
 
 The purpose of this showcase is to demonstrate the use of high-quality character shader techniques, as demonstrated on a character from Epic's MOBA game Paragon. To view this showcase, simply open the project and press Play In Editor to see the cinematic pan-around. 
 
+[REGION:note]Note that this example only focuses on the character's skin, hair, and eyes.[/REGION]
+
 For more information on the tech used to produce this character, please refer to this [Unreal Engine Livestream - Tech & Techniques Behind Creating the Characters for Paragon](https://www.youtube.com/watch?v=toLJh5nnJA8).
 
 ## Skin Shading
@@ -100,7 +102,7 @@ To help alleviate this problem, the Material calculates the camera's vector agai
 
 ### Hair Authoring
 
-Creating hair using the techniques shown in this showcase requires a bit of understand of how Epic constructed this character's hair.
+Creating hair using the techniques shown in this showcase requires a bit of understanding of how Epic constructed this example.
 
 #### Hair Geometry
 
