Twinblast bust documentation: Accepting incoming edits from peer review per RichardH

diff --git a/Engine/Documentation/Source/Resources/Showcases/TwinblastBust/TwinblastBust.INT.udn b/Engine/Documentation/Source/Resources/Showcases/TwinblastBust/TwinblastBust.INT.udn
--- a/Engine/Documentation/Source/Resources/Showcases/TwinblastBust/TwinblastBust.INT.udn
+++ b/Engine/Documentation/Source/Resources/Showcases/TwinblastBust/TwinblastBust.INT.udn
@@ -14,8 +14,6 @@ Version: 4.13
 
 The purpose of this showcase is to demonstrate the use of high-quality character shader techniques, as demonstrated on a character from Epic's MOBA game Paragon. To view this showcase, simply open the project and press Play In Editor to see the cinematic pan-around. 
 
-[REGION:note]Note that this example only focuses on the character's skin, hair, and eyes.[/REGION]
-
 For more information on the tech used to produce this character, please refer to this [Unreal Engine Livestream - Tech & Techniques Behind Creating the Characters for Paragon](https://www.youtube.com/watch?v=toLJh5nnJA8).
 
 ## Skin Shading
@@ -62,16 +60,16 @@ The approximation algorithm used in UE4's hair shader creates these effects in a
 
 | Number | Description |
 |--------|-------------|
-|0. |The growth (root/tip) direction of the hair follicle. |
-|1. |Reflection only path, with light bouncing off the surface of the hair. This produces the primary specular highlight.|
+|0. |The growth (root to tip) direction of the hair follicle. |
+|1. |Reflection-only path, with light bouncing off the surface of the hair. This produces the primary specular highlight.|
 |2. |Transmission-transmission path, in which light passes into the hair follicle and back out the other side. This is how light scatters within a volume of hair.|
 |3. |Transmission-reflection-transmission path, in which light enters the hair follicle, reflects off the inside boundary of the surface, and then exits . This produces the secondary specular highlight.|
 
-As shown in the diagram above, a strand of hair is not a perfect cylinder or tube. In actuality hair appears more as a series of stacked cones. This means light bouncing off the surface of hair will be scattered differently than it would if hair were perfectly smooth. Further, because each strand of hair is generally pointing in a different direction, the specular highlight will not be unified, but will be independently placed based on the direction the hair is pointing. This is often referred to as *anisotropic* specularity, which is also supported by the UE4 hair shader.
+As shown in the diagram above, a strand of hair is not a perfect cylinder or tube. In actuality, hair appears more as a series of stacked cones. This means light bouncing off the surface of hair will be scattered differently than it would if hair were perfectly smooth. Further, because each strand of hair is generally pointing in a different direction, the specular highlight will not be unified, but will be independently placed based on the direction the hair is pointing. This is often referred to as *anisotropic* specularity, which is also supported by the UE4 hair shader.
 
 ### Hair and Transparency
 
-The Hair Shader uses the Masked blend mode instead of translucency. However, Masked transparency results in a binary result - surfaces are either fully opaque or fully transparent. A noisy dither pattern is constantly moved across the surface, becoming more dense in areas that need to be more solid. Dither is used as a way to blend Masked transparency, but only works when TemporalAA is active. 
+The Hair Shader uses the Masked blend mode instead of translucency. However, Masked transparency produces a binary result - surfaces are either fully opaque or fully transparent. A noisy dither pattern is constantly moved across the surface, becoming more dense in areas that need to be more solid. Dither is used as a way to blend Masked transparency, but only works when TemporalAA is active. 
 
 [OBJECT:ComparisonSlider]
 [PARAM:before]
@@ -102,11 +100,11 @@ To help alleviate this problem, the Material calculates the camera's vector agai
 
 ### Hair Authoring
 
-Creating hair using the techniques shown in this showcase requires a bit of understanding of how Epic constructed this example.
+Creating hair using the techniques shown in this showcase requires a bit of understanding of how Epic constructed this character's hair.
 
 #### Hair Geometry
 
-Hair geometry using UE4's hair shader is generally going to be constructed using a series of planar surfaces, which is a common approach in many game hair solutions. These can be authored in the DCC app of your choice.
+Hair geometry using UE4's hair shader is generally going to be constructed using a series of planar surfaces, which is a common approach in many real-time hair solutions. These can be authored in the DCC app of your choice.
 
 ![](HairGeo.png)(w:700 convert:false)
 
@@ -117,7 +115,7 @@ In this use of the UE4 Hair shader, the final result is driven by five primary t
 
 | Texture | Name | Description |
 |---------|------|-------------|
-|![](Diffuse.png)(w:200 convert:false)|**Diffuse**|The Diffuse texture is precisely what one might expect in that supplies the primary diffuse - or *base* - color of the hair itself. It is sometimes desirable to leave this texture uncolored and drive its color via parameters, especially in those cases where a character's hair may take on many different colors. (*TwinBlast_UE4_Demo_Hair_D*)|
+|![](Diffuse.png)(w:200 convert:false)|**Diffuse**|The Diffuse texture is precisely what one might expect in that it supplies the primary diffuse - or *base* - color of the hair itself. It is sometimes desirable to leave this texture uncolored and drive its color via parameters, especially in those cases where a character's hair may take on many different colors. (*TwinBlast_UE4_Demo_Hair_D*)|
 |![](Alpha.png)(w:200 convert:false)|**Alpha**|The Alpha texture supplies the areas of opacity for the hair, hiding geometry outside the body of strands. (*TwinBlast_UE4_Demo_Hair_A*)|
 |![](Root.png)(w:200 convert:false)|**Root**|The Root texture provides a mask which is used to change the color of the hair from the root to the tip, useful for lightening or darkening hair along its length. (*TwinBlast_UE4_Demo_Hair_Roots*)|
 |![](Depth.png)(w:200 convert:false)|**Depth**|The Depth texture is used by Pixel Depth offset to give the illusion of pushing the hair deeper into the hair volume. It can also be used as a basis to change color or shading values for hair at different depths, such as reducing overall specularity as hair falls further toward the scalp. (*TwinBlast_UE4_Demo_Hair_Depth*)|
@@ -204,7 +202,7 @@ Eyes are often referred to as "the windows to the soul." Few aspects of a charac
 
 ### Eye Biology
 
-In order to fully understand the features of the eye it is best to know a little bit of eye biology. Each of the parts you see here are represented in this eye setup, and all can be modified in some way using the exposed parameters in the Material Instance applied to Twinblast's eyes (*MI_TwinBlast_EyeRefractive_Bust*).
+In order to fully understand the features of the eye, it is best to know a little bit of eye biology. Each of the parts you see here are represented in this eye setup, and all can be modified in some way using the exposed parameters in the Material Instance applied to Twinblast's eyes (*MI_TwinBlast_EyeRefractive_Bust*).
 
 Here is a quick refresher:
 
