#UE4 Docs: Adding a missing step to this doc.

diff --git a/Engine/Documentation/Source/Platforms/Android/VulkanMobileRenderer/VulkanMobileRender.INT.udn b/Engine/Documentation/Source/Platforms/Android/VulkanMobileRenderer/VulkanMobileRender.INT.udn
--- a/Engine/Documentation/Source/Platforms/Android/VulkanMobileRenderer/VulkanMobileRender.INT.udn
+++ b/Engine/Documentation/Source/Platforms/Android/VulkanMobileRenderer/VulkanMobileRender.INT.udn
@@ -51,6 +51,10 @@ To build a UE4 project that has support for the Vulkan API you will need to do t
 	[![](VMP_LunarG_00.png)(w:469 h:577)](VMP_LunarG_00.png)
 	[/REGION]
 
+	[REGION:caption]
+	Click for full image.
+	[/REGION]
+
 	[region:note] 
 	You do not need to sign in or create an account to download the Vulkan SDK.
 	[/region]
@@ -61,6 +65,10 @@ To build a UE4 project that has support for the Vulkan API you will need to do t
 	[![](VMP_LunarG_01.png)(w:463 h:224)](VMP_LunarG_01.png)
 	[/REGION]
 
+	[REGION:caption]
+	Click for full image.
+	[/REGION]
+
 1. Once the download completes, run the **VulkanSDK-1.0.17.0** installer and follow the prompts to install the Vulkan SDK. When the Vulkan SDK has finished installing, reboot your PC to make sure that the environment variables are set properly. 
 
 1. When your PC finished rebooting open up your system's **Environment Variables** and check for the **Vulkan_SDK** in the **Variable** section. 
@@ -79,6 +87,9 @@ To build a UE4 project that has support for the Vulkan API you will need to do t
 	[![](VMP_New_Project_Settings.png)(w:555 h:387)](VMP_New_Project_Settings.png)
 	[/REGION]
 
+	[REGION:caption]
+	Click for full image.
+	[/REGION]
 	When that is completed, press the **Create** button to create and load the new project.
 
 1. Once the project has loaded go to **Edit** > **Project Settings** then under **Engine** go to the **Rendering** section and make sure that **Mobile HDR** has been enabled.
@@ -87,18 +98,30 @@ To build a UE4 project that has support for the Vulkan API you will need to do t
 	[![](VMP_Android_Settings_04.png)(w:639 h:385)](VMP_Android_Settings_04.png)
 	[/REGION]
 
+	[REGION:caption]
+	Click for full image.
+	[/REGION]
+
 1. Next, locate the **Android** section which is under **Platforms**.
 	
 	[REGION:lightbox]
 	[![](VMP_Android_Settings_00.png)(w:639 h:385)](VMP_Android_Settings_00.png)
 	[/REGION]
 
+	[REGION:caption]
+	Click for full image.
+	[/REGION]
+
 1. In the Android section under **APKPackaging** set the **Minimum SDK Version** to **19** and set the **Target SDK Version** to **19**. 
 
 	[REGION:lightbox]
 	[![](VMP_Android_Settings_01.png)(w:639 h:385)](VMP_Android_Settings_01.png)
 	[/REGION]
 
+	[REGION:caption]
+	Click for full image.
+	[/REGION]
+
 1. Below the APKPackaging section in the **Build** section make sure that the following options are enabled.
 
 	* **Support arm7[aka armebi-v7a]**
@@ -109,24 +132,39 @@ To build a UE4 project that has support for the Vulkan API you will need to do t
 	[![](VMP_Android_Settings_02.png)(w:639 h:385)](VMP_Android_Settings_02.png)
 	[/REGION]
 
+	[REGION:caption]
+	Click for full image.
+	[/REGION]
+
 1. In the **Android SDK** section make sure that **SDKConfig** paths are pointed to the correct SDK on your development PC.
 
 	[REGION:lightbox]
 	[![](VMP_Android_Settings_03.png)(w:639 h:385)](VMP_Android_Settings_03.png)
 	[/REGION]
 
+	[REGION:caption]
+	Click for full image.
+	[/REGION]
+
 1. From the menu bar go to **File** > **Package Project** > **Android** and select the **Android(ASTC)** option. 
 	
 	[REGION:lightbox]
 	[![](VMP_Package_Game.png)(w:329 h:266)](VMP_Package_Game.png)
 	[/REGION]
 
+	[REGION:caption]
+	Click for full image.
+	[/REGION]
+
 1. Select a location for UE4 to save the Android build and then press the **OK** button to start the packaging process.
 
 	[REGION:lightbox]
 	[![](VMP_Save_Build_Location.png)(w:245 h:267)](VMP_Save_Build_Location.png)
 	[/REGION]
-	![]()
+
+	[REGION:caption]
+	Click for full image.
+	[/REGION]
 
 1. When the packaging process has been completed, open up the folder that the packaged build was placed in. Inside this folder, you should see two .BAT files, locate the .BAT file that has the word **Install** in the name and double - click on it to install it to your device.
 	
@@ -134,12 +172,20 @@ To build a UE4 project that has support for the Vulkan API you will need to do t
 	[![](VMP_Build_Files.png)(w:439 h:198)](VMP_Build_Files.png)
 	[/REGION]
 
+	[REGION:caption]
+	Click for full image.
+	[/REGION]
+
 1. Once installed on your device, press on the UE4 icon that has your project name under it to launch the project on the device. 
 
 	[REGION:lightbox]
 	[![](VMP_On_Device.png)(w:818 h:320)](VMP_On_Device.png)
 	[/REGION]
 
+	[REGION:caption]
+	Click for full image.
+	[/REGION]
+
 ## Enabling Vulkan Preview Rendering in Editor
 You can enable the ability to use the Vulkan render while inside of the UE4 Editor, you will need to do the following in your project:
 
@@ -152,17 +198,31 @@ You can enable the ability to use the Vulkan render while inside of the UE4 Edit
 	[REGION:lightbox]
 	[![](VMP_03.png)(w:447 h:308)](VMP_03.png)
 	[/REGION]
+	
+	[REGION:caption]
+	Click for full image.
+	[/REGION]
 
 	[region:note]
 	After doing this, the Editor will need to recompile the entire shader cache to include the required Vulkan options. Depending on the size of your project and power of your development computer, this process could take anywhere from a few minutes to an hour or more to complete.
 	[/region]
 
-1. In the lower right-hand conner of the UE4 Editor window you should see the text, **Feature Level: ES3_1**. If you see this text then everything is setup correctly and you are now seeing what your project will look like when viewed on a device that supports the Vulkan API.
+1. On the **Main Toolbar**, click on the **Settings** button and go to **Preview Rendering Level** and select the **High-End Mobile / Metal** option to enable OpenGL ES 3.1 previewing in the UE4 viewport.
+
+	[REGION:lightbox]
+	[![](VMP_Android_Settings_05.png)(w:447 h:308)](VMP_Android_Settings_05.png)
+	[/REGION]
+
+1. In the lower right-hand conner of the UE4 Editor viewport you should see the text, **Feature Level: ES3_1**. If you see this text then everything is setup correctly and you are now seeing what your project will look like when viewed on a device that supports the Vulkan API.
 
 	[REGION:lightbox]
 	[![](VMP_04.png)(w:447 h:308)](VMP_04.png)
 	[/REGION]
 
+	[REGION:caption]
+	Click for full image.
+	[/REGION]
+
 ## Enabling the Vulkan Mobile Preview Render
 To enable the ability to use the Vulkan render with the mobile preview render, you will need to do the following in your project:
 
@@ -176,6 +236,9 @@ To enable the ability to use the Vulkan render with the mobile preview render, y
 	[![](VMP_01.png)(w:447 h:308)](VMP_01.png)
 	[/REGION]
 
+	[REGION:caption]
+	Click for full image.
+	[/REGION]
 
 1. Back in the Main Toolbar, change the **Play Mode** from **Play this level in the selected viewport** to the **Vulkan Mobile Preview (PIE)** option by clicking on the small white triangle next to the play button and selecting the Vulkan Mobile Preview (PIE) from the drop down list.
 
@@ -187,9 +250,13 @@ To enable the ability to use the Vulkan render with the mobile preview render, y
 	[![](VMP_05.png)(w:447 h:308)](VMP_05.png)
 	[/REGION]
 
-	[region:note]
-	If you do not see **(SF_VULKAN_ES31)** at the top of the preview window next to the project name then that means the project is not using the Vulkan API.
-	[/region]
+	[REGION:caption]
+	Click for full image.
+	[/REGION]
+	
+[region:note]
+If you do not see **(SF_VULKAN_ES31)** at the top of the preview window next to the project name then that means the project is not using the Vulkan API.
+[/region]
 
 
 
