#UE4 Docs: Fixed a small mistake with the wording that was missed before.

diff --git a/Engine/Documentation/Source/Platforms/Android/VulkanMobileRenderer/VulkanMobileRender.INT.udn b/Engine/Documentation/Source/Platforms/Android/VulkanMobileRenderer/VulkanMobileRender.INT.udn
--- a/Engine/Documentation/Source/Platforms/Android/VulkanMobileRenderer/VulkanMobileRender.INT.udn
+++ b/Engine/Documentation/Source/Platforms/Android/VulkanMobileRenderer/VulkanMobileRender.INT.udn
@@ -207,7 +207,7 @@ You can enable the ability to use the Vulkan render while inside of the UE4 Edit
 	After doing this, the Editor will need to recompile the entire shader cache to include the required Vulkan options. Depending on the size of your project and power of your development computer, this process could take anywhere from a few minutes to an hour or more to complete.
 	[/region]
 
-1. On the **Main Toolbar**, click on the **Settings** button and go to **Preview Rendering Level** and select the **High-End Mobile / Metal** option to enable OpenGL ES 3.1 previewing in the UE4 viewport.
+1. On the **Main Toolbar**, click on the **Settings** button and go to the **Preview Rendering Level** option and select the **High-End Mobile / Metal** option to enable OpenGL ES 3.1 previewing in the UE4 viewport.
 
 	[REGION:lightbox]
 	[![](VMP_Android_Settings_05.png)(w:447 h:308)](VMP_Android_Settings_05.png)
@@ -240,7 +240,7 @@ To enable the ability to use the Vulkan render with the mobile preview render, y
 	Click for full image.
 	[/REGION]
 
-1. Back in the Main Toolbar, change the **Play Mode** from **Play this level in the selected viewport** to the **Vulkan Mobile Preview (PIE)** option by clicking on the small white triangle next to the play button and selecting the Vulkan Mobile Preview (PIE) from the drop down list.
+1. On the Main Toolbar, change the **Play Mode** from **Play this level in the selected viewport** to the **Vulkan Mobile Preview (PIE)** option by clicking on the small white triangle next to the play button and selecting the Vulkan Mobile Preview (PIE) from the drop down list.
 
 	![](VMP_02.png)
 
