Updating based on SME comments and peer review.

diff --git a/Engine/Documentation/Source/Gameplay/Networking/Actors/Properties/Properties.INT.udn b/Engine/Documentation/Source/Gameplay/Networking/Actors/Properties/Properties.INT.udn
--- a/Engine/Documentation/Source/Gameplay/Networking/Actors/Properties/Properties.INT.udn
+++ b/Engine/Documentation/Source/Gameplay/Networking/Actors/Properties/Properties.INT.udn
@@ -53,7 +53,7 @@ Actors will observe a maximum update frequency set in their `NetUpdateFrequency`
 
 ### Adaptive Network Update Frequency
 
-With Adaptive Network Update Frequency, replicated Actors will increase or reduce their update frequency based on whether or not their updates are meaningful. In this context, a "meaningful" update is any update which initializes the Actor, adds or removes a subobject (i.e. an owned Component), or changes the value of a replicated field on the Actor or any of its subobjects. Each Actor’s possible range of update rates is determined by two variables on the Actor itself: `NetUpdateFrequency` and `MinNetUpdateFrequency`. `NetUpdateFrequency` indicates the maximum number of times per second that the Actor will attempt to update itself, while `MinNetUpdateFrequency` indicates the minimum number of update attempts per second.
+With Adaptive Network Update Frequency, we can save CPU cycles that would normally be wasted in redundant attempts to replicate an actor when nothing is really changing. When this feature is enabled, the system will dynamically adapt the update frequencies of individual Actors based on whether or not their updates are meaningful. In this context, a "meaningful" update is any update which initializes the Actor, adds or removes a subobject (i.e. an owned Component), or changes the value of a replicated field on the Actor or any of its subobjects. Each Actor’s possible range of update rates is determined by two variables on the Actor itself: `NetUpdateFrequency` and `MinNetUpdateFrequency`. `NetUpdateFrequency` indicates the maximum number of times per second that the Actor will attempt to update itself, while `MinNetUpdateFrequency` indicates the minimum number of update attempts per second.
 
 #### (Advanced) Update Frequency Decrease Algorithm
 
