#UE4 Docs: Checking in my fist pass on this document.

diff --git a/Engine/Documentation/Source/Platforms/Android/OpenGLES31MobileRenderer/OpenGLES31MobileRenderer.INT.udn b/Engine/Documentation/Source/Platforms/Android/OpenGLES31MobileRenderer/OpenGLES31MobileRenderer.INT.udn
new file mode 100644
--- /dev/null
+++ b/Engine/Documentation/Source/Platforms/Android/OpenGLES31MobileRenderer/OpenGLES31MobileRenderer.INT.udn
@@ -0,0 +1,168 @@
+Title: Android OpenGL ES 3.1 Mobile Renderer
+Description:Setting up your UE4 Android project to work with the OpenGL ES 3.1 mobile render
+Crumbs:%ROOT%, Platforms, Platforms/Android, Platforms/Android/OpenGLES31MobileRenderer/
+Availability:Public
+version:4.12
+Parent:Platforms/Android
+tags:Mobile
+tags:Android
+type:reference
+SkillLevel: Intermediate
+Related:Platforms\Android\VulkanMobileRenderer
+Related:Platforms\Android\DeviceCompatibility
+Related:Engine/Basics/Projects/Packaging
+Related:Platforms\Mobile\Performance
+Related:Programming/Development/BuildingUnrealEngine
+
+[TOC(start:2 end:3)]
+
+While Unreal Engine 4 (UE4) has long supported many OpenGL ES 3.0 and 3.1 features on Android, you can now specifically target ES 3.1 when building a UE4 Android project. This brings feature parity with Metal and Vulkan to higher-end Android devices and gives you access to 16 Texture samplers as well as improved performance through the use of uniform buffers. In the following document we will take a look how you can setup your UE4 project to use OpenGL ES 3.1 when run on an Android device that supports it.
+
+
+## Enabling the OpenGL ES 3.1 UE4 Editor Preview 
+
+In order to preview what OpenGL ES 3.1 will look like in the UE4 Editor you will need to first enable it by doing the following:
+
+1. From the **Main Toolbar** go to **Edit** and then select the **Editor Preferences** option.
+
+	![](OGL3_01.png)
+
+1. Under the **General** section locate the **Experimental** category and then look for the **Rendering** section and click on the checkmark box next to **Enable Metal/Vulkan/High-end mobile Preview Rendering Level in editor** to enable the option to preview OpenGL ES 3.1. 
+
+	[REGION:lightbox]
+	[![](Platforms\Android\VulkanMobileRenderer\VMP_03.png)(w:447 h:308)](Platforms\Android\VulkanMobileRenderer\VMP_03.png)
+	[/REGION]
+
+	[REGION:caption]
+	Click for full image.
+	[/REGION]
+
+	[region:note]
+	For a more detailed description on how to enable Metal/Vulkan/High-End  previewing in the UE4 Editor, check out the [Enabling Vulkan Preview Rendering in Editor](Platforms/Android/VulkanMobileRenderer/#enablingvulkanpreviewrenderingineditor) document.
+	[/region]
+
+1. Then on the **Main Tool Bar** click on the **Settings** option and go to **Preview Rendering Level** option and select the **High-End Mobile / Metal** option.
+
+	[REGION:lightbox]
+	[![](Platforms\Android\VulkanMobileRenderer\VMP_Android_Settings_05.png)(w:391)](Platforms\Android\VulkanMobileRenderer\VMP_Android_Settings_05.png)
+	[/REGION]
+
+	[REGION:caption]
+	Click for full image.
+	[/REGION]
+
+1. Now if you look at the lower right-hand corner of the UE4 Editor viewport, you should see the text **Feature Level: ES3_1** which means you are now previewing the level using OpenGL ES 3.1.
+
+	[REGION:lightbox]
+	[![](OGL3_02.png)(w:404)](OGL3_02.png)
+	[/REGION]
+
+	[REGION:caption]
+	Click for full image.
+	[/REGION]
+
+## Setting up a UE4 Project with OpenGL ES 3.1 Support
+
+To use OpenGL ES 3.1 in your UE4 Android project, you will need to make sure your UE4 project has the following options set:
+
+[region:warning]
+Please note that OpenGL ES 3.1 Support will only work on code based projects that are downloaded and compiled from Github. If you are not familiar with how to do this check out the [Download Unreal Engine Source Code](GettingStarted/DownloadingUnrealEngine) for detailed instructions on how to do this.
+[/region]
+
+1. First, go to the **Main Toolbar** and click on the **Edit** option and then select the **Project Settings** option. 
+
+	![](OGL3_03.png)
+
+1. From the **Project Settings** menu go to the **Platforms** section and click on the **Android** section.
+
+	[REGION:lightbox]
+	[![](Platforms\Android\VulkanMobileRenderer\VMP_Android_Settings_00.png)(w:639 h:385)](Platforms\Android\VulkanMobileRenderer\VMP_Android_Settings_00.png)
+	[/REGION]
+
+	[REGION:caption]
+	Click for full image.
+	[/REGION]
+
+1. Now in the **APKPackaging** section, set the **Minimum SDK** and **Target SDK** version to **21**.
+
+	[REGION:lightbox]
+	[![](OGL3_04.png)(w:639 h:385)](OGL3_04.png)
+	[/REGION]
+
+	[REGION:caption]
+	Click for full image.
+	[/REGION]
+
+1. Under the **APKPackaging** section, locate the  **Build** section and make sure that the following options are checked. 
+
+	* Support armv7
+	* Support OpenGL ES2
+	* Support OpenGL ES3.1
+
+	[REGION:lightbox]
+	[![](OGL3_04.png)(w:640)](OGL3_04.png)
+	[/REGION]
+
+	[REGION:caption]
+	Click for full image.
+	[/REGION]
+
+1. Under the **APKPackaging** section, locate the  **Build** section and make sure that the following options are checked. 
+	
+	* Support armv7
+	* Support OpenGL ES3.1
+
+	[REGION:lightbox]
+	[![](OGL3_05.png)(w:640)](OGL3_05.png)
+	[/REGION]
+
+	[REGION:caption]
+	Click for full image.
+	[/REGION]
+
+1. Click on the **Android SDK** section and set the **SDK API Level**  and **NDK API Level** to the following.
+
+	* SDK API Level = matchndk
+	* NDK API Level = android-21
+
+	[REGION:lightbox]
+	[![](OGL3_06.png)(w:640)](OGL3_06.png)
+	[/REGION]
+
+	[REGION:caption]
+	Click for full image.
+	[/REGION]
+
+## Building and Deploying a UE4 Project with OpenGL ES 3.1 Support
+
+Now that the project has been setup to work with OpenGL ES 3.1, it is now time to build the project and deploy it to a Android based Smartphone. To build a project that has OpenGL ES 3.1 support you will need to do the following:
+
+[region:note]
+Before you begin with the next part, make sure that you have your Android based Smartphone plugged into your development PC via a USB cable. Failing to do this will result in your project not being able to be deployed to your Android Smartphone.
+[/region]
+
+
+1. Locate the **Launch** option that is on the far right-hand side of the **Toolbar** menu.  
+
+	[REGION:lightbox]
+	[![](OGL3_07.png)(w:640)](OGL3_07.png)
+	[/REGION]
+
+	[REGION:caption]
+	Click for full image.
+	[/REGION]
+
+1. Click on the small white triangle that is on the left-hand side of the **Launch** button to expose the devices that your project can be launched on and then click on your Android device to start the build process.
+
+	![](OGL3_08.png)
+
+1. After this is completed, a build status bar will be displayed in the lower right-hand corner of the UE4 Editor that will show you how the build is progressing. To see a more detailed breakdown of what is going on with the build, click on the  **Show Output Log** to open the UE4 log.
+
+	![](OGL3_09.png)
+
+
+
+
+
+
+
