#UE4 Docs: fixed image links

diff --git a/Engine/Documentation/Source/Engine/Components/Widget/WidgetComponents.INT.udn b/Engine/Documentation/Source/Engine/Components/Widget/WidgetComponents.INT.udn
--- a/Engine/Documentation/Source/Engine/Components/Widget/WidgetComponents.INT.udn
+++ b/Engine/Documentation/Source/Engine/Components/Widget/WidgetComponents.INT.udn
@@ -17,7 +17,7 @@ The **Widget** Component itself is a 3D instance of a Widget Blueprint that you
 
 In the example below, we have added a Widget Component which will now allow us to specify a Widget Class Blueprint. 
 
-![](UMG_3D_21.png)
+![](Engine/UMG/HowTo/Create3DWidgets/UMG_3D_21.png)
 
 [REGION:caption]
 This example uses a Widget Blueprint to display a interactive menu that appears in the game world.
@@ -25,11 +25,11 @@ This example uses a Widget Blueprint to display a interactive menu that appears
 
 You can change the size of the WidgetComponent in the world by changing the **Draw Size** or using **Draw at Desired Size**. 
 
-![](UMG_3D_22.png)
+![](Engine/UMG/HowTo/Create3DWidgets/UMG_3D_22.png)
 
 Once you place the Actor that contains the WidgetComponent in your level, the Widget Class Blueprint is displayed in the world. 
 
-![](UMG_3D_30.png)
+![](Engine/UMG/HowTo/Create3DWidgets/UMG_3D_30.png)
 
 [REGION:caption]
 Here we have an interactive keypad that sits on a wall.
@@ -50,11 +50,11 @@ Here are the properties that can be adjusted via the **Details** panel for the W
 | **Draw at Desired Size** | Causes the render target to automatically match the desired size specified by the Widget Class. Expensive if is done every frame. 
 | **Pivot** | The alignment/pivot point that the widget is placed at relative to the position. |
 
-
+[PUBLISH:Docs]
 [REGION:note]
 See [](Engine\UMG\HowTo\InWorldWidgetInteraction) and [](Engine\UMG\HowTo\VirtualKeyboards) for creating 3D UI elements. 
 [/REGION]
-
+[/PUBLISH:Docs]
 
 
 
