#UE4 Docs: fixed metadata

diff --git a/Engine/Documentation/Source/Platforms/Android/InAppPurchases/In-AppPurchases.INT.udn b/Engine/Documentation/Source/Platforms/Android/InAppPurchases/In-AppPurchases.INT.udn
--- a/Engine/Documentation/Source/Platforms/Android/InAppPurchases/In-AppPurchases.INT.udn
+++ b/Engine/Documentation/Source/Platforms/Android/InAppPurchases/In-AppPurchases.INT.udn
@@ -7,7 +7,6 @@ Parent:Platforms/Android
 tags:Mobile
 tags:Android
 type:howto
-Prereq:P
 Related:Platforms/Mobile/InAppPurchases/
 
 ## Configuration
