Merging

//depot/UE4/Engine/Documentation/Source/...

to //UE4/Main/Engine/Documentation/Source/...

diff --git a/Engine/Documentation/Source/Engine/Rendering/LightingAndShadows/CapsuleShadows/BestPractices/CapsuleShadowsBestPractices.INT.udn b/Engine/Documentation/Source/Engine/Rendering/LightingAndShadows/CapsuleShadows/BestPractices/CapsuleShadowsBestPractices.INT.udn
new file mode 100644
--- /dev/null
+++ b/Engine/Documentation/Source/Engine/Rendering/LightingAndShadows/CapsuleShadows/BestPractices/CapsuleShadowsBestPractices.INT.udn
@@ -0,0 +1,135 @@
+﻿Availability: Docs
+Title: Capsule Shadows Best Practices
+Crumbs: %ROOT%, Engine, 
+Description: Capsule Shadows Best Practices
+Version: 4.13
+
+[VAR:Topic]
+[OBJECT:Topic]
+	[PARAM:image]
+		![%Engine/Rendering/LightingAndShadows/CapsuleShadows/BestPractices:title%]()
+	[/PARAM]
+	[PARAM:icon]
+		![](%ROOT%/light_icon.png)(convert:false)
+	[/PARAM]
+	[PARAM:title]
+		%Engine/Rendering/LightingAndShadows/CapsuleShadows/BestPractices:title%
+	[/PARAM]
+	[PARAM:description]
+		%Engine/Rendering/LightingAndShadows/CapsuleShadows/BestPractices:description%
+	[/PARAM]
+	[PARAM:path]
+		[RELATIVE:Engine/Rendering/LightingAndShadows/CapsuleShadows/BestPractices]
+	[/PARAM]
+[/OBJECT]
+[/VAR]
+[VAR:TopicCompact]
+[OBJECT:TopicCompact]
+	[PARAM:image]
+		![%Engine/Rendering/LightingAndShadows/CapsuleShadows/BestPractices:title%]()
+	[/PARAM]
+	[PARAM:icon]
+		![](%ROOT%/light_icon.png)(convert:false)
+	[/PARAM]
+	[PARAM:title]
+		%Engine/Rendering/LightingAndShadows/CapsuleShadows/BestPractices:title%
+	[/PARAM]
+	[PARAM:description]
+		%Engine/Rendering/LightingAndShadows/CapsuleShadows/BestPractices:description%
+	[/PARAM]
+	[PARAM:path]
+		[RELATIVE:Engine/Rendering/LightingAndShadows/CapsuleShadows/BestPractices]
+	[/PARAM]
+[/OBJECT]
+[/VAR]
+
+
+
+[TOC(start:2 end:3)]
+
+
+
+
+
+## Best Practices
+
+Below are some steps you can take to ensure that your Capsule Shadows look correct and perform well for your game.
+
+### Feet Capsule Placement
+
+It is important to get the foot placement right for your characters to make them feel grounded in the environment. Take your time to tweak this part of the setup for your Shadow Physics Asset to 
+make sure that it looks right. You may also want to have several different lighting scenarios set up in a test scene to make sure everything looks correct.
+
+
+[REGION:lightbox]
+[![](FootPlacement1.png)(w:600)](FootPlacement1.png)
+[/REGION]
+
+In this setup the capsules for the feet are close to relative size of the feet and have some overlap with the capsules for the shins. The shadows are soft and not overly dark.
+
+[REGION:lightbox]
+[![](FootPlacement2.png)(w:600)](FootPlacement2.png)
+[/REGION]
+ 
+In this setup the capsules are too large for the feet and extend beyond the back of the heel leaving the soft shadows to appear darker than they should.
+
+
+### Limit Capsule Bodies
+
+Capsule Shadows can be considered an advanced form of "blob" shadows. You can simply use a single capsule or as many as you need to get better shadowing accuracy. With that in mind, start by 
+determining the needs of your character's movement for your game and whether they will ever be near a surface by using a cover system or falling down. If they will never be near surfaces consider
+not having capsule bodies for the arms of your character, since the performance costs is determined by the number of capsule bodies.
+
+[REGION:lightbox]
+[![](BlobShadow.png)(w:400)](BlobShadow.png)
+[/REGION]
+[REGION:caption]
+_Click image for full size view_
+[/REGION]
+
+In this example, there is a single capsule that is used to represent the character. In indirectly lit areas this can save performance costs by only using a single capsule for the character's representation
+if accuracy is not a concern.
+
+[OBJECT:ComparisonSlider]
+	[PARAM:before]
+	![With Capsules for Arms](WithArms.png)(w:600)
+	[/PARAM]
+	[PARAM:after]
+	![Without Capusles for Arms](WithoutArms.png)(w:600)
+	[/PARAM]
+[/OBJECT]
+
+In this comparison, the arms have been removed from the character's representation. In the directly lit character areas not near surfaces you can see that it only makes a slight visual difference. 
+
+
+### Indirect Lighting Cache (Volume Light Samples)
+
+Capsule Shadows for indirect lighting are heavily influenced by the [Indirect Lighting Cache](Engine/Rendering/LightingAndShadows/IndirectLightingCache/) to determine the directionality and intensity of the 
+soft shadow. Sometimes this requires tweaking the placement of the light samples to be more dense or spaced out on your game's needs. 
+
+To adjust the **Volume Light Sample Placement** you must enable the **World Settings** panel by going to **Settings** > **World Settings** in the main toolbar. Then in the World Settings panel under the 
+**Lightmass** section you can change the **Volume Light Sample Placement**. Use lower values to make the placement of samples more dense or higher values to make them more spread out.
+
+
+You can visualize these samples in your level by going to the **Viewport** > **Show** > **Visualize** > **Volume Lighting Samples**. 
+
+[OBJECT:ComparisonSlider]
+	[PARAM:before]
+	![Volume Light Samples Placement: 1.0 (Default)](VolmeLightSamplesDefault.png)(w:600)
+	[/PARAM]
+	[PARAM:after]
+	![Volume Light Samples Placement: 0.25](WithArms.png)(w:600)
+	[/PARAM]
+[/OBJECT]
+
+You may need to adjust the settings for the Volume Light Samples Placement to get the results that you desire. In this scene setup, adjusting to a value between 0.5 and 0.25 gave better results, but 
+this may be different for your game.
+
+[REGION:warning]
+Values too low will cause a "popping" as the capsule moves between the samples instead of providing a soft smooth transition between them, since the placement of the samples is so close together.
+[/REGION]
+
+
+
+
+
diff --git a/Engine/Documentation/Source/Engine/Rendering/LightingAndShadows/CapsuleShadows/CapsuleShadowsLanding.INT.udn b/Engine/Documentation/Source/Engine/Rendering/LightingAndShadows/CapsuleShadows/CapsuleShadowsLanding.INT.udn
new file mode 100644
--- /dev/null
+++ b/Engine/Documentation/Source/Engine/Rendering/LightingAndShadows/CapsuleShadows/CapsuleShadowsLanding.INT.udn
@@ -0,0 +1,28 @@
+﻿availability:docs
+title: Capsule Shadows
+description:Landing page for Capsule Shadows in Unreal Engine 4
+crumbs:
+version:4.13
+
+**Capsule Shadows** have been developed to support soft shadow casting for Skeletal Meshes by creating a capsule representation of your character using a custom Physics Asset. 
+The soft shadows serve to ground the character in indirectly areas and give soft shadows in directly lit areas. This shadowing technique can be used with Unreal Engine 4.11 or later. 
+
+
+[OBJECT:TopicButtonList]
+	[PARAM:icon]
+		![](%ROOT%/tech_icon.png)(convert:false)
+	[/PARAM]
+	[PARAM:title]
+		Capsule Shadows Topics
+	[/PARAM]
+	[PARAM:description]
+	[/PARAM]
+	[PARAM:links]
+	* [](Engine\Rendering\LightingAndShadows\CapsuleShadows\Overview "%Engine\Rendering\LightingAndShadows\CapsuleShadows\Overview:description%")
+	* [](Engine\Rendering\LightingAndShadows\CapsuleShadows\BestPractices "%Engine\Rendering\LightingAndShadows\CapsuleShadows\BestPractices:description%")
+	* [](Engine\Rendering\LightingAndShadows\CapsuleShadows\Reference "%Engine\Rendering\LightingAndShadows\CapsuleShadows\Reference:description%")
+	* [](Engine\Rendering\LightingAndShadows\CapsuleShadows\HowTo "%Engine\Rendering\LightingAndShadows\CapsuleShadows\HowTo:description%")
+	[/PARAM]
+[/OBJECT]
+
+
diff --git a/Engine/Documentation/Source/Engine/Rendering/LightingAndShadows/CapsuleShadows/HowTo/AssignShadowAsset/AssignShadowAssetHowTo.INT.udn b/Engine/Documentation/Source/Engine/Rendering/LightingAndShadows/CapsuleShadows/HowTo/AssignShadowAsset/AssignShadowAssetHowTo.INT.udn
new file mode 100644
--- /dev/null
+++ b/Engine/Documentation/Source/Engine/Rendering/LightingAndShadows/CapsuleShadows/HowTo/AssignShadowAsset/AssignShadowAssetHowTo.INT.udn
@@ -0,0 +1,75 @@
+﻿Availability: Docs
+Title: How-To Assign a Shadow Physics Asset to a Skeletal Mesh
+Crumbs: %ROOT%, Engine, 
+Description: How to assign a Shadow Physics Asset to a Skeletal Mesh.
+Version: 4.13
+related:Engine\Rendering\LightingAndShadows\CapsuleShadows\HowTo\CreateShadowAsset
+related:Engine\Rendering\LightingAndShadows\CapsuleShadows\HowTo\EnableCapsuleShadows
+related:Engine\Rendering\LightingAndShadows\CapsuleShadows\Overview
+related: Engine\Rendering\LightingAndShadows\CapsuleShadows\BestPractices
+related: Engine\Rendering\LightingAndShadows\CapsuleShadows\Reference
+
+[VAR:Topic]
+[OBJECT:Topic]
+	[PARAM:image]
+		![%Engine/Rendering/LightingAndShadows/CapsuleShadows:title%]()
+	[/PARAM]
+	[PARAM:icon]
+		![](%ROOT%/light_icon.png)(convert:false)
+	[/PARAM]
+	[PARAM:title]
+		%Engine/Rendering/LightingAndShadows/CapsuleShadows:title%
+	[/PARAM]
+	[PARAM:description]
+		%Engine/Rendering/LightingAndShadows/CapsuleShadows:description%
+	[/PARAM]
+	[PARAM:path]
+		[RELATIVE:Engine/Rendering/LightingAndShadows/CapsuleShadows]
+	[/PARAM]
+[/OBJECT]
+[/VAR]
+[VAR:TopicCompact]
+[OBJECT:TopicCompact]
+	[PARAM:image]
+		![%Engine/Rendering/LightingAndShadows/CapsuleShadows:title%]()
+	[/PARAM]
+	[PARAM:icon]
+		![](%ROOT%/light_icon.png)(convert:false)
+	[/PARAM]
+	[PARAM:title]
+		%Engine/Rendering/LightingAndShadows/CapsuleShadows:title%
+	[/PARAM]
+	[PARAM:description]
+		%Engine/Rendering/LightingAndShadows/CapsuleShadows:description%
+	[/PARAM]
+	[PARAM:path]
+		[RELATIVE:Engine/LightingAndShadows/Rendering/CapsuleShadows]
+	[/PARAM]
+[/OBJECT]
+[/VAR]
+
+
+To have Capsule Shadows for your Skeletal Mesh you must assign a Physics Asset you created for your shadows to your Actor. Follow the steps below to assign your shadow Physics Asset.
+
+
+## Steps
+
+1. Inside the **Content Browser**, select your character Skeletal Mesh and **Double-Click** to open it.
+
+	![](CB.png)
+
+	After a moment **Persona**, the Skeletal Mesh editor, will open.
+	[REGION:lightbox]
+	[![](Persona.png)(w:600)](Persona.png)
+	[/REGION]
+	_Click image for full size_
+
+2. In the **Details** panel under the **Lighting** category, click the drop-down for **Shadow Phsyics Asset** and select the **Physics Asset** you created for Capsule Shadows.
+	
+	![](AssignShadowPhsyicsAsset.png)(w:450)
+
+	[REGION:note]
+	Since the **Shadow Physics Asset** uses a Physics Asset to create Capsule Shadows, you may see other Physics Assets in this list. It is important to name your Shadow Physics Assets so that you can easily locate it.
+	[/REGION]
+
+3. Now you can **Save** and **Close** Persona.
diff --git a/Engine/Documentation/Source/Engine/Rendering/LightingAndShadows/CapsuleShadows/HowTo/CapsuleShadowsHowTos.INT.udn b/Engine/Documentation/Source/Engine/Rendering/LightingAndShadows/CapsuleShadows/HowTo/CapsuleShadowsHowTos.INT.udn
new file mode 100644
--- /dev/null
+++ b/Engine/Documentation/Source/Engine/Rendering/LightingAndShadows/CapsuleShadows/HowTo/CapsuleShadowsHowTos.INT.udn
@@ -0,0 +1,25 @@
+﻿availability:docs
+title: Capsule Shadows How-To's
+description:How-To's for Capsule Shadows in Unreal Engine 4.
+crumbs:
+version:4.13
+related:Engine\Rendering\LightingAndShadows\CapsuleShadows\Overview
+related:Engine\Rendering\LightingAndShadows\CapsuleShadows\BestPractices
+related:Engine\Rendering\LightingAndShadows\CapsuleShadows\Reference
+
+[OBJECT:TopicButtonList]
+	[PARAM:icon]
+		![](%ROOT%/tech_icon.png)(convert:false)
+	[/PARAM]
+	[PARAM:title]
+		Capsule Shadow Topics
+	[/PARAM]
+	[PARAM:description]
+	[/PARAM]
+	[PARAM:links]
+	* [](Engine\Rendering\LightingAndShadows\CapsuleShadows\HowTo\CreateShadowAsset "%Engine\Rendering\LightingAndShadows\CapsuleShadows\HowTo\CreateShadowAsset:description%")
+	* [](Engine\Rendering\LightingAndShadows\CapsuleShadows\HowTo\AssignShadowAsset "%Engine\Rendering\LightingAndShadows\CapsuleShadows\HowTo\AssignShadowAsset:description%")
+	* [](Engine\Rendering\LightingAndShadows\CapsuleShadows\HowTo\EnableCapsuleShadows "%Engine\Rendering\LightingAndShadows\CapsuleShadows\HowTo\EnableCapsuleShadows:description%")
+	[/PARAM]
+[/OBJECT]
+
diff --git a/Engine/Documentation/Source/Engine/Rendering/LightingAndShadows/CapsuleShadows/HowTo/CreateShadowAsset/CreateShadowAssetHowTo.INT.udn b/Engine/Documentation/Source/Engine/Rendering/LightingAndShadows/CapsuleShadows/HowTo/CreateShadowAsset/CreateShadowAssetHowTo.INT.udn
new file mode 100644
--- /dev/null
+++ b/Engine/Documentation/Source/Engine/Rendering/LightingAndShadows/CapsuleShadows/HowTo/CreateShadowAsset/CreateShadowAssetHowTo.INT.udn
@@ -0,0 +1,165 @@
+﻿Availability: Docs
+Title: How-To Create a Shadow Phsyics Assets
+Crumbs: %ROOT%, Engine, 
+Description: How to create a Shadows Physics Asset for Capsule Shadows.
+Version: 4.13
+prereq: Engine\Physics\PhAT\HowTo
+related: Engine\Rendering\LightingAndShadows\CapsuleShadows\HowTo\AssignShadowAsset
+related: Engine\Rendering\LightingAndShadows\CapsuleShadows\HowTo\EnableCapsuleShadows
+related: Engine\Rendering\LightingAndShadows\CapsuleShadows\Overview
+related: Engine\Rendering\LightingAndShadows\CapsuleShadows\BestPractices
+related: Engine\Rendering\LightingAndShadows\CapsuleShadows\Reference
+
+[VAR:Topic]
+[OBJECT:Topic]
+	[PARAM:image]
+		![%Engine/Rendering/LightingAndShadows/CapsuleShadows:title%]()
+	[/PARAM]
+	[PARAM:icon]
+		![](%ROOT%/light_icon.png)(convert:false)
+	[/PARAM]
+	[PARAM:title]
+		%Engine/Rendering/LightingAndShadows/CapsuleShadows:title%
+	[/PARAM]
+	[PARAM:description]
+		%Engine/Rendering/LightingAndShadows/CapsuleShadows:description%
+	[/PARAM]
+	[PARAM:path]
+		[RELATIVE:Engine/Rendering/LightingAndShadows/CapsuleShadows]
+	[/PARAM]
+[/OBJECT]
+[/VAR]
+[VAR:TopicCompact]
+[OBJECT:TopicCompact]
+	[PARAM:image]
+		![%Engine/Rendering/LightingAndShadows/CapsuleShadows:title%]()
+	[/PARAM]
+	[PARAM:icon]
+		![](%ROOT%/light_icon.png)(convert:false)
+	[/PARAM]
+	[PARAM:title]
+		%Engine/Rendering/LightingAndShadows/CapsuleShadows:title%
+	[/PARAM]
+	[PARAM:description]
+		%Engine/Rendering/LightingAndShadows/CapsuleShadows:description%
+	[/PARAM]
+	[PARAM:path]
+		[RELATIVE:Engine/Rendering/LightingAndShadows/CapsuleShadows]
+	[/PARAM]
+[/OBJECT]
+[/VAR]
+
+
+In order to have Capsule Shadow representations a Physics Asset that will be used for soft shadowing must be created that is unique to your Skeletal Mesh Actor. 
+Since this Physics Asset is not used for any physics simulations it does not have to accurately match the Skeletal Mesh. Follow the steps below to create a Shadow 
+Physics Asset for your Skeletal Mesh.
+
+
+## Steps
+
+1. Inside the **Content Browser**, find and select your character **Skeletal Mesh**. 
+
+	[REGION:note]
+	If you do not have a Skeletal Mesh to use, you can enable **Show Engine Content** under the **View Options** selection in the **Content Browser.** Then navigate to
+	**Engine Content** > **Tutorial** > **SubEditors** > **TutorialAssets** > **Character** folder and copy the **TutorialTPP** Skeletal Mesh to your own folder under **Content**.
+	[/REGION]
+
+2. Right-Click on the **TutorialTPP** Skeletal Mesh to open the context menu and select **Create**, then select **Physics Asset**.
+	
+	[REGION:raw]
+	![](RtClickCreatePA.png)
+	[/REGION]
+
+3. A **New Asset** window will open. Set the **Minimum Bone Size** to **15** and click **Ok**.
+
+	![](PA_NewAssetWindow.png)
+
+	[REGION:note]
+	Only use **Sphyl** bodies for the best results. **Sphere** bodies will work, but are not as flexible.
+	[/REGION]
+
+4. At this point, the **Physics Asset Tool** (PhAT) will open. Before moving forward, minimize the **PhAT** window and go back to the **Content Browser** and rename the newly created Physics asset from **TutorialTPP_Physics** to something
+like **CS_TutorialTPP**, or anything that is recognizable for you.
+
+5. Now **maximize** the **PhAT** editor window to view the Physics Asset that has been created. You should see something similar to this.
+
+	![](PhysicsAssetCapsules.png)(w:500)
+
+6. Adjust the **Physics Bodies** so that the capsules represent the body of your Character more accurately. There are several things to consider when adjusting the capsules for your character listed below. 
+
+	* Remove any unnecessary capsules, like hands, multiple bodies for the torso, and neck.
+	* Foot capsules are the most important to make the character look grounded.
+	* Use a little bit of overlap of the capsules between joints.
+	* Capsules for the arms are not necessary unless you're character will be near surfaces like with cover systems or if the character can go prone.
+	
+	[REGION:tip]
+	In the **PhAT Editor Main Viewport** set the **Collision Render Mode** to **Solid** to make it easier to see the adjustments you make to the Sphyl bodies for your Skeletal Mesh.
+	[INCLUDE:#collisionmode]
+	[INCLUDE:#imagetable]
+	[/REGION]
+	
+	[REGION:note]
+	The shadow casting in the PhAT editor is not representative of the shadows that will be cast from the capsule bodies. You must test them in your scene.
+	[/REGION]
+	
+	For more information on how to use the [Physics Asset Tool](https://docs.unrealengine.com/latest/INT/Engine/Physics/PhAT/HowTo/index.html) see the following page.
+
+7. When you are finished with your adjustments you should have something that looks similiar to this. Depending on the Skeletal Mesh you used the results may have something different.
+	
+	[INCLUDE:#imagetable2]
+
+
+
+<!---
+[EXCERPT:Capsule1]
+[REGION:lightbox]
+[![](CollisionRenderWireframe.png)(w:350)](CollisionRenderWireframe.png)
+[/REGION]
+[/EXCERPT:Capsule1]
+
+[EXCERPT:Capsule2]
+[REGION:lightbox]
+[![](CollisionRenderSolid.png)(w:350)](CollisionRenderSolid.png)
+[/REGION]
+[/EXCERPT:Capsule2]
+
+[EXCERPT:collisionmode]
+[REGION:lightbox]
+[![](PhATCollisionRenderMode.png)(w:375)](PhATCollisionRenderMode.png)
+[/REGION]
+[/EXCERPT:collisionmode]
+
+[EXCERPT:imagetable]
+[REGION:imagetable]
+| [INCLUDE:#Capsule1] | [INCLUDE:#capsule2] |
+|-----------|-------|
+| Wireframe | Solid |
+[/REGION]
+[/EXCERPT:imagetable]
+
+[EXCERPT:imagetable2]
+[REGION:imagetable]
+| [INCLUDE:#Front] | [INCLUDE:#Left] |
+|-----------|-------|
+| Front | Left |
+[/REGION]
+[/EXCERPT:imagetable2]
+
+[EXCERPT:Front]
+[REGION:lightbox]
+[![](PA_Front.png)(w:425)](PA_Front.png)
+[/REGION]
+[/EXCERPT:Front]
+
+[EXCERPT:Left]
+[REGION:lightbox]
+[![](PA_Left.png)(w:425)](PA_Left.png)
+[/REGION]
+[/EXCERPT:Left]
+
+--->
+
+
+
+
+
diff --git a/Engine/Documentation/Source/Engine/Rendering/LightingAndShadows/CapsuleShadows/HowTo/EnableCapsuleShadows/EnableCapsuleShadowsHowTo.INT.udn b/Engine/Documentation/Source/Engine/Rendering/LightingAndShadows/CapsuleShadows/HowTo/EnableCapsuleShadows/EnableCapsuleShadowsHowTo.INT.udn
new file mode 100644
--- /dev/null
+++ b/Engine/Documentation/Source/Engine/Rendering/LightingAndShadows/CapsuleShadows/HowTo/EnableCapsuleShadows/EnableCapsuleShadowsHowTo.INT.udn
@@ -0,0 +1,87 @@
+﻿Availability: Docs
+Title: How-To Enable Capsule Shadows for Skeletal Meshes
+Crumbs: %ROOT%, Engine, 
+Description: How to enable Capsule Shadows for Skeletal Meshes.
+Version: 4.13
+related:Engine\Rendering\LightingAndShadows\CapsuleShadows\HowTo\CreateShadowAsset
+related:Engine\Rendering\LightingAndShadows\CapsuleShadows\HowTo\AssignShadowAsset
+related: Engine\Rendering\LightingAndShadows\CapsuleShadows\Overview
+related: Engine\Rendering\LightingAndShadows\CapsuleShadows\BestPractices
+related: Engine\Rendering\LightingAndShadows\CapsuleShadows\Reference
+
+[VAR:Topic]
+[OBJECT:Topic]
+	[PARAM:image]
+		![%Engine/Rendering/LightingAndShadows/CapsuleShadows:title%]()
+	[/PARAM]
+	[PARAM:icon]
+		![](%ROOT%/light_icon.png)(convert:false)
+	[/PARAM]
+	[PARAM:title]
+		%Engine/Rendering/LightingAndShadows/CapsuleShadows:title%
+	[/PARAM]
+	[PARAM:description]
+		%Engine/Rendering/LightingAndShadows/CapsuleShadows:description%
+	[/PARAM]
+	[PARAM:path]
+		[RELATIVE:Engine/Rendering/LightingAndShadows/CapsuleShadows]
+	[/PARAM]
+[/OBJECT]
+[/VAR]
+[VAR:TopicCompact]
+[OBJECT:TopicCompact]
+	[PARAM:image]
+		![%Engine/Rendering/LightingAndShadows/CapsuleShadows:title%]()
+	[/PARAM]
+	[PARAM:icon]
+		![](%ROOT%/light_icon.png)(convert:false)
+	[/PARAM]
+	[PARAM:title]
+		%Engine/Rendering/LightingAndShadows/CapsuleShadows:title%
+	[/PARAM]
+	[PARAM:description]
+		%Engine/Rendering/LightingAndShadows/CapsuleShadows:description%
+	[/PARAM]
+	[PARAM:path]
+		[RELATIVE:Engine/Rendering/LightingAndShadows/CapsuleShadows]
+	[/PARAM]
+[/OBJECT]
+[/VAR]
+
+
+The Capsule Shadows must be enabled on each Skeletal Mesh Actor in your level and requires a unique Shadow Physics Asset to have been created and assigned. Follow the steps below to enable
+Capsule Shadows for your character in your levels.
+
+## Steps
+
+1. Inside the **Content Browser**, find and select your Skeletal Mesh that has been assigned a Physics Asset for shadows.
+
+	![](CB.png)
+
+2. **Left-Click** on the Skeletal Mesh and drag-and-drop it into the Level **Viewport**.
+	[REGION:lightbox]
+	[![](AddSkelMeshToScene.png)(w:500)](AddSkelMeshToScene.png)
+	[/REGION]
+
+3. Inside the **Details** panel for your Skeletal Mesh, look under the **Lighting** section, locate and enable **Capsule Direct Shadows** and **Capsule Indirect Shadows**.
+	
+	![](EnableCapsuleShadows.png)(w:350)
+	
+4. Finish by creating a scene that has some shadowed and unshadowed areas, then use the **Build** button in the main toolbar to build lighting. After the lighting build has completed
+you should be able to see similar results as the ones below for your characters.  
+
+	[OBJECT:EmbeddedVideo]
+	[PARAMLITERAL:width]
+	640
+	[/PARAMLITERAL]
+	[PARAMLITERAL:height]
+	360
+	[/PARAMLITERAL]
+	[PARAMLITERAL:videoid]
+	YwQg9mj-cEE
+	[/PARAMLITERAL]
+	[/OBJECT]
+	 
+[REGION:caption]
+Capsule Indirect Shadows (Left), Capsule Direct Shadows (Right)
+[/REGION]
diff --git a/Engine/Documentation/Source/Engine/Rendering/LightingAndShadows/CapsuleShadows/Overview/CapsuleShadowsOverview.INT.udn b/Engine/Documentation/Source/Engine/Rendering/LightingAndShadows/CapsuleShadows/Overview/CapsuleShadowsOverview.INT.udn
new file mode 100644
--- /dev/null
+++ b/Engine/Documentation/Source/Engine/Rendering/LightingAndShadows/CapsuleShadows/Overview/CapsuleShadowsOverview.INT.udn
@@ -0,0 +1,215 @@
+﻿Availability: Docs
+Title: Capsule Shadows Overview
+Crumbs: %ROOT%, Engine, 
+Description: Capsule Shadows Overview
+Version: 4.13
+
+[VAR:Topic]
+[OBJECT:Topic]
+	[PARAM:image]
+		![%Engine/Rendering/LightingAndShadows/CapsuleShadows/CapsuleShadowsOverview:title%]()
+	[/PARAM]
+	[PARAM:icon]
+		![](%ROOT%/light_icon.png)(convert:false)
+	[/PARAM]
+	[PARAM:title]
+		%Engine/Rendering/LightingAndShadows/CapsuleShadows/CapsuleShadowsOverview:title%
+	[/PARAM]
+	[PARAM:description]
+		%Engine/Rendering/LightingAndShadows/CapsuleShadows/CapsuleShadowsOverview:description%
+	[/PARAM]
+	[PARAM:path]
+		[RELATIVE:Engine/Rendering/LightingAndShadows/CapsuleShadows/CapsuleShadowsOverview]
+	[/PARAM]
+[/OBJECT]
+[/VAR]
+[VAR:TopicCompact]
+[OBJECT:TopicCompact]
+	[PARAM:image]
+		![%Engine/Rendering/LightingAndShadows/CapsuleShadows/CapsuleShadowsOverview:title%]()
+	[/PARAM]
+	[PARAM:icon]
+		![](%ROOT%/light_icon.png)(convert:false)
+	[/PARAM]
+	[PARAM:title]
+		%Engine/Rendering/LightingAndShadows/CapsuleShadows/CapsuleShadowsOverview:title%
+	[/PARAM]
+	[PARAM:description]
+		%Engine/Rendering/LightingAndShadows/CapsuleShadows/CapsuleShadowsOverview:description%
+	[/PARAM]
+	[PARAM:path]
+		[RELATIVE:Engine/Rendering/LightingAndShadows/CapsuleShadows/CapsuleShadowsOverview]
+	[/PARAM]
+[/OBJECT]
+[/VAR]
+
+
+[TOC(start:2 end:3)]
+
+
+**Capsule Shadows** is a technique that is used to generate soft shadows in indirectly and directly lit areas by using a capsule representation of your characters with a Physics Asset.
+This solves the problem of making a character feel grounded in indirectly lit areas and by also giving soft shadow representations for those which are directly lit.
+
+
+## How do they work?
+
+Capsule Shadows start by using a unique Physics Asset that can be assigned to any Skeletal Mesh as a **Shadow Physics Asset**. This asset is unique to each Skeletal Mesh as an 
+accurate representation of your character. Once this information is known the [Indirect Lighting Samples](Engine/Rendering/LightingAndShadows/IndirectLightingCache/) that are 
+generated after a light build can be used to cast a cone towards the direction of the light intensity. As the character nears this point the cone is used to check if any capsules 
+are occluding the cone. If they are, then all the capsule shadows can be accumulated to get this end result of soft shadows.
+
+<!---
+The Capsule Shadow technique allows a Physics Asset to be assigned to any Skeletal Mesh as a **Shadow Physics Asset**. This asset is unique to each Skeletal Mesh so that it can represent your character 
+accurately for soft shadowing. Character shadowing has generally been handled by using standard shadow mapping techniques, which can only have sharp shadows with slight blurring. Capsule Shadows make 
+it possible to have these soft shadow representations for your character so that it can seamlessly blend with the environment.
+--->
+
+[REGION:lightbox]
+[![](CapsuleShadowRepresentation.png)(w:500)](CapsuleShadowRepresentation.png)
+[/REGION]
+[REGION:caption]
+_Click the image for full size view_
+[/REGION]
+
+The scene above breaks down how Capsule Shadows are created for your character. 
+
+1. Once lighting has been built for the scene the Indirect Lighting Cache samples that are generated for your level's surfaces is used to gather the directionality and intensity of the light. 
+2. The cone is placed near a light sample and from this can determine the direction the shadow will cast and how dark the shadow should be. If the light sample is brighter than ones in more covered areas the cone will be larger so that the shadows will be softer with less directionality.
+3. The cone checks if the capsules are occluding any part of it. If the capsules are occluding, the cone direction and intensity of the lighting samples is used to draw the soft shadows.
+4. Finally, we have the soft shadow that is created from the capsules that are occluding the cone.
+
+
+
+#### Capsule Representation of Character
+
+A Physics Asset is used to create a accurate representation of the character. Since the capsules are attached to the bones of the character this allows for the shadows to move with the character.
+
+For steps to setup your Phsyics Asset see the [How to Create a Shadow Phsyics Assets](Engine/Rendering/LightingAndShadows/CapsuleShadows/HowTo/CreateShadowAsset) page.
+
+For tips on setting up your Shadow Physics Asset see the [Best Practices](Engine/Rendering/LightingAndShadows/CapsuleShadows/BestPractices) page for Capsule Shadows.
+
+[REGION:lightbox]
+[![](PA_WithoutArms.png)(w:500)](PA_withoutArms.png)
+[/REGION]
+
+[REGION:caption]
+_Capsule representation of a Skeletal Mesh character_
+[/REGION]
+
+[REGION:note]
+The shadows on the ground are not Capsule Shadows. 
+[/REGION]
+
+
+## Capsule Shadow Settings
+
+You can enable Capsule Shadows for each Skeletal Mesh in your level by enabling the option for **Capsule Direct Shadows** and/or **Capsule Indirect Shadows** in its **Detail's** panel.
+
+![](EnableCapsuleShadows.png)(w:350)
+
+
+#### Capsules for Direct Shadowing
+
+**Capsule Direct Shadows** enables soft shadowing for your character when in areas that are directly lit by any **Movable** or **Stationary** light source. 
+
+[REGION:lightbox]
+[![](DirectCapsuleShadows.png)(w:500)](DirectCapsuleShadows.png)
+[/REGION]
+[REGION:caption]
+_Click the image for full size_
+[/REGION]
+
+##### Light Source Angle
+
+For directly lit Capsule Shadows you can adjust the **Light Source Angle** of the Directional Light to soften the shadows. 
+
+[OBJECT:ComparisonSlider]
+	[PARAM:before]
+	![Light Source Angle: 1.0 (Default)](DirectLightSourceAngle1.png)(w:500)
+	[/PARAM]
+	[PARAM:after]
+	![Light Source Angle: 2.0](DirectLightSourceAngle2.png)(w:500)
+	[/PARAM]
+[/OBJECT]
+
+#### Capsules for Indirect Shadowing
+
+**Capsule Indirect Shadows** enables soft shadowing for your character in areas that are indirectly lit by any **Static** or **Stationary** light source. Capsule shadowing is created using the 
+[Indirect Lighting Cache](Engine/Rendering/LightingAndShadows/IndirectLightingCache/) that is generated during a light build. 
+
+[REGION:lightbox]
+[![](IndirectDefault.png)(w:500)](IndirectDefault.png)
+[/REGION]
+[REGION:caption]
+_Click the image for full size_
+[/REGION]
+
+
+## Directonality and Intensity of Indirect Shadows
+
+The directionality and intensity of the capsule shadowing is determined by the light samples generated after a light build. The brighter intensity of the lighting sample the less intense and softer the shadow will be since light is coming from many directions. 
+
+### Closed Area Shadowing
+
+When you have a opening for light to pour through the intensity of the light will be brighter at the opening and dimmer at the back part of this area. 
+
+[REGION:lightbox]
+[![](CS_CaveFront.png)(w:500)](CS_CaveFront.png)
+[/REGION]
+[REGION:caption]
+_Click the image for full size_
+[/REGION]
+
+In the front of this area the character's shadow is much softer since there is more light and the shadow has less directionality since more light is coming from above. 
+
+[REGION:lightbox]
+[![](CS_CaveBack.png)(w:500)](CS_CaveBack.png)
+[/REGION]
+[REGION:caption]
+_Click the image for full size_
+[/REGION]
+
+However, at the back of this area the shadow is more intense since it is more dimly lit and has more directionality. The intensity of the light is coming from a single opening towards the front which means there will be
+less light traveling to the back of this area.
+
+
+### Open Area Shadowing
+
+In open areas that receive indirect shadowing, the Capsule Shadows will appear much softer since the light can come from all around.
+
+[REGION:lightbox]
+[![](CS_openarea.png)(w:500)](CS_openarea.png)
+[/REGION]
+[REGION:caption]
+_Click the image for full size_
+[/REGION]
+
+
+
+
+## Performance
+
+The cost of Capsule Shadows in your game is cheaper and more efficient than standard shadow mapping techniques because of tile defered culling which cull the capsule influences to only the tiles that could possibly 
+receive shadows from that capsule. This means that most of the areas on the screen knew they did not have to do any work. The tiles that are possibly receiving shadows will have to compute for the given pixel 
+the number of capsules that directly affect it. The performance benefit of Capsule Shadows comes by only testing these tiles that are affected and by having fewer capsules to compute.
+
+[OBJECT:ComparisonSlider]
+	[PARAM:before]
+	![Standard Shdow Mapping](DirectShadowMapping.png)(w:500)
+	[/PARAM]
+	[PARAM:after]
+	![Capsule Direct Shadows Enabled](DirectCapsuleShadows.png)(w:500)
+	[/PARAM]
+[/OBJECT]
+
+The performance benefit of Capsule Shadows is determined by the number of capsule bodies that are being used and the tiles influenced. Whereas, standard shadow mapping is more expensive because
+it needs to consider all the triangles in the shadow map, which means you can only get sharp shadows that have a little bit of blur to them.
+
+
+## Limitations
+
+* Requires DX-11 due to using tiled deferred implementation using compute shaders 
+* Self-shadowing artifacts for arbitrary mesh shapes 
+* Can only use Spheres and Sphylls for capsule representation
+* Only support for Skeletal Meshes, but could be made to work with static meshes
+* When capsule shadows get so soft that they become ambient occlusion, there's an artifact in the shadowing that causes a hard line. 
\ No newline at end of file
diff --git a/Engine/Documentation/Source/Engine/Rendering/LightingAndShadows/CapsuleShadows/Reference/CapsuleShadowsReference.INT.udn b/Engine/Documentation/Source/Engine/Rendering/LightingAndShadows/CapsuleShadows/Reference/CapsuleShadowsReference.INT.udn
new file mode 100644
--- /dev/null
+++ b/Engine/Documentation/Source/Engine/Rendering/LightingAndShadows/CapsuleShadows/Reference/CapsuleShadowsReference.INT.udn
@@ -0,0 +1,238 @@
+﻿Availability: Docs
+Title: Capsule Shadows Reference
+Crumbs: %ROOT%, Engine, 
+Description: Capsule Shadows Reference
+Version: 4.13
+
+[VAR:Topic]
+[OBJECT:Topic]
+	[PARAM:image]
+		![%Engine/Rendering/LightingAndShadows/CapsuleShadows/Reference:title%]()
+	[/PARAM]
+	[PARAM:icon]
+		![](%ROOT%/light_icon.png)(convert:false)
+	[/PARAM]
+	[PARAM:title]
+		%Engine/Rendering/LightingAndShadows/CapsuleShadows/Reference:title%
+	[/PARAM]
+	[PARAM:description]
+		%Engine/Rendering/LightingAndShadows/CapsuleShadows/Reference:description%
+	[/PARAM]
+	[PARAM:path]
+		[RELATIVE:Engine/Rendering/LightingAndShadows/CapsuleShadows/Reference]
+	[/PARAM]
+[/OBJECT]
+[/VAR]
+[VAR:TopicCompact]
+[OBJECT:TopicCompact]
+	[PARAM:image]
+		![%Engine/Rendering/LightingAndShadows/CapsuleShadows/Reference:title%]()
+	[/PARAM]
+	[PARAM:icon]
+		![](%ROOT%/light_icon.png)(convert:false)
+	[/PARAM]
+	[PARAM:title]
+		%Engine/Rendering/LightingAndShadows/CapsuleShadows/Reference:title%
+	[/PARAM]
+	[PARAM:description]
+		%Engine/Rendering/LightingAndShadows/CapsuleShadows/Reference:description%
+	[/PARAM]
+	[PARAM:path]
+		[RELATIVE:Engine/Rendering/LightingAndShadows/CapsuleShadows/Reference]
+	[/PARAM]
+[/OBJECT]
+[/VAR]
+
+
+
+[TOC(start:2 end:2)]
+
+
+
+The following console commands allow you to adjust the quality of various features for Capsule Shadows.
+
+<!---
+| Console Variable | Description |
+| ---------------- | ----------- |
+| [r.CapsuleIndirectConeAngle](#Parameters) |	Light source angle used when the indirect shadow direction is derived from precomputed indirect lighting (no stationary skylight present). Default Value: 0.392699 |
+| [r.CapsuleIndirectShadowMinVisibility](#capsuleindirectshadowminvisibility) |	Minimum visibility caused by capsule indirect shadows. This is used to keep the indirect shadows from going fully black. Default Value: 0.1 |
+| r.CapsuleMaxDirectOcclusionDistance |	Maximum cast distance for direct shadows from capsules. �This has a big impact on performance. Default Value: 400 |
+| r.CapsuleMaxIndirectOcclusionDistance |	Maximum cast distance for indirect shadows from capsules. This has a big impact on performance. Default Value: 200 |
+| r.CapsuleShadowFadeAngleFromVertical |	Angle from vertical up to start fading out the indirect shadow, to avoid self shadowing artifacts. Default Value: 1.047198 |
+| r.CapsuleShadows |	Whether to allow capsule shadowing on skinned components with bCastCapsuleDirectShadow or bCastCapsuleIndirectShadow enabled. |
+| r.CapsuleShadowsFullResolution |	Whether to compute capsule shadows at full resolution. Disabled by Default. |
+| r.CapsuleSkyAngleScale |	Scales the light source angle derived from the precomputed unoccluded sky vector (stationary skylight present) Default Value: 0.6 |
+--->
+
+
+### Capsule Indirect Cone Angle
+
+The **Capsule Indirect Cone Angle** command allows you to specify the light source angle when using precomputed indirect lighting. If a Skylight is used in the scene its mobility must be set to Static in order to use this command. 
+
+Type the following command followed by a value to adjust the light source angle.
+
+	r.CapsuleIndirectConeAngle
+	
+		
+[OBJECT:ComparisonSlider]
+	[PARAM:before]
+	![Indirect Cone Angle = 0.392699 (Default)](IndirectConeAngle_Default.png)(w:600)
+	[/PARAM]
+	[PARAM:after]
+	![Indirect Cone Angle = 0.55](IndirectConeAngle_55.png)(w:600)
+	[/PARAM]
+[/OBJECT]
+
+
+### Capsule Indirect Shadow Min Visibility
+
+The **Capsule Indirect Shadow Minimum Visibility** command allows you to specify a value for the capsule's indirect shadows. This is used to keep the indirect
+shadows from going fully black. Lower values will soften the shadows intensity, whereas higher values will darken the shadows.
+
+Type the following command followed by a value to adjust the indirect shadow's visibility.
+
+	r.CapsuleIndirectShadowMinVisibility
+	
+[OBJECT:ComparisonSlider]
+	[PARAM:before]
+	![Indirect Shadow Minimum Visibility = 0.1 (Default)](DirectIndirectShadowsDefault.png)(w:600)
+	[/PARAM]
+	[PARAM:after]
+	![Indirect Shadow Minimum Visibility = 0.55](IndirectShadMinVis_0_55.png)(w:600)
+	[/PARAM]
+[/OBJECT]
+
+
+
+### Capsule Occlusion Distance
+
+The **Capsule Occlusion Distance** command allows you to specify the maximum distance from the capsule that it can cast shadows. Values can be set for Capsule Shadows that are in directly and indirectly lit areas.
+
+[Region:warning]
+Adjusting this value will have a significant impact on performance.
+[/REGION]
+
+#### Direct Occlusion Distance
+
+Type the following command followed by a value for distance that direct shadows should cast from the capsules.
+
+	r.CapsuleMaxDirectOcclusionDistance 
+
+[OBJECT:ComparisonSlider]
+	[PARAM:before]
+	![Direct Max Occlusion Distance = 400 (Default)](DirectDefault.png)(w:600)
+	[/PARAM]
+	[PARAM:after]
+	![Direct Max Occlusion Distance = 200](MaxDirOccDist_200.png)(w:600)
+	[/PARAM]
+[/OBJECT]
+
+#### Indirect Occlusion Distance
+
+Type the following command followed by a value for distance that indirect shadows should cast from the capsules.
+
+	r.CapsuleMaxIndirectOcclusionDistance 
+		
+[OBJECT:ComparisonSlider]
+	[PARAM:before]
+	![Indirect Max Occlusion Distance = 200 (Default)](IndirectDefault.png)(w:600)
+	[/PARAM]
+	[PARAM:after]
+	![Indirect Max Occlusion Distance = 50 ](MaxIndirOccDist_100.png)(w:600)
+	[/PARAM]
+[/OBJECT]
+
+
+
+### Capsule Shadow Fade Angle From Vertical
+
+The **Capsule Shadow Fade Angle from Vertical** command allows you to set a value to start fading out indirect shadows. This can help avoid self shadowing artifacts.
+
+Type the following command followed by a value for fading the indirect shadows. Higher values will fade out more quickly.
+
+	r.CapsuleShadowFadeAngleFromVertical 
+
+		
+[OBJECT:ComparisonSlider]
+	[PARAM:before]
+	![Shadow Fade Angle from Vertical = 1.047198 (Default)](FadeAngleVertical_Default.png)(w:600)
+	[/PARAM]
+	[PARAM:after]
+	![Shadow Fade Angle from Vertical = 1.9](FadeAngleVertical_1_9.png)(w:600)
+	[/PARAM]
+[/OBJECT]
+
+
+<!---FULL RESOLUTION---
+[Couldn't get this one to show any difference!]
+
+### Capsule Shadows Full Resolution
+
+Whether to compute capsule shadows at full resolution. Disabled by Default.
+
+Type the following to enable full resolution of the Capsule Shadows.
+
+	r.CapsuleShadowsFullResolution 1
+
+<!---		
+[OBJECT:ComparisonSlider]
+	[PARAM:before]
+	![Indirect Cone Angle = 0.392699 (Default)](IndirectCS_Default.png)(w:600)
+	[/PARAM]
+	[PARAM:after]
+	![Indirect Cone Angle = 0.75](CapsuleIndirectConeAngle_75.png)(w:600)
+	[/PARAM]
+[/OBJECT]
+---/FULL RESOLUTION-->
+
+--->
+
+### Capsule Sky Angle
+
+The **Capsule Sky Angle** commands allow you to specify minimum light source angle for the sky and the scale. The commands require you to have a 
+Stationary Skylight present in your scene and an unoccluded sky to see the results. 
+
+
+#### Mininmum Sky Angle
+
+Type the following to set the minimum value for the Stationary Skylight's light source angle.
+
+	r.CapsuleMinSkyAngle
+
+		
+[OBJECT:ComparisonSlider]
+	[PARAM:before]
+	![Min Sky Angle = 15.0 (Default)](DirectIndirectShadowsDefault.png)(w:600)
+	[/PARAM]
+	[PARAM:after]
+	![Min Sky Angle = 35.0](MinSkyAngle_35.png)(w:600)
+	[/PARAM]
+[/OBJECT]
+
+
+#### Sky Angle Scale
+
+Type the following to set the scale of the light source angle derived from the Stationary Skylight.
+
+	r.CapsuleShadowsSkyAngleScale
+
+		
+[OBJECT:ComparisonSlider]
+	[PARAM:before]
+	![Sky Angle Scale = 0.6 (Default)](DirectIndirectShadowsDefault.png)(w:600)
+	[/PARAM]
+	[PARAM:after]
+	![Sky Angle Scale = 0.1](SkyAngleScale_0_1.png)(w:600)
+	[/PARAM]
+[/OBJECT]
+
+
+
+
+
+
+
+
+
+
+
diff --git a/Engine/Documentation/Source/Engine/Rendering/VisibilityCulling/CullingBestPractices/CullingBestPractices.INT.udn b/Engine/Documentation/Source/Engine/Rendering/VisibilityCulling/CullingBestPractices/CullingBestPractices.INT.udn
new file mode 100644
--- /dev/null
+++ b/Engine/Documentation/Source/Engine/Rendering/VisibilityCulling/CullingBestPractices/CullingBestPractices.INT.udn
@@ -0,0 +1,90 @@
+﻿Availability: Docs
+Title: Visibility Culling Best Practices
+Crumbs: %ROOT%, Engine, 
+Description: Visibility Culling Strategies and Best Practices
+Version: 4.13
+
+[VAR:Topic]
+[OBJECT:Topic]
+	[PARAM:image]
+		![%Engine/Rendering/VisibilityCulling:title%]()
+	[/PARAM]
+	[PARAM:icon]
+		![](%ROOT%/light_icon.png)(convert:false)
+	[/PARAM]
+	[PARAM:title]
+		%Engine/Rendering/VisibilityCulling:title%
+	[/PARAM]
+	[PARAM:description]
+		%Engine/Rendering/VisibilityCulling:description%
+	[/PARAM]
+	[PARAM:path]
+		[RELATIVE:Engine/Rendering/VisibilityCulling]
+	[/PARAM]
+[/OBJECT]
+[/VAR]
+[VAR:TopicCompact]
+[OBJECT:TopicCompact]
+	[PARAM:image]
+		![%Engine/Rendering/VisibilityCulling:title%]()
+	[/PARAM]
+	[PARAM:icon]
+		![](%ROOT%/light_icon.png)(convert:false)
+	[/PARAM]
+	[PARAM:title]
+		%Engine/Rendering/VisibilityCulling:title%
+	[/PARAM]
+	[PARAM:description]
+		%Engine/Rendering/VisibilityCulling:description%
+	[/PARAM]
+	[PARAM:path]
+		[RELATIVE:Engine/Rendering/VisibilityCulling]
+	[/PARAM]
+[/OBJECT]
+[/VAR]
+
+[TOC(start:2 end:2)]
+
+
+
+Below are some good practices to follow when designing and working with visibility culling for your game. 
+
+## Large View Distances
+
+Large view distances can cause an impact on performance due to the number of Actors that need to be rendered. By blocking these large view distances with large occluders you can block Actors that would have been rendered otherwise.
+This can reduce the number of queries that need to be sent to the GPU by having fewer Actors to check. In turn, you get better performance by having fewer Actors to check and render.
+
+* Use large occluders where possible. 
+* Use Cull Distance Volumes to remove any smaller Actors that will not be easily visible at far distances, even if they are occluded from view.
+* Limit the amount of overdraw you get from Translucent Actors, since they will not occlude.
+
+## "Popping" in/Lag from GPU Queries
+
+This is a common issue with using visibility culling for the view frustum and hardware occlusion queries since they are being read back from the GPU one frame later. This is often 
+noticeable when turning your player camera very fast or when you come around a corner in darkly lit areas that have a contrasting lit area behind them. 
+
+* Use other geometry to help hide this artifact. For something like caves or darkened areas with bright outsides use a large occluder mesh that can sit just outside of these pieces to hide the popping in and block the light that makes it evident.
+* Increase the Bounds Scale of the Actor that is causing the offending behavior. 
+* Build geometry in a way that does not have too many small pieces that need to be assembled. This is especially true for walls, floors, and ceilings.
+
+## Asset Construction
+
+When creating geometry for your game it is a good idea to break your geometry up into modular pieces for design purposes and reusability. However, it is possible to modularize too much in the design 
+process by using pieces that are too small to make up larger geometry. Because of this, each actor will need to be have its visibility state checked. It's best to find a balance that makes sense
+for your game design and optimize for the best performance. Reducing the number of queries needed to check visibility is key for optimization and to help reduce the "Popping" in that can happen. 
+
+* Use larger geometry when crafting your levels for structural surfaces rather than breaking it up into many smaller pieces. There is more impact on the GPU testing multiple smaller Actors than a single Actor.
+* "Popping" in of Actors is reduced when considering the construction of the mesh.
+
+
+## Consider Occlusion vs Distance Culling Cost
+
+Occlusion and visibility culling methods cost much more than distance culling since these require GPU queries, so it is a good idea to setup Cull Distance Volumes for distant stuff even if they will 
+be occluded by other Actors from a certain part of the scene.
+
+* Make good use of Cull Distance Volumes to reduce occlusion culling costs.
+* Stream out levels where possible instead of relying on distance or occlusion culling.
+* Use the `Stat Initviews` command to keep performance balance in check.
+
+For information about [Level Streaming](Engine/LevelStreaming) see how to set this up in this page.
+
diff --git a/Engine/Documentation/Source/Engine/Rendering/VisibilityCulling/CullingOverview/CullingOverview.INT.udn b/Engine/Documentation/Source/Engine/Rendering/VisibilityCulling/CullingOverview/CullingOverview.INT.udn
new file mode 100644
--- /dev/null
+++ b/Engine/Documentation/Source/Engine/Rendering/VisibilityCulling/CullingOverview/CullingOverview.INT.udn
@@ -0,0 +1,420 @@
+﻿Availability: Docs
+Title: Visibility Culling Overview
+Crumbs: %ROOT%, Engine, 
+Description: Visibility Culling Overview
+Version: 4.13
+
+[VAR:Topic]
+[OBJECT:Topic]
+	[PARAM:image]
+		![%Engine/Rendering/VisibilityCulling:title%]()
+	[/PARAM]
+	[PARAM:icon]
+		![](%ROOT%/light_icon.png)(convert:false)
+	[/PARAM]
+	[PARAM:title]
+		%Engine/Rendering/VisibilityCulling:title%
+	[/PARAM]
+	[PARAM:description]
+		%Engine/Rendering/VisibilityCulling:description%
+	[/PARAM]
+	[PARAM:path]
+		[RELATIVE:Engine/Rendering/VisibilityCulling]
+	[/PARAM]
+[/OBJECT]
+[/VAR]
+[VAR:TopicCompact]
+[OBJECT:TopicCompact]
+	[PARAM:image]
+		![%Engine/Rendering/VisibilityCulling:title%]()
+	[/PARAM]
+	[PARAM:icon]
+		![](%ROOT%/light_icon.png)(convert:false)
+	[/PARAM]
+	[PARAM:title]
+		%Engine/Rendering/VisibilityCulling:title%
+	[/PARAM]
+	[PARAM:description]
+		%Engine/Rendering/VisibilityCulling:description%
+	[/PARAM]
+	[PARAM:path]
+		[RELATIVE:Engine/Rendering/VisibilityCulling]
+	[/PARAM]
+[/OBJECT]
+[/VAR]
+
+
+[TOC(start:2 end:2)]
+
+
+
+
+Visibility culling uses a process to gather all the Actors in the scene and then will use the available culling methods to test whether the it is visible or not. The process of checking 
+an Actor's visibility state is used to increase the culling effecitveness by reducing the number of queries that need to be sent to the GPU to test. 
+
+In order to understand how this culling process works it is necessary to first understand how an Actor's bounds is used to test its visibility state, 
+what the view frustum does, and lastly what makes an Actor an occluder or occludee. 
+
+## Testing an Actor's Visibility State
+
+For visibility culling to work there needs to be a way to check if the Actor is visible within the Camera's view. This is done by using the Actor's bounds and is an important part about what makes
+visibility culling work. To understand how this culling process works it is necessary to understand how the Actor's bounds, Camera View Frustum, and Occlusion work together to set an Actor's visibility
+state.
+
+## Actor Bounds
+
+Each Actor has its own bounds that is made up of a bounding box and and sphere surrounding the geometry of the mesh. The bounds are automatically calculated based on the mesh, but they can be 
+edited via the Static Mesh Editor, Persona, or per-instance via the Detail's panel in your scene when the need arises. For visibility and occlusion culling purposes it is used to test whether 
+an Actor is visible within the Camera's field of view.
+
+![](BoundsEditorView1.png)(w:350)
+[REGION:caption]
+Actor's bounding box and sphere
+[/REGION]
+
+#### Setting the Bounds Scale
+
+You can set the Static and Skeletal Meshes bounds by using the following options. 
+
+[REGION:imagetable]
+| ![](BoundScaleSetting.png)(w:400) | ![](BoundsScaleSettingSMEditor.png)(w:400) |
+| --- | --- |
+| Scene Per-Instance Detail's panel | Static and Skeletal Mesh Editor Detail's panel | 
+[/REGION]
+
+[REGION:note]
+The per-instance setting for the **Bounds Scale** will use a multiplier that scales evenly for the Actor's bounds, whereas the Static Mesh Editor and Persona will allow you to set specific XYZ coordinate values.
+[/REGION]
+
+To view the bounds of an Actor in your scene, you must first go to the viewport and **Show** > **Advanced** > **Bounds** or use the console command `Show Bounds`. Then select the Actor in the scene
+to see the bounds visualization.
+
+[OBJECT:ComparisonSlider]
+ [PARAM:before]
+ ![Bounds Scale: 1.0](Bounds2.png)(w:500)
+ [/PARAM]
+ [PARAM:after]
+ ![Bounds Scale: 1.25](Bounds3.png)(w:500)
+ [/PARAM]
+[/OBJECT]
+[REGION:caption]
+_Setting the Bounds Scale in the Detail's panel of the selected Actor._
+[/REGION]
+
+[REGION:tip]
+When you increase the bounds of an Actor you do not have to increase by whole numbers. You can increase with decimal values incrementally until you get the results you are wanting. For example, 
+increase the Bounds Scale to 1.15 rather than increasing to 2.0. 
+[/REGION]
+
+
+## What is the Camera's View Frustum?
+
+The Camera's view frustum is the viewing area within 3D space that is represented by a pyramidal shape from its viewpoint. 
+
+[REGION:imagetable]
+| [INCLUDE:#cameraview1] |
+| ---------------------- |
+| Camera Field of View (FOV) set to 90 degrees (Default) |
+[REGION:caption]
+_Click images for full size_
+[/REGION]
+[/REGION]
+
+This pyramidal shape represents the **Field of View** (FOV) as seen from the Camera's perspective. This perspective shows the various Actors in the scene that are within the view frustum.
+The image below contains all the elements that make up the scene above. This can include Actors that are not visibile within the view frustum, like the Spot Light that is lighting the scene, 
+or other Actors not visible within the view frusum.  
+
+[REGION:imagetable]
+| [INCLUDE:#sceneview1]| 
+| --- |
+| Scene View with Camera Attributes  |
+[/REGION]
+[REGION:caption]
+_Click images for full size_
+[/REGION]
+
+This image shows the complete scene that includes all the Actors visible within the FOV and the ones that are not, like the Spot Light. The Camera's attributes are listed below to breakdown the different
+elements of the Camera.
+
+| Number | Description |
+| ------ | ----------- |
+| **1**  | **Field of View (FOV):** The width of the view angle measured in degrees. |
+| **2**  | **Near Clip Plane:** The nearest point to the camera that will render. |
+| **3**  | **Camera View Frustum:** The area between the near and far clip planes. Anything within this view is considered for rendering. |
+| **4**  | **Far Clip Plane:** The farthest point from the camera that will render. This is controlled by scene depth instead of a set distance that clips the view. |
+
+<!---
+[EXCERPT:CameraView1]
+[REGION:lightbox]
+[![](CameraView.png)(w:600)](CameraView.png)
+[/REGION]
+[/EXCERPT:CameraView1]
+
+[EXCERPT:SceneView1]
+[REGION:lightbox]
+[![](FrustumView.png)(w:600)](FrustumView.png)
+[/REGION]
+[/EXCERPT:SceneView1]
+--->
+
+### Camera Field of View
+The view frustum is controlled by its **Field of View** (FOV). Higher values will give wider viewing angle and a lower value will restrict the available viewing angle.
+
+For more information on [using and setting up Cameras](Gameplay/HowTo/UsingCameras) you can read about them in this page.
+
+[REGION:note]
+The Camera's view frustum visualization is not enabled by default. You must first go to the Viewport and select **Show** > **Advanced** > **Camera Frustum** to enable it.
+[/REGION]
+
+[OBJECT:EmbeddedVideo]
+ [PARAMLITERAL:width]
+ 640
+ [/PARAMLITERAL]
+ [PARAMLITERAL:height]
+ 360
+ [/PARAMLITERAL]
+ [PARAMLITERAL:videoid]
+ RqTqc0dLKPM
+ [/PARAMLITERAL]
+[/OBJECT]
+
+
+
+
+## What is Occlusion?
+
+Occlusion happens when you have an Actor that can block or become blocked by another Actor. These Actors are either occluders, occludees, or both. 
+ 
+* An occluder is any Actor that can occlude another Actor from being visible to the view frustum.
+* An occludee is any Actor that is occluded by another Actor from being visible to the view frustum. If this Actor is large enough it can also be an occluder for smaller Actors that it blocks from the view frustum.
+
+In the example below, the Camera represents the view of the scene. Behind the wall lies two Actors that cannot be seen because they are blocked from the camera view. The statue and barrel both would be occludees since 
+they are being blocked by the wall that is an occluding them. Also, remember that both the barrel and statue could be occluders if they were blocking anything behind them. Every occluder can be an occludee if it is large 
+enough to block another actor.
+
+[REGION:imagetable]
+| [INCLUDE:#OccluderVsOccludee1] | [INCLUDE:#OccluderVsOccludeeCV] |
+| --- | --- | 
+| Occluder = Wall, Occludee = Barrels and Statue  | Camera FOV shows the wall that occludes the barrels and statue. |
+[/REGION]
+[REGION:caption]
+_Click images for full size_
+[/REGION]
+
+<!---
+[EXCERPT:OccluderVsOccludee1]
+[REGION:lightbox]
+[![](OccluderVsOccludee.png)(w:450)](OccluderVsOccludee.png)
+[/REGION]
+[/EXCERPT:OccluderVsOccludee1]
+
+[EXCERPT:OccluderVsOccludeeCV]
+[REGION:lightbox]
+[![](OccluderVsOccludeeCameraView.png)(w:450)](OccluderVsOccludeeCameraView.png)
+[/REGION]
+[/EXCERPT:OccluderVsOccludeeCV]
+--->
+
+<!---
+### Visibility Culling
+
+The following types of culling are used to determine the current visibility state of an Actor in the scene.
+
+| Types of Culling | Description |
+| ---------------- | ----------- |
+| **Distance** | The Actor's set distance is used to determine its visibility state. |
+| **View Frustum** | This is used to remove any Actors that do not have its bounds visible within the Camera's view frustum. Anything outside of this viewing area will not be rendered.  |
+| **Occlusion Queries** | This is used to remove any additional Actors that are within the Camera's view frustum that is blocked by another from being viewed. This is done by testing if the Actor is visible within the scene's depth buffer. If it is visible it will be rendered, and if it is not it will be considered occluded and not rendered. |
+
+
+For more information on [Buffer Visualizations](Engine/UI/LevelEditor/Viewports/ViewModes#buffervisualization) you can read about them on this page.
+--->
+
+## Visibility Culling Methods
+
+Once the data has been culled from the scene it is tested against the available culling methods that have been implemented to test the visibility state of the Actors. The visibility culling methods work
+together and in the case of mobile offer a solution where dynamic occlusion is not possible.
+
+These culling methods listed here are in the order they are applied, which is also from least to most expensive in terms of rendering thread cost.
+
+* Distance culling (Per-instance and Cull Distance Volumes)
+* View frustum culling
+* Precomputed occlusion culling
+* Dynamic occlusion culling (hardware occlusion queries)
+
+
+### Distance Culling (Per-instance)
+
+Each Actor in the scene can have its own set draw distance that will allows minimum and maximum distance for it to be rendered visible.
+
+![](DistanceCullSettings.png)
+
+| **Property** | **Description** |
+| ------------ | --------------- |
+| **Min Draw Distance** | The minimum distance at which the primitive should be rendered. This is measured in world space units from the center of the primitive's bounding sphere to the camera position. |
+| **Desired Max Draw Distance** | The max draw distance exposed to Level Designers. The real max draw distance is the min (disregarding 0) of this and volumes effecting this object. |
+| **Current Max Draw Distance** | This is not an editable property. This field shows the cull distance that is set by the CullDistanceVolume that the primitive is residing inside of. If it is set to 0 no Cull Distance Volume has been set so the value will be 0. |
+
+In the examples below, the **Min Draw Distance** and **Desired Max Draw Distance** have been set for the cube. These values determine the closest distance to render the Actor or the farthest distance that 
+the Actor can be rendered from the Camera's view.
+
+* **Min Draw Distance:** 500
+* **Desired Max Draw Distance:** 1000
+
+[REGION:note]
+Note that the **Desired Max Draw Distance** here will override any higher value set in a Cull Distance Volume. See the [Cull Distance Volume](1_TimH\2_IN_PROGRESS\VisibilityCulling\CullDistanceVolume) page for more information.
+[/REGION]
+
+[OBJECT:EmbeddedVideo]
+ [PARAMLITERAL:width]
+ 640
+ [/PARAMLITERAL]
+ [PARAMLITERAL:height]
+ 360
+ [/PARAMLITERAL]
+ [PARAMLITERAL:videoid]
+ LiiJ13LulIU
+ [/PARAMLITERAL]
+[/OBJECT]
+
+
+<!---
+[REGION:imagetable]
+| [INCLUDE:#Units1500] | [INCLUDE:#Units1000] | [INCLUDE:#Units500] |
+| --- | --- | --- |
+| More than Desired Max Draw Distance | Within Range of Min and Max Draw Distance | Less than Min Draw Distance | 
+[/REGION]
+
+To help illustrate what is happening in the video above, the images here show a top-down view and mark the distance from the **Camera** Actor to the center point of the **Cube** Actor. 
+The image directly below the top-down view shows the camera's view.
+
+--->
+
+Since the **Desired Max Draw Distance** is set to 1000 and the **Min Draw Distance** is set to 500 the cube Actor will only be visible when the Camera moves within this distance from the Actor's bounding box
+center. When the Camera moves outside of this space the Actor will no longer be rendered. 
+
+<!--
+[EXCERPT:Units1500]
+[REGION:lightbox]
+[![](1500units.png)(w:300)](1500Units.png)
+[/REGION]
+[/EXCERPT:Units1500]
+
+[EXCERPT:Units1000]
+[REGION:lightbox]
+[![](1000units.png)(w:300)](1000Units.png)
+[/REGION]
+[/EXCERPT:Units1000]
+
+[EXCERPT:Units500]
+[REGION:lightbox]
+[![](500units.png)(w:300)](500Units.png)
+[/REGION]
+[/EXCERPT:Units500]
+-->
+
+
+### Cull Distance Volume
+
+Cull Distance Volumes allow you to specify a set of sizes and distances in an array that will automatically remove any Actor based on its bounds size at the corresponding distance. These volumes are useful
+for large outdoor scenes that have detailed interiors. 
+
+For more information on using and setting up [Cull Distance Volumes](1_TimH\2_IN_PROGRESS\VisibilityCulling\CullDistanceVolume) you can read about it in this page.
+
+
+### View Frustum Culling
+
+The Camera's view frustum is used to test whether the culled data from the scene is visible within its viewing area. This process uses the Actor's bounding box and sphere to check if it is visible within the view. If not, the 
+Actor will not be rendered.
+
+The following example demonstrates how the view frustum considers Actors that are visible within the viewing area.
+
+* Yellow represents FOV of the view frustum and the area that is no occluded by other Actors.
+* Red represents the part of the view frustum that is occluded due to other Actors (ie. walls and large objects).
+
+[REGION:imagetable]
+| [INCLUDE:#Camera1] | [INCLUDE:#Camera2] | [INCLUDE:#Camera3] | [INCLUDE:#Camera4] |
+| ------------------ | ------------------ | ------------------ | ------------------ |
+|          1         |          2         |          3         |          4         | 
+[/REGION]
+
+| **Number** | **Description** |
+| ---------- | --------------- |
+| **1** | From this perspective there are large walls that act as occluders obstructing all but the view along the hallway. Because of this only the cube is visible. This means that the sphere and cone are occludees of the wall and will not be rendered. |
+| **2** | As the Camera moves forward there is more of the scene that is exposed to the view frustum. The walls no longer fully occlude the sphere and cone now that we can see parts of them around the corners of the wall. This means that they are now being fully rendered even though we can only see a small part of them. |
+| **3** | Next the as the Camera turns to the left we now have the sphere and cube within the view frustum. The cone and walls out of view will no longer be rendered. |
+| **4** | Now as the Camera is turned back to the right the cone and cube are within the view and the walls and sphere that are out of view are no longer being rendered. |
+
+[REGION:note]
+This is a generalized example of how visibility culling works and there is a number of caveats and reasons to how things can be rendered that are not covered here. For additional tips and strategies to use
+with visibility culling see the [Culling Strategies for Best Practices](Engine\Rendering\VisibilityCulling\CullingBestPractices) section in this page.
+[/REGION]
+
+<!--
+[EXCERPT:Camera1]
+[REGION:lightbox]
+[![](Camera1.png)(w:215)](Camera1.png)
+[/REGION]
+[/EXCERPT:Camera1]
+
+[EXCERPT:Camera2]
+[REGION:lightbox]
+[![](Camera2.png)(w:215)](Camera2.png)
+[/REGION]
+[/EXCERPT:Camera2]
+
+[EXCERPT:Camera3]
+[REGION:lightbox]
+[![](Camera3.png)(w:215)](Camera3.png)
+[/REGION]
+[/EXCERPT:Camera3]
+
+[EXCERPT:Camera4]
+[REGION:lightbox]
+[![](Camera4.png)(w:215)](Camera4.png)
+[/REGION]
+[/EXCERPT:Camera4]
+-->
+
+[TOC(start:2 end:2)]
+
+### Precomputed Visibility Volume
+
+<!---![](PrecomputedVisibilityVolume.png)(w:500)--->
+
+Precomputed Visibility Volumes use visibility cells that are placed on surfaces after a static lighting build. These visibility cells use the Camera view to change the visibility state of an Actor.
+This culling method is good for having visibility culling on mobile devices where hardware occlusion queries cannot be issued to check the visibility state. 
+
+For more information on using and setting up [Precomputed Visibility Volumes](1_TimH\2_IN_PROGRESS\VisibilityCulling\PrecomputedVisibilityVolume) you can read about them in this page.
+
+### Dynamic Occlusion
+
+Dynamic occlusion queries are performed last in the chain of culling methods by only looking at the Actors that have been determined to be within the view frustum of the Camera. The Actors bounds that are 
+within the view frustum are then tested against the Scene Depth Buffer for their visibility by sending them to the GPU and reading them back one frame later. If an Actor's bounds is not visible within the 
+scene depth it will not be rendered. This allows the visibility state for any Actor occluded by another Actor to not be rendered. This method can allow for farther view distances because it relies on appoximate
+distances rather than setting a specific value for the far clip plane's distance.
+
+#### Hierarchical Z Buffer
+
+[REGION:warning]
+The Hierarchical Z-Buffer occlusion implementation is still considered experimental. There are some known issues that have not yet been addressed, like some meshes being culled when they should not be. 
+[/REGION]
+
+The idea behind Hierarchical Z-Buffer occlusion is to represent the depth of the scene with multiple resolutions to determine the visibility of an Actor. This is done by creating mipmapped version 
+of the render target for the scene depth which has texels that identify the depth in the scene for the Actors tested. The Actors can be tested for occlusion by comparing the depth value of the Actor's bounds
+to that in the render target that was captured. By using a lower resolution mip level this will use a more conservative check for the bounds of an Actor which leads to fewer texture fetches when sampling from
+a mip level. 
+
+[REGION:note]
+By rendering the scene to a lower resolution than the actual scene resolution you risk falsely culling Actors that should be visible.  
+[/REGION]
+
+Hierarchical Z-Buffer Occlusion Culling can be enabled by using the following console command: 
+
+	r.HZBOcclusion
+
+* 0 will disable Hierarchical Z-Buffer occlusion
+* 1 will enable Hierarchical Z-Buffer occlusion
+* 2 will force the Hierarchical Z-Buffer system, which will override any rendering platform preferences
+
+
diff --git a/Engine/Documentation/Source/Engine/Rendering/VisibilityCulling/CullingPerformance/CullingPerformance.INT.udn b/Engine/Documentation/Source/Engine/Rendering/VisibilityCulling/CullingPerformance/CullingPerformance.INT.udn
new file mode 100644
--- /dev/null
+++ b/Engine/Documentation/Source/Engine/Rendering/VisibilityCulling/CullingPerformance/CullingPerformance.INT.udn
@@ -0,0 +1,141 @@
+﻿Availability: Docs
+Title: Visibility Culling Settings and Performance 
+Crumbs: %ROOT%, Engine, 
+Description: Visibility Culling Settings and Performance
+Version: 4.13
+
+[VAR:Topic]
+[OBJECT:Topic]
+	[PARAM:image]
+		![%Engine/Rendering/VisibilityCulling:title%]()
+	[/PARAM]
+	[PARAM:icon]
+		![](%ROOT%/light_icon.png)(convert:false)
+	[/PARAM]
+	[PARAM:title]
+		%Engine/Rendering/VisibilityCulling:title%
+	[/PARAM]
+	[PARAM:description]
+		%Engine/Rendering/VisibilityCulling:description%
+	[/PARAM]
+	[PARAM:path]
+		[RELATIVE:Engine/Rendering/VisibilityCulling]
+	[/PARAM]
+[/OBJECT]
+[/VAR]
+[VAR:TopicCompact]
+[OBJECT:TopicCompact]
+	[PARAM:image]
+		![%Engine/Rendering/VisibilityCulling:title%]()
+	[/PARAM]
+	[PARAM:icon]
+		![](%ROOT%/light_icon.png)(convert:false)
+	[/PARAM]
+	[PARAM:title]
+		%Engine/Rendering/VisibilityCulling:title%
+	[/PARAM]
+	[PARAM:description]
+		%Engine/Rendering/VisibilityCulling:description%
+	[/PARAM]
+	[PARAM:path]
+		[RELATIVE:Engine/Rendering/VisibilityCulling]
+	[/PARAM]
+[/OBJECT]
+[/VAR]
+
+
+[TOC(start:2 end:2)]
+
+
+
+
+## Project Settings
+
+The **Project Settings** contain various settings that will influence how culling works for your game. You can get to these settings by first going to the main menu and select 
+**Edit** > **Project Settings**. When the window opens you can select the **Rendering** category to see the available settings under the section named **Culling**.
+
+![](ProjectSettingsCulling.png)
+
+| Setting | Description |
+| --- | --- |
+| **Occlusion Culling** | Allows occluded meshes to be culled and not rendered. If this is disabled all occlusion culling will be disabled for the project which can have performance implications. |
+| **Min Screen Radius for EarlyZPass** | Screen radius at which objects are culled for early Z pass. Larger values can improve performance but very large values can degrade performance if large occluders are not rendered. |
+| **Warn about no precomputed visibility** | Displays a warning when no precomputed visibility data is available for the current camera location. This can be helpful if you are making a game that relies on precomputed visibility. |
+
+[REGION:tip]
+​If you are developing a mobile game it would be beneficial to enable the option for **Warn about no precomputed visibility** as a good reminder for your Level Designers when developing for mobile, since it does not currently use dynamic 
+occlusion culling.
+[/REGION]
+
+
+
+## Relevant Stats for Performance
+
+[REGION:warning]
+Keep in mind that when testing visibility and occlusion culling there are a couple of things to be aware of.
+
+* Occlusion culling is disabled in **Wireframe** view mode.
+* Use the console command **Stat Initviews** in PIE or Standalone games for more accurate results, because geometry for editor Actors that represent Lights, Cameras, and so on will be included when calculating visibility and occlusion culling.
+[/REGION]
+
+In real-time rendering for games, it is important to manage your scene for optimization by keeping check of the stats generated for your scene that are relevant to visibility and occlusion culling.
+
+Use the following console command to view the stats relevant to visibility and occlusion culling for your scene: 
+
+	Stat Initviews
+
+
+[REGION:lightbox]
+![](StatInitViews.png)
+[/REGION]
+[REGION:caption]
+The scene's stats above were taken from the [Infinity Blade Grasslands](https://www.unrealengine.com/marketplace/infinity-blade-plain-lands) created by EPIC Games and is freely available on the EPIC Games Marketplace.
+[/REGION]
+
+The **Stat Initiviews** window is broken up into two sections for `Cycle Counters` and `Counters`. The Cycle Counter's stats focus on the number of cycles in milliseconds it took to process, while the Counter's stats simply
+add up the primitives that are processed.
+
+[REGION:tip]
+`Visible Static Mesh Elements` is the single biggest contributor to rendering thread time and should be carefully watched and optimized.
+[/REGION]
+
+| Stat Initviews |  |
+| --- | --- |
+| **Cycle Counters** | [INCLUDE:#cyclecounter] |
+| **Counters** | [INCLUDE:#counters] |
+
+
+<!--
+[EXCERPT:CycleCounter]
+| Stat | Description|
+| --- | --- |
+| **View Visibility** | The amount of frame time spent to process Actor visibility queries. |
+| **Occlusion Cull** | The amount of frame time spent to query the scene for Actors within the view frustum that were occluded by other Actors. |
+| **Frustum Cull** | The amount of frame time spent to query if an Actor's bounds were within the view frustum. |
+| **Decompress Occlusion** | The amount of frame time it took to load the precomputed visibility. |
+[/EXCERPT:CycleCounter]
+
+[EXCERPT:Counters]
+| Stat | Description|
+| --- | --- |
+| **Processed Primitives** | This is the total number of Actors that are being processed in the scene.  |
+| **Frustum Culled Primitives** | This is the number of primitives culled when not within the view frustum's FOV. |
+| **Occluded Primitives** | The number of Actors that are occluded by other Actors from within the view frustum's FOV. |
+| **Occlusion Queries** | This is the number of scene occlusion queries sent to the GPU based on the number of Actors visible within the view frustum. |
+| **Visible Dynamic Primitives** | This is the number of dynamic primitives in the scene that have their mobility set to Movable, like Particle Systems and Foliage instances. |
+| **Visible Static Mesh Elements** | This is the number of visible meshes that have their mobility set to Static within the scene. |
+[/EXCERPT:Counters]
+-->
+
+## Useful Console Commands
+
+UE4 provides some [Console Commands](https://docs.unrealengine.com/latest/INT/GettingStarted/RunningUnrealEngine/index.html#usefulin-gamecommands) that can help with occlusion culling.
+
+| Console Variable | Description |
+| --- | --- |
+| **r.OneFrameThreadLag** | This controls when occlusion queries are rendered. Rendering before the base pass may give worse occlusion (because not all occluders generally render in the earlyzpass). However, it may reduce CPU waiting for query result stalls on some platforms and increase overall performance. 0: After BasePass. 1: After EarlyZPass, but before BasePass. |
+| **r.AllowOcclusionQueries** | If zero, occlusion queries will not be used to cull primitives. |
+| **r.DownsampledOcclusionQueries** | Whether to issue occlusion queries to a downsampled depth buffer. |
+| **r.NumBufferedOcclusionQueries** | Number of frames to buffer occlusion queries (including the current render thread frame). More frames reduce the chance of stalling the CPU waiting for the results, but increases out of date query artifacts. |
+
+
diff --git a/Engine/Documentation/Source/Engine/Rendering/VisibilityCulling/CullingTipsTricks/CullingTipsTricks.INT.udn b/Engine/Documentation/Source/Engine/Rendering/VisibilityCulling/CullingTipsTricks/CullingTipsTricks.INT.udn
new file mode 100644
--- /dev/null
+++ b/Engine/Documentation/Source/Engine/Rendering/VisibilityCulling/CullingTipsTricks/CullingTipsTricks.INT.udn
@@ -0,0 +1,123 @@
+﻿Availability: Docs
+Title: Visibility Culling Tips & Tricks
+Crumbs: %ROOT%, Engine, 
+Description: Tips & Tricks you can use to troubleshoot your scenes.
+Version: 4.13
+
+[VAR:Topic]
+[OBJECT:Topic]
+	[PARAM:image]
+		![%Engine/Rendering/VisibilityCulling:title%]()
+	[/PARAM]
+	[PARAM:icon]
+		![](%ROOT%/light_icon.png)(convert:false)
+	[/PARAM]
+	[PARAM:title]
+		%Engine/Rendering/VisibilityCulling:title%
+	[/PARAM]
+	[PARAM:description]
+		%Engine/Rendering/VisibilityCulling:description%
+	[/PARAM]
+	[PARAM:path]
+		[RELATIVE:Engine/Rendering/VisibilityCulling]
+	[/PARAM]
+[/OBJECT]
+[/VAR]
+[VAR:TopicCompact]
+[OBJECT:TopicCompact]
+	[PARAM:image]
+		![%Engine/Rendering/VisibilityCulling:title%]()
+	[/PARAM]
+	[PARAM:icon]
+		![](%ROOT%/light_icon.png)(convert:false)
+	[/PARAM]
+	[PARAM:title]
+		%Engine/Rendering/VisibilityCulling:title%
+	[/PARAM]
+	[PARAM:description]
+		%Engine/Rendering/VisibilityCulling:description%
+	[/PARAM]
+	[PARAM:path]
+		[RELATIVE:Engine/Rendering/VisibilityCulling]
+	[/PARAM]
+[/OBJECT]
+[/VAR]
+
+
+
+The following are some Tips & Tricks you can use to help you troubleshoot and test your scene when you suspect problems with visibility or occlusion culling.
+
+
+## Visualize Occluded Primitives
+
+This will enable the visualization mode to draw the bounding boxes of any Actors in your scene that are occluded from the Camera view. A green bounding box indicates the Actor is 
+dynamically occluded and a red bounding box indicates that the Actor is statically occluded when using [Precomputed Visibility Volumes](1_TimH\2_IN_PROGRESS\VisibilityCulling\PrecomputedVisibilityVolume). 
+
+Use the following console command to enable this visualization: 
+
+	r.VisualizeOccludedPrimitives 1
+
+
+[REGION:note]
+This console command will only work while in the editor and not in PIE or Standalone games.
+[/REGION]
+
+In the example below, the visualization of occluded Actors blocked by the wall is compared to the placement of the Actor's with the wall removed. 
+
+[OBJECT:ComparisonSlider]
+ [PARAM:before]
+ ![Visualize Occluded Primitives Enabled (Green Box)](VisOccludedPrimitives1.png)(w:500)
+ [/PARAM]
+ [PARAM:after]
+ ![Actors occluded by the wall](VisOccludedPrimitives2.png)(w:500)
+ [/PARAM]
+[/OBJECT]
+
+
+## Freeze Rendering
+
+This will toggle the rendering state of the scene from the current Camera perspective while allowing you to freely move around the world.
+
+Use the following console command to toggle this mode:
+
+	FREEZERENDERING
+
+
+[OBJECT:EmbeddedVideo]
+ [PARAMLITERAL:width]
+ 640
+ [/PARAMLITERAL]
+ [PARAMLITERAL:height]
+ 360
+ [/PARAMLITERAL]
+ [PARAMLITERAL:videoid]
+ El127q0lv5s
+ [/PARAMLITERAL]
+[/OBJECT]
+
+
+## Game View Mode
+
+You can use the keyboard shortcut for **G** or use the main **Viewport Options** down-arrow to select **Game View** to be able to work in the editor and switch between what you would see in game. 
+This is helpful when you do not want to PIE or launch a game but need to check what the game would look like. This is ideal for checking [Cull Distance Volumes](Engine\Rendering\VisibilityCulling\CullDistanceVolume) 
+and [Precomputed Visibility Volumes](Engine\Rendering\VisibilityCulling\PrecomputedVisibilityVolume) since these do not work unless you are in game or using the Game View.
+
+For more information on the [Game View](Engine\UI\LevelEditor\Viewports\Basics) mode, you can read about it in this page.
+
+
+<!---
+[OBJECT:TopicButtonList]
+ [PARAM:icon]
+  ![](%ROOT%/related_icon.png)(convert:false)
+ [/PARAM]
+ [PARAM:title]
+  Related Links
+ [/PARAM]
+ [PARAM:description]
+ [/PARAM]
+ [PARAM:links]
+  * [](1_TimH\2_IN_PROGRESS\VisibilityCulling\PrecomputedVisibilityVolume)
+  * [](1_TimH\2_IN_PROGRESS\VisibilityCulling\CullDistanceVolume)
+ [/PARAM]
+[/OBJECT]
+--->
\ No newline at end of file
diff --git a/Engine/Documentation/Source/Engine/Rendering/VisibilityCulling/VisibilityCulling.INT.udn b/Engine/Documentation/Source/Engine/Rendering/VisibilityCulling/VisibilityCulling.INT.udn
new file mode 100644
--- /dev/null
+++ b/Engine/Documentation/Source/Engine/Rendering/VisibilityCulling/VisibilityCulling.INT.udn
@@ -0,0 +1,614 @@
+﻿Availability: Docs
+Title: Visibility Culling
+Crumbs: %ROOT%, Engine, 
+Description: How Visibility Culling works in your game.
+Version: 4.13
+
+[VAR:Topic]
+[OBJECT:Topic]
+	[PARAM:image]
+		![%Engine/Rendering/VisibilityCulling:title%]()
+	[/PARAM]
+	[PARAM:icon]
+		![](%ROOT%/light_icon.png)(convert:false)
+	[/PARAM]
+	[PARAM:title]
+		%Engine/Rendering/VisibilityCulling:title%
+	[/PARAM]
+	[PARAM:description]
+		%Engine/Rendering/VisibilityCulling:description%
+	[/PARAM]
+	[PARAM:path]
+		[RELATIVE:Engine/Rendering/VisibilityCulling]
+	[/PARAM]
+[/OBJECT]
+[/VAR]
+[VAR:TopicCompact]
+[OBJECT:TopicCompact]
+	[PARAM:image]
+		![%Engine/Rendering/VisibilityCulling:title%]()
+	[/PARAM]
+	[PARAM:icon]
+		![](%ROOT%/light_icon.png)(convert:false)
+	[/PARAM]
+	[PARAM:title]
+		%Engine/Rendering/VisibilityCulling:title%
+	[/PARAM]
+	[PARAM:description]
+		%Engine/Rendering/VisibilityCulling:description%
+	[/PARAM]
+	[PARAM:path]
+		[RELATIVE:Engine/Rendering/VisibilityCulling]
+	[/PARAM]
+[/OBJECT]
+[/VAR]
+
+
+[TOC(start:2 end:2)]
+
+## What is Visibility Culling?
+
+Visibility culling uses a process to gather all the Actors in the scene and then will use one of the available culling methods to test whether the Actor is visible or not. The process of checking 
+an Actor's visibility state is used to increase the culling effecitveness by reducing the number of queries to the GPU to test. 
+
+In order to understand how visibility and occlusion culling works it is necessary to first understand how an Actors bounding box and sphere is used to test its visibility state, 
+what the view frustum does, and lastly what makes an Actor an occluder or occludee. 
+
+### Actor Bounds
+
+Each Actor has its own bounds. These are made up of a bounding box and sphere that surround the geometry of the mesh. The bounds can be edited via the Static Mesh Editor, Persona, or per-instance via the 
+Detail's panel in your scene. For visibility and occlusion culling purposes it is used to test whether an Actor is visible within the Camera's field of view.
+
+
+![](BoundsEditorView1.png)(w:350)
+[REGION:caption]
+Actor's bounding box and sphere
+[/REGION]
+
+##### Setting the Bounds Scale
+
+You can set the Static and Skeletal Meshes bounds by using the following options.
+
+[REGION:imagetable]
+| ![](BoundScaleSetting.png)(w:400) | ![](BoundsScaleSettingSMEditor.png)(w:400) |
+| --- | --- |
+| Scene Per-Instance Detail's panel | Static and Skeletal Mesh Editor Detail's panel | 
+[/REGION]
+
+[REGION:note]
+The per-instance setting for the **Bounds Scale** will use a multiplier that scales evenly for the Actor's bounds, whereas the Static Mesh Editor and Persona will allow you to set specific XYZ coordinate values.
+[/REGION]
+
+To view the bounds of an Actor, you must first go to the viewport and **Show** > **Advanced** > **Bounds** or use the console command `Show Bounds`. Then when you select the Actor in the scene
+to see the bounds visualization.
+
+[OBJECT:ComparisonSlider]
+ [PARAM:before]
+ ![Bounds Scale: 1.0](Bounds2.png)(w:500)
+ [/PARAM]
+ [PARAM:after]
+ ![Bounds Scale: 1.25](Bounds3.png)(w:500)
+ [/PARAM]
+[/OBJECT]
+[REGION:caption]
+_Setting the Bounds Scale in the Detail's panel of the selected Actor._
+[/REGION]
+
+[REGION:tip]
+When you increase the bounds of an Actor you do not have to increase by whole numbers. You can increase with decimal values incrementally until you get the results you are wanting. For example, 
+increase the Bounds Scale to 1.15 rather than increasing to 2.0. 
+[/REGION]
+
+
+
+### What is the Camera's View Frustum?
+
+The Camera's view frustum is the viewing area within 3D space that is represented by a pyramidal shape from its viewpoint. 
+
+[REGION:imagetable]
+| [INCLUDE:#cameraview1] |
+| ---------------------- |
+| Camera Field of View (FOV) set to 90 degrees (Default) |
+[REGION:caption]
+_Click images for full size_
+[/REGION]
+[/REGION]
+
+This pyramidal shape represents up our **Field of View** (FOV) as seen from the Camera's perspective. This perspective shows our various Actors in the scene that are within this view.
+The image below contains all the elements that make up the scene above. This can include Actors that are not visibile within the view frustum, like the Spot Light that is lighting the scene, 
+or other Actors not visible within the view frusum.  
+
+[REGION:imagetable]
+| [INCLUDE:#sceneview1]| 
+| --- |
+| Scene View with Camera Attributes  |
+[/REGION]
+[REGION:caption]
+_Click images for full size_
+[/REGION]
+
+| Number | Description |
+| ------ | ----------- |
+| **1**  | **Field of View (FOV):** The width of the view angle measured in degrees. |
+| **2**  | **Near Clip Plane:** The nearest point to the camera that will render. |
+| **3**  | **Camera View Frustum:** The area between the near and far clip planes. Anything within this view is considered for rendering. |
+| **4**  | **Far Clip Plane:** The farthest point from the camera that will render. This is controlled by scene depth instead of a set distance that clips the view. |
+
+<!---
+[EXCERPT:CameraView1]
+[REGION:lightbox]
+[![](CameraView.png)(w:600)](CameraView.png)
+[/REGION]
+[/EXCERPT:CameraView1]
+
+[EXCERPT:SceneView1]
+[REGION:lightbox]
+[![](FrustumView.png)(w:600)](FrustumView.png)
+[/REGION]
+[/EXCERPT:SceneView1]
+--->
+
+The view frustum is controlled by its **Field of View** (FOV). Higher values will give wider viewing angle and a lower value will restrict the available viewing angle.
+
+[REGION:note]
+The Camera's view frustum visualization is not enabled by default. You must first go to the Viewport and select **Show** > **Advanced** > **Camera Frustum** to enable it.
+[/REGION]
+
+[OBJECT:EmbeddedVideo]
+ [PARAMLITERAL:width]
+ 640
+ [/PARAMLITERAL]
+ [PARAMLITERAL:height]
+ 360
+ [/PARAMLITERAL]
+ [PARAMLITERAL:videoid]
+ RqTqc0dLKPM
+ [/PARAMLITERAL]
+[/OBJECT]
+
+For more information on [using and setting up Cameras](Gameplay/HowTo/UsingCameras) you can read about them in this page.
+
+
+### What is Occlusion?
+
+Occlusion happens when you have an Actor that can block or become blocked by another Actor. The Actors are either occluders, occludees, or both. 
+ 
+* An occluder is any Actor that can occlude another Actor from being visible to the view frustum.
+* An occludee is any Actor that is occluded by another Actor from being visible to the view frustum. If this Actor is large enough it can also be an occluder for smaller Actors that it blocks from the view frustum.
+
+In the example below, the Camera represents the view of the scene. Behind the wall lies two Actors that cannot be seen because they are blocked from the camera view. The statue and barrel both would be occludees since 
+they are being blocked by the wall that is an occluding them. Also, remember that both the barrel and statue could be occluders if they were blocking anything behind them. Every occluder can be an occludee if it is large 
+enough to block another actor.
+
+[REGION:imagetable]
+| [INCLUDE:#OccluderVsOccludee1] | [INCLUDE:#OccluderVsOccludeeCV] |
+| --- | --- | 
+| Occluder = Wall, Occludee = Barrels and Statue  | Camera FOV shows the wall that occludes the barrels and statue. |
+[/REGION]
+[REGION:caption]
+_Click images for full size_
+[/REGION]
+
+<!---
+[EXCERPT:OccluderVsOccludee1]
+[REGION:lightbox]
+[![](OccluderVsOccludee.png)(w:450)](OccluderVsOccludee.png)
+[/REGION]
+[/EXCERPT:OccluderVsOccludee1]
+
+[EXCERPT:OccluderVsOccludeeCV]
+[REGION:lightbox]
+[![](OccluderVsOccludeeCameraView.png)(w:450)](OccluderVsOccludeeCameraView.png)
+[/REGION]
+[/EXCERPT:OccluderVsOccludeeCV]
+--->
+
+### Visibility Culling
+
+The following types of culling are used to determine the current visibility state of an Actor in the scene.
+
+| Types of Culling | Description |
+| ---------------- | ----------- |
+| **Distance** | The Actor's set distance is used to determine its visibility state. |
+| **View Frustum** | This is used to remove any Actors that do not have its bounds visible within the Camera's view frustum. Anything outside of this viewing area will not be rendered.  |
+| **Occlusion** | This is used to remove any additional Actors that are within the Camera's view frustum that is blocked by another from being viewed. This is done by testing if the Actor is visible within the scene's depth buffer. If it is visible it will be rendered, and if it is not it will be considered occluded and not rendered. |
+
+
+For more information on [Buffer Visualizations](Engine/UI/LevelEditor/Viewports/ViewModes#buffervisualization) you can read about them on this page.
+
+## Visibility Culling Methods
+
+Once the data has been culled from the scene it is tested against the available culling methods that have been implemented to test the visibility state of the Actors. The visibility culling methods work
+together and in the case of mobile offer a solution where dynamic occlusion is not possible.
+
+These culling methods listed here are in the order they are applied, which is also from least to most expensive in terms of rendering thread cost.
+
+* Distance culling (Per-instance and Cull Distance Volumes)
+* View frustum culling
+* Precomputed occlusion culling
+* Dynamic occlusion culling (hardware occlusion queries)
+
+
+#### Distance Culling (Per-instance)
+
+Each Actor in the scene can have its own set draw distance that will allows minimum and maximum distance for it to be rendered visible.
+
+![](DistanceCullSettings.png)
+
+| **Property** | **Description** |
+| ------------ | --------------- |
+| **Min Draw Distance** | The minimum distance at which the primitive should be rendered. This is measured in world space units from the center of the primitive's bounding sphere to the camera position. |
+| **Desired Max Draw Distance** | The max draw distance exposed to Level Designers. The real max draw distance is the min (disregarding 0) of this and volumes effecting this object. |
+| **Current Max Draw Distance** | This is not an editable property. This field shows the cull distance that is set by the CullDistanceVolume that the primitive is residing inside of. If it is set to 0 no Cull Distance Volume has been set so the value will be 0. |
+
+In the examples below, the **Min Draw Distance** and **Desired Max Draw Distance** have been set for the cube. These values determine the closest distance to render the Actor or the farthest distance that 
+the Actor can be rendered from the Camera's view.
+
+* **Min Draw Distance:** 500
+* **Desired Max Draw Distance:** 1000
+
+[REGION:note]
+Note that the **Desired Max Draw Distance** here will override any higher value set in a Cull Distance Volume. See the [Cull Distance Volume](1_TimH\2_IN_PROGRESS\VisibilityCulling\CullDistanceVolume) page for more information.
+[/REGION]
+
+[OBJECT:EmbeddedVideo]
+ [PARAMLITERAL:width]
+ 640
+ [/PARAMLITERAL]
+ [PARAMLITERAL:height]
+ 360
+ [/PARAMLITERAL]
+ [PARAMLITERAL:videoid]
+ LiiJ13LulIU
+ [/PARAMLITERAL]
+[/OBJECT]
+
+
+<!---
+[REGION:imagetable]
+| [INCLUDE:#Units1500] | [INCLUDE:#Units1000] | [INCLUDE:#Units500] |
+| --- | --- | --- |
+| More than Desired Max Draw Distance | Within Range of Min and Max Draw Distance | Less than Min Draw Distance | 
+[/REGION]
+
+To help illustrate what is happening in the video above, the images here show a top-down view and mark the distance from the **Camera** Actor to the center point of the **Cube** Actor. 
+The image directly below the top-down view shows the camera's view.
+
+--->
+
+Since the **Desired Max Draw Distance** is set to 1000 and the **Min Draw Distance** is set to 500 the cube Actor will only be visible when the Camera moves within this distance from the Actor's bounding box
+center. When the Camera moves outside of this space the Actor will no longer be rendered. 
+
+<!--
+[EXCERPT:Units1500]
+[REGION:lightbox]
+[![](1500units.png)(w:300)](1500Units.png)
+[/REGION]
+[/EXCERPT:Units1500]
+
+[EXCERPT:Units1000]
+[REGION:lightbox]
+[![](1000units.png)(w:300)](1000Units.png)
+[/REGION]
+[/EXCERPT:Units1000]
+
+[EXCERPT:Units500]
+[REGION:lightbox]
+[![](500units.png)(w:300)](500Units.png)
+[/REGION]
+[/EXCERPT:Units500]
+-->
+
+
+#### Cull Distance Volume
+
+Cull Distance Volumes allow you to specify a set of sizes and distances in an array that will automatically remove any Actor based on its bounds size at the corresponding distance. These volumes are useful
+for large outdoor scenes that have detailed interiors. 
+
+For more information on using and setting up [Cull Distance Volumes](1_TimH\2_IN_PROGRESS\VisibilityCulling\CullDistanceVolume) you can read about it in this page.
+
+
+#### View Frustum Culling
+
+The Camera's view frustum is used to test whether the culled data from the scene is visible within its viewing area. This process uses the Actor's bounding box and sphere to check if it is visible within the view. If not, the 
+Actor will not be rendered.
+
+The following example demonstrates how the view frustum considers Actors that are visible within the viewing area.
+
+* Yellow represents FOV of the view frustum and the area that is no occluded by other Actors.
+* Red represents the part of the view frustum that is occluded due to other Actors (ie. walls and large objects).
+
+[REGION:imagetable]
+| [INCLUDE:#Camera1] | [INCLUDE:#Camera2] | [INCLUDE:#Camera3] | [INCLUDE:#Camera4] |
+| ------------------ | ------------------ | ------------------ | ------------------ |
+|          1         |          2         |          3         |          4         | 
+[/REGION]
+
+| **Number** | **Description** |
+| ---------- | --------------- |
+| **1** | From this perspective there are large walls that act as occluders obstructing all but the view along the hallway. Because of this only the cube is visible. This means that the sphere and cone are occludees of the wall and will not be rendered. |
+| **2** | As the Camera moves forward there is more of the scene that is exposed to the view frustum. The walls no longer fully occlude the sphere and cone now that we can see parts of them around the corners of the wall. This means that they are now being fully rendered even though we can only see a small part of them. |
+| **3** | Next the as the Camera turns to the left we now have the sphere and cube within the view frustum. The cone and walls out of view will no longer be rendered. |
+| **4** | Now as the Camera is turned back to the right the cone and cube are within the view and the walls and sphere that are out of view are no longer being rendered. |
+
+[REGION:note]
+This is a generalized example of how visibility culling works and there is a number of caveats and reasons to how things can be rendered that are not covered here. For additional tips and strategies to use
+with visibility culling see the [Culling Strategies for Best Practices](#cullingstrategiesforbestpractices) section in this page.
+[/REGION]
+
+<!--
+[EXCERPT:Camera1]
+[REGION:lightbox]
+[![](Camera1.png)(w:215)](Camera1.png)
+[/REGION]
+[/EXCERPT:Camera1]
+
+[EXCERPT:Camera2]
+[REGION:lightbox]
+[![](Camera2.png)(w:215)](Camera2.png)
+[/REGION]
+[/EXCERPT:Camera2]
+
+[EXCERPT:Camera3]
+[REGION:lightbox]
+[![](Camera3.png)(w:215)](Camera3.png)
+[/REGION]
+[/EXCERPT:Camera3]
+
+[EXCERPT:Camera4]
+[REGION:lightbox]
+[![](Camera4.png)(w:215)](Camera4.png)
+[/REGION]
+[/EXCERPT:Camera4]
+-->
+
+
+#### Precomputed Visibility Volume
+
+<!---![](PrecomputedVisibilityVolume.png)(w:500)--->
+
+Precomputed Visibility Volumes use visibility cells that are placed on surfaces after a static lighting build. These visibility cells use the Camera view to change the visibility state of an Actor.
+This culling method is good for having visibility culling on mobile devices where hardware occlusion queries cannot be issued to check the visibility state. 
+
+For more information on using and setting up [Precomputed Visibility Volumes](1_TimH\2_IN_PROGRESS\VisibilityCulling\PrecomputedVisibilityVolume) you can read about them in this page.
+
+#### Dynamic Occlusion
+
+Dynamic occlusion queries are performed last in the chain of culling methods by only looking at the Actors that have been determined to be within the view frustum of the Camera. The Actors bounds that are 
+within the view frustum are then tested against the Scene Depth Buffer for their visibility by sending them to the GPU and reading them back one frame later. If an Actor's bounds is not visible within the 
+scene depth it will not be rendered. This allows the visibility state for any Actor occluded by another Actor to not be rendered. This method can allow for farther view distances because it relies on appoximate
+distances rather than setting a specific value for the far clip plane's distance.
+
+#### Hierarchical Z Buffer
+
+[REGION:warning]
+The Hierarchical Z-Buffer occlusion implementation is still considered experimental. There are some known issues that have not yet been addressed, like some meshes being culled when they should not be. 
+[/REGION]
+
+The idea behind Hierarchical Z-Buffer occlusion is to represent the depth of the scene with multiple resolutions to determine the visibility of an Actor. This is done by creating mipmapped version 
+of the render target for the scene depth which has texels that identify the depth in the scene for the Actors tested. The Actors can be tested for occlusion by comparing the depth value of the Actor's bounds
+to that in the render target that was captured. By using a lower resolution mip level this will use a more conservative check for the bounds of an Actor which leads to fewer texture fetches when sampling from
+a mip level. 
+
+[REGION:note]
+By rendering the scene to a lower resolution than the actual scene resolution you risk falsely culling Actors that should be visible.  
+[/REGION]
+
+Hierarchical Z-Buffer Occlusion Culling can be enabled by using the following console command: 
+
+	r.HZBOcclusion
+
+* 0 will disable Hierarchical Z-Buffer occlusion
+* 1 will enable Hierarchical Z-Buffer occlusion
+* 2 will force the Hierarchical Z-Buffer system, which will override any rendering platform preferences
+
+
+## Culling Strategies for Best Practices
+
+Below are some good practices to follow when designing and working with visibility culling for your game.
+
+### "Popping" in/Lag from GPU Queries of Actors
+
+This is a common issue with using visibility culling for the view frustum and hardware occlusion queries since they are being read back from the GPU one frame later. This is often 
+noticeable when turning your player camera very fast or when you come around a corner in darkly lit areas that have a contrasting lit area behind them. 
+
+* Use other geometry to help hide this artifact. For something like caves or darkened areas with bright outsides use a large occluder mesh that can sit just outside of these pieces to hide the popping in and block the light that makes it evident.
+* Increase the Bounds Scale of the Actor that is causing the offending behavior. 
+* Build geometry in a way that does not have too many small pieces that need to be assembled. This is especially true for walls, floors, and ceilings.
+
+### Large View Distances
+
+Large view distances can cause an impact on performance due to the number of Actors that need to be rendered. By blocking these large view distances with large occluders you can block Actors that would have been rendered otherwise.
+This can reduce the number of queries that need to be sent to the GPU by having fewer Actors to check. In turn, you get better performance by having fewer Actors to check and render.
+
+* Use large occluders where possible. 
+* Use Cull Distance Volumes to remove any smaller Actors that will not be easily visible at far distances, even if they are occluded from view.
+* Limit the amount of overdraw you get from Translucent Actors, since they will not occlude.
+
+### Asset Construction
+
+When creating geometry for your game it is a good idea to break your geometry up into modular pieces for design purposes and reusability. However, it is possible to modularize too much in the design 
+process by using pieces that are too small to make up larger geometry. Because of this, each actor will need to be have its visibility state checked. It's best to find a balance that makes sense
+for your game design and optimize for the best performance. Reducing the number of queries needed to check visibility is key for optimization and to help reduce the "Popping" in that can happen. 
+
+* Use larger geometry when crafting your levels for structural surfaces rather than breaking it up into many smaller pieces. There is more impact on the GPU testing multiple smaller Actors than a single Actor.
+* "Popping" in of Actors is reduced when considering the construction of the mesh.
+
+
+### Consider Occlusion vs Distance Culling Cost
+
+Occlusion and visibility culling methods cost much more than distance culling since these require GPU queries, so it is a good idea to setup Cull Distance Volumes for distant stuff even if they will 
+be occluded from a certain part of the scene.
+
+* Make good use of Cull Distance Volumes to reduce occlusion culling costs.
+* Stream out levels where possible instead of relying on distance or occlusion culling.
+* Use the `Stat Initviews` command to keep performance balance in check.
+
+For information about [Level Streaming](Engine/LevelStreaming) see how to set this up in this page.
+
+## Project Settings
+
+To open the **Project Settings** you must first go to the main menu and select **Edit** > **Project Settings**. When the window opens you can select the **Rendering** category to see the available
+settings under the section named **Culling**.
+
+
+![](ProjectSettingsCulling.png)
+
+| Setting | Description |
+| --- | --- |
+| **Occlusion Culling** | Allows occluded meshes to be culled and not rendered. If this is disabled all occlusion culling will be disabled for the project which can have performance implications. |
+| **Min Screen Radius for EarlyZPass** | Screen radius at which objects are culled for early Z pass. Larger values can improve performance but very large values can degrade performance if large occluders are not rendered. |
+| **Warn about no precomputed visibility** | Displays a warning when no precomputed visibility data is available for the current camera location. This can be helpful if you are making a game that relies on precomputed visibility. |
+
+[REGION:tip]
+​If you are developing a mobile game it would be beneficial to enable the option for **Warn about no precomputed visibility** as a good reminder for your Level Designers when developing for mobile, since it does not currently use dynamic 
+occlusion culling.
+[/REGION]
+
+
+<!---
+## Console Commands
+
+| Console Variable | Description |
+| --- | --- |
+| **r.OneFrameThreadLag** | Controls when occlusion queries are rendered. Rendering before the base pass may give worse occlusion (because not all occluders generally render in the earlyzpass). However, it may reduce CPU waiting for query result stalls on some platforms and increase overall performance. 0: After BasePass. 1: After EarlyZPass, but before BasePass. |
+| **r.AllowOcclusionQueries** | If zero, occlusion queries will not be used to cull primitives. |
+| **r.DownsampledOcclusionQueries** | Whether to issue occlusion queries to a downsampled depth buffer. |
+| **r.NumBufferedOcclusionQueries** | Number of frames to buffer occlusion queries (including the current render thread frame). More frames reduce the chance of stalling the CPU waiting for the results, but increases out of date query artifacts. |
+-->
+
+## Relevant Stats for Performance
+
+In real-time rendering for games, it is important to manage your scene for optimization by keeping check of the stats generated for your scene that are relevant to visibility and occlusion culling.
+
+Use the following console command to view the stats relevant to visibility and occlusion culling for your scene: 
+
+	Stat Initviews
+
+[REGION:warning]
+Keep in mind that when testing visibility and occlusion culling there are a couple of things to be aware of.
+
+* Occlusion culling is disabled in **Wireframe** view mode.
+* Use the console command **Stat Initviews** in PIE or Standalone games for more accurate results, because geometry for editor Actors that represent Lights, Cameras, and so on will be included when calculating visibility and occlusion culling.
+[/REGION]
+
+
+[REGION:lightbox]
+![](StatInitViews.png)
+[/REGION]
+[REGION:caption]
+The scene's stats above were taken from the [Infinity Blade Grasslands](https://www.unrealengine.com/marketplace/infinity-blade-plain-lands) created by EPIC Games and is freely available on the EPIC Games Marketplace.
+[/REGION]
+
+The **Stat Initiviews** window is broken up into two sections for `Cycle Counters` and `Counters`. The Cycle Counter's stats focus on the number of cycles in milliseconds it took to process, while the Counter's stats simply
+add up the primitives that are processed.
+
+[REGION:tip]
+`Visible Static Mesh Elements` is the single biggest contributor to rendering thread time and should be carefully watched and optimized.
+[/REGION]
+
+| Stat Initviews |  |
+| --- | --- |
+| **Cycle Counters** | [INCLUDE:#cyclecounter] |
+| **Counters** | [INCLUDE:#counters] |
+
+
+<!--
+[EXCERPT:CycleCounter]
+| Stat | Description|
+| --- | --- |
+| **View Visibility** | The amount of frame time spent to process Actor visibility queries. |
+| **Occlusion Cull** | The amount of frame time spent to query the scene for Actors within the view frustum that were occluded by other Actors. |
+| **Frustum Cull** | The amount of frame time spent to query if an Actor's bounds were within the view frustum. |
+| **Decompress Occlusion** | The amount of frame time it took to load the precomputed visibility. |
+[/EXCERPT:CycleCounter]
+
+[EXCERPT:Counters]
+| Stat | Description|
+| --- | --- |
+| **Processed Primitives** | This is the total number of Actors that are being processed in the scene.  |
+| **Frustum Culled Primitives** | This is the number of primitives culled when not within the view frustum's FOV. |
+| **Occluded Primitives** | The number of Actors that are occluded by other Actors from within the view frustum's FOV. |
+| **Occlusion Queries** | This is the number of scene occlusion queries sent to the GPU based on the number of Actors visible within the view frustum. |
+| **Visible Dynamic Primitives** | This is the number of dynamic primitives in the scene that have their mobility set to Movable, like Particle Systems and Foliage instances. |
+| **Visible Static Mesh Elements** | This is the number of visible meshes that have their mobility set to Static within the scene. |
+[/EXCERPT:Counters]
+-->
+
+## Tips & Tricks
+
+The following Tips & Tricks can be used to help troubleshoot and test your scene when you suspect problem related to visibility or occlusion culling. 
+
+##### Visualize Occluded Primitives
+
+This will enable the visualization mode to draw the bounding boxes of any Actors in your scene that are occluded from the Camera view. A green bounding box indicates the Actor is 
+dynamically occluded and a red bounding box indicates that the Actor is statically occluded when using [Precomputed Visibility Volumes](1_TimH\2_IN_PROGRESS\VisibilityCulling\PrecomputedVisibilityVolume). 
+
+Use the following console command to enable this visualization: 
+
+	r.VisualizeOccludedPrimitives 1
+
+
+[REGION:note]
+This console command will only work while in the editor and not in PIE or Standalone games.
+[/REGION]
+
+In the example below, the visualization of occluded Actors blocked by the wall is compared to the placement of the Actor's with the wall removed. 
+
+[OBJECT:ComparisonSlider]
+ [PARAM:before]
+ ![Visualize Occluded Primitives Enabled (Green Box)](VisOccludedPrimitives1.png)(w:500)
+ [/PARAM]
+ [PARAM:after]
+ ![Actors occluded by the wall](VisOccludedPrimitives2.png)(w:500)
+ [/PARAM]
+[/OBJECT]
+
+
+##### Freeze Rendering
+
+This will toggle the rendering state of the scene from the current Camera perspective while allowing you to freely move around the world.
+
+Use the following console command to toggle this mode:
+
+	FREEZERENDERING
+
+
+[OBJECT:EmbeddedVideo]
+ [PARAMLITERAL:width]
+ 640
+ [/PARAMLITERAL]
+ [PARAMLITERAL:height]
+ 360
+ [/PARAMLITERAL]
+ [PARAMLITERAL:videoid]
+ El127q0lv5s
+ [/PARAMLITERAL]
+[/OBJECT]
+
+
+##### Game View Mode
+
+You can use the keyboard shortcut for **G** or use the main **Viewport Options** down-arrow to select **Game View** to be able to work in the editor and switch between what you would see in game. 
+This is helpful when you do not want to PIE or launch a game but need to check what the game would look like. This is ideal for checking [Cull Distance Volumes](1_TimH\2_IN_PROGRESS\VisibilityCulling\CullDistanceVolume) 
+and [Precomputed Visibility Volumes](1_TimH\2_IN_PROGRESS\VisibilityCulling\PrecomputedVisibilityVolume) since these do not work unless you are in game or using the Game View.
+
+For more information on the [Game View](Engine\UI\LevelEditor\Viewports\Basics) mode, you can read about it in this page.
+
+
+
+
+[OBJECT:TopicButtonList]
+ [PARAM:icon]
+  ![](%ROOT%/related_icon.png)(convert:false)
+ [/PARAM]
+ [PARAM:title]
+  Related Links
+ [/PARAM]
+ [PARAM:description]
+ [/PARAM]
+ [PARAM:links]
+  * [](1_TimH\2_IN_PROGRESS\VisibilityCulling\PrecomputedVisibilityVolume)
+  * [](1_TimH\2_IN_PROGRESS\VisibilityCulling\CullDistanceVolume)
+ [/PARAM]
+[/OBJECT]
\ No newline at end of file
diff --git a/Engine/Documentation/Source/Engine/Rendering/VisibilityCulling/VisibilityCullingLanding.INT.udn b/Engine/Documentation/Source/Engine/Rendering/VisibilityCulling/VisibilityCullingLanding.INT.udn
new file mode 100644
--- /dev/null
+++ b/Engine/Documentation/Source/Engine/Rendering/VisibilityCulling/VisibilityCullingLanding.INT.udn
@@ -0,0 +1,70 @@
+﻿Availability: Docs
+Title: Visibility Culling
+Crumbs: %ROOT%, Engine, 
+Description: Visibility Culling Landing Page
+Version: 4.13
+
+[VAR:Topic]
+[OBJECT:Topic]
+	[PARAM:image]
+		![%Engine/Rendering/VisibilityCulling:title%]()
+	[/PARAM]
+	[PARAM:icon]
+		![](%ROOT%/light_icon.png)(convert:false)
+	[/PARAM]
+	[PARAM:title]
+		%Engine/Rendering/VisibilityCulling:title%
+	[/PARAM]
+	[PARAM:description]
+		%Engine/Rendering/VisibilityCulling:description%
+	[/PARAM]
+	[PARAM:path]
+		[RELATIVE:Engine/Rendering/VisibilityCulling]
+	[/PARAM]
+[/OBJECT]
+[/VAR]
+[VAR:TopicCompact]
+[OBJECT:TopicCompact]
+	[PARAM:image]
+		![%Engine/Rendering/VisibilityCulling:title%]()
+	[/PARAM]
+	[PARAM:icon]
+		![](%ROOT%/light_icon.png)(convert:false)
+	[/PARAM]
+	[PARAM:title]
+		%Engine/Rendering/VisibilityCulling:title%
+	[/PARAM]
+	[PARAM:description]
+		%Engine/Rendering/VisibilityCulling:description%
+	[/PARAM]
+	[PARAM:path]
+		[RELATIVE:Engine/Rendering/VisibilityCulling]
+	[/PARAM]
+[/OBJECT]
+[/VAR]
+
+
+Unreal Engine 4 uses a combination of methods to gather data from the scene for visibility and occlusion culling that can be used to test whether an Actor should be visible. This system is useful
+for optimizing your game by not rendering what is not needed. These culling methods largely work in conjunction with one another, while some can be used when a dynamic visibility option is not 
+available, like with mobile. 
+
+
+[OBJECT:TopicButtonList]
+	[PARAM:icon]
+		![](%ROOT%/tech_icon.png)(convert:false)
+	[/PARAM]
+	[PARAM:title]
+		Visibility Culling Topics
+	[/PARAM]
+	[PARAM:description]
+	[/PARAM]
+	[PARAM:links]
+	* [](Engine\Rendering\VisibilityCulling\CullingOverview "%Engine\Rendering\VisibilityCulling\CullingOverview:description%")
+	* [](Engine\Rendering\VisibilityCulling\CullingBestPractices "%Engine\Rendering\VisibilityCulling\CullingBestPractices:description%")
+	* [](Engine\Rendering\VisibilityCulling\CullingPerformance "%Engine\Rendering\VisibilityCulling\CullingPerformance:description%")
+	* [](Engine\Rendering\VisibilityCulling\CullingTipsTricks "%Engine\Rendering\VisibilityCulling\CullingTipsTricks:description%")
+	<!---* [](Engine\Rendering\VisibilityCulling\CullDistanceVolume "%Engine\Rendering\VisibilityCulling\CullDistanceVolume:description%")--->
+	<!---* [](Engine\Rendering\VisibilityCulling\PrecomputedVisibilityVolume "%Engine\Rendering\VisibilityCulling\PrecomputedVisibilityVolume:description%")--->
+	[/PARAM]
+[/OBJECT]
+
