#UE4 Docs: updated Mac specs

#UEDOC-3187

diff --git a/Engine/Documentation/Source/GettingStarted/RecommendedSpecifications/Specifications.INT.udn b/Engine/Documentation/Source/GettingStarted/RecommendedSpecifications/Specifications.INT.udn
--- a/Engine/Documentation/Source/GettingStarted/RecommendedSpecifications/Specifications.INT.udn
+++ b/Engine/Documentation/Source/GettingStarted/RecommendedSpecifications/Specifications.INT.udn
@@ -69,7 +69,7 @@ Unreal Engine has some specific hardware and software requirements for running t
 [EXCERPT:MacOSMinHardware]
 |||
 | -------- | --- |
-| **Operating System** | Mac OS X 10.9.2  |
+| **Operating System** | Mac OS X 10.10.5  |
 | **Processor** | Quad-core Intel, 2.5 GHz or faster |
 | **Memory** | 8 GB RAM |
 | **Video Card** | OpenGL 4.1 compatible graphics card |
@@ -156,7 +156,7 @@ Below are the requirements for running the engine or editor.
 [EXCERPT:MacOSMinSoftware]
 | Running the Engine ||
 | -------- | --- |
-| **Operating System** | Mac OS X 10.9.2  |
+| **Operating System** | Mac OS X 10.10.5  |
 [/EXCERPT:MacOSMinSoftware]
 
 [EXCERPT:LinuxMinSoftware]
@@ -235,7 +235,7 @@ The requirements for programmers developing with the engine are listed below.
 [EXCERPT:MacOSMinDevelopment]
 | Developing with the Engine ||
 | -------- | --- |
-| **Xcode Version** | 5.0 |
+| **Xcode Version** | 7.2.1 |
 [/EXCERPT:MacOSMinDevelopment]
 [EXCERPT:LinuxMinDevelopment]
 | Developing with the Engine ||
