#UE4 Docs: Changed the aval to doc instead of public. Removed a step as it was doubled up.

diff --git a/Engine/Documentation/Source/Platforms/Android/OpenGLES31MobileRenderer/OpenGLES31MobileRenderer.INT.udn b/Engine/Documentation/Source/Platforms/Android/OpenGLES31MobileRenderer/OpenGLES31MobileRenderer.INT.udn
--- a/Engine/Documentation/Source/Platforms/Android/OpenGLES31MobileRenderer/OpenGLES31MobileRenderer.INT.udn
+++ b/Engine/Documentation/Source/Platforms/Android/OpenGLES31MobileRenderer/OpenGLES31MobileRenderer.INT.udn
@@ -1,7 +1,7 @@
 Title: Android OpenGL ES 3.1 Mobile Renderer
 Description:Setting up your UE4 Android project to work with the OpenGL ES 3.1 mobile render
 Crumbs:%ROOT%, Platforms, Platforms/Android, Platforms/Android/OpenGLES31MobileRenderer/
-Availability:Public
+Availability:Docs
 version:4.12
 Parent:Platforms/Android
 tags:Mobile
@@ -93,20 +93,6 @@ Please note that OpenGL ES 3.1 Support will only work on code based projects tha
 	Click for full image.
 	[/REGION]
 
-1. Under the **APKPackaging** section, locate the  **Build** section and make sure that the following options are checked. 
-
-	* Support armv7
-	* Support OpenGL ES2
-	* Support OpenGL ES3.1
-
-	[REGION:lightbox]
-	[![](OGL3_04.png)(w:640)](OGL3_04.png)
-	[/REGION]
-
-	[REGION:caption]
-	Click for full image.
-	[/REGION]
-
 1. Under the **APKPackaging** section, locate the  **Build** section and make sure that the following options are checked. 
 	
 	* Support armv7
