Updating images, changing page layout for added clarity on the Universal Gizmo page, now named Transforming Actors in VR UEDOC-3162

diff --git a/Engine/Documentation/Source/Engine/Editor/VR/UniversalGizmo/UniversalGizmo.INT.udn b/Engine/Documentation/Source/Engine/Editor/VR/UniversalGizmo/UniversalGizmo.INT.udn
--- a/Engine/Documentation/Source/Engine/Editor/VR/UniversalGizmo/UniversalGizmo.INT.udn
+++ b/Engine/Documentation/Source/Engine/Editor/VR/UniversalGizmo/UniversalGizmo.INT.udn
@@ -1,12 +1,12 @@
-Title: Universal Gizmo
-Description: Transform widget designed for moving, rotating, and scaling Actors in VR.
+Title: Transforming Actors in VR
+Description: Transform widgets designed for moving, rotating, and scaling Actors in VR.
 Crumbs:
 Availability: public
 parent: Engine/Editor/VR
 type: reference
 order: 4
 tags: VR
-version: 4.11
+version: 4.13
 skilllevel: intermediate
 
 [VAR:Topic]
@@ -43,18 +43,82 @@ skilllevel: intermediate
 [/OBJECT]
 [/VAR]
 
-The **Universal Gizmo** is a transform widget designed for moving, rotating, and scaling Actors when working in VR mode. 
+The **Universal Gizmo** is a transform widget designed for moving, rotating, and scaling Actors when working in VR mode. There are both local and world space versions available, which can be swapped between using the [](Engine/Editor/VR/QuickSelect). 
+The elements of the Universal Gizmo are colored red, green, and blue representing the X, Y, and Z axes, respectively. Much like in the standard editor, there are gizmos which can be used specifically for translating, rotating, and scaling your Actors. Local and world space versions of these are accessible from the [](Engine/Editor/VR/QuickSelect). 
 
-![](TransformGizmo.png)(w:400)
+| Gizmo | World | Local |
+| ---- |-------| -----|
+| **Universal** | ![](WorldUniversalGizmo.png)(w:350)[INCLUDE:#worlduniversal]|![](LocalUniversalGizmo.png)(w:350) [INCLUDE:#localuniversal]  |
+| **Translate** | ![](WorldTranslateGizmo.png)(w:350)[INCLUDE:#worldtranslate]|![](LocalTranslateGizmo.png)(w:350) [INCLUDE:#localtranslate]|
+| **Rotate** | ![](WorldRotationGizmo.png)(w:350) [INCLUDE:#worldrotate]|![](LocalRotateGizmo.png)(w:350) [INCLUDE:#localrotate] |
+| **Scale** | ![](WorldScaleGizmo.png)(w:350) [INCLUDE:#worldscale]|![](LocalScaleGizmo.png)(w:350)[INCLUDE:#localscale] |
 
 
-The elements of the Universal Gizmo are colored red, green, and blue representing the X, Y, and Z axes, respectively. 
+<!--
+
+[EXCERPT:WorldUniversal]
+* [Translate](#translate)
+* [Translate Plane](#translateplane)
+* [Rotate](#rotate)
+* [Uniform Scale](#uniformscale)
+[/EXCERPT]
+
+[EXCERPT:LocalUniversal]
+* [Translate](#translate)
+* [Translate Plane](#translateplane)
+* [Rotate](#rotate)
+* [Uniform Scale](#uniformscale)
+* [Scale](#scale)
+[/EXCERPT]
+-->
+<!--
+
+[EXCERPT:WorldTranslate]
+* [Translate](#translate)
+* [Translate Plane](#translateplane)
+[/EXCERPT]
+
+[EXCERPT:LocalTranslate]
+* [Translate](#translate)
+* [Translate Plane](#translateplane)
+[/EXCERPT]
+-->
+<!--
+
+[EXCERPT:WorldRotate]
+* [Rotate](#rotate)
+[/EXCERPT]
+
+[EXCERPT:LocalRotate]
+* [Rotate](#rotate)
+[/EXCERPT]
+-->
+<!--
+
+[EXCERPT:WorldScale]
+* [Uniform Scale](#uniformscale)
+[/EXCERPT]
+
+[EXCERPT:LocalScale]
+* [Uniform Scale](#uniformscale)
+* [Scale](#scale)
+* [Stretch](#stretch)
+[/EXCERPT]
+-->
+
+## Transform Gizmo Actions
+
+Each gizmo gives you access to different transform actions. By aiming the controllers at the manipulators highlighted below (usually on a translation, rotation, or scale gizmo for clarity), you can manipulate your Actor's Location, Rotation, and Scale independently.
 
 | Action | Manipulator | Usage |
 | --- | --- | --- |
-| **Translate** | ![](Translate.png)(w:300) | Aim the controller at one of the translation manipulators and squeeze the **Trigger** halfway to grab it. Move the controller to move the selected Actor(s) along the axis. |
-| **Translate Plane** | ![](TranslatePlane.png)(w:300) | Aim the controller at one of the translation plane manipulators and squeeze the **Trigger** halfway to grab it. Move the controller to move the selected Actor(s) along the plane (in 2 axes). |
-| **Rotate** | ![](Rotate.png)(w:300) | Aim the controller at one of the rotation manipulators and squeeze the **Trigger** halfway to grab it. Move the controller to rotate the selected Actor(s) around the axis. |
-| **Uniform Scale** | ![](UniformScale.png)(w:300) | Aim the controller at the white cube in the center of the gizmo and squeeze the **Trigger** halfway to grab it. Move the controller to scale the selected Actor(s) along all axes uniformly. |
-| **Scale** | ![](NonUniformScale.png)(w:300) | Aim the controller at one of the scale manipulators and squeeze the **Trigger** halfway to grab it. Move the controller to scale the selected Actor(s) along the Y-axis. Note that the individual scale manipulators only show if you have set **Local Space** using the [Quick Select Menu](Engine/Editor/VR/QuickSelect).|
+|[REGION:none](#translate)[/REGION] **Translate** | ![](LocalTranslateGizmo_Translate.png)(w:300) | Aim the controller at one of the translation manipulators and squeeze the **Trigger** halfway to grab it. Move the controller to move the selected Actor(s) along the axis. |
+| [REGION:none](#translateplane)[/REGION]**Translate Plane** | ![](LocalTranslateGizmo_TranslatePlane.png)(w:300) | Aim the controller at one of the translation plane manipulators and squeeze the **Trigger** halfway to grab it. Move the controller to move the selected Actor(s) along the plane (in 2 axes). |
+|[REGION:none](#rotate)[/REGION] **Rotate** | ![](LocalRotateGizmo_Highlight.png)(w:300) | Aim the controller at one of the rotation manipulators and squeeze the **Trigger** halfway to grab it. Move the controller to rotate the selected Actor(s) around the axis. |
+| [REGION:none](#uniformscale)[/REGION]**Uniform Scale** | ![](WorldScaleGizmo_Highlight.png)(w:300) | Aim the controller at the white cube in the center of the gizmo and squeeze the **Trigger** halfway to grab it. Move the controller to scale the selected Actor(s) along all axes uniformly. |
+| [REGION:none](#scale)[/REGION]**Scale** | ![](LocalScaleGizmo_Highlight.png)(w:300) | Aim the controller at one of the scale manipulators and squeeze the **Trigger** halfway to grab it. Move the controller to scale the selected Actor(s) along the selected axis or plane. |
+| [REGION:none](#stretch)[/REGION]**Stretch** |![](LocalScaleGizmo_Stretch.png)(w:300) | Aim the controller at one of the stretch manipulators and squeeze the Trigger halfway to grab it. Move the controller to scale the selected Actor(s) non-uniformly along all axes while remaining anchored to the opposite corner.| 
+
+
 
+  
