Minor word swap in VR Editor controls page for clarity

diff --git a/Engine/Documentation/Source/Engine/Editor/VR/Controls/VREditorControls.INT.udn b/Engine/Documentation/Source/Engine/Editor/VR/Controls/VREditorControls.INT.udn
--- a/Engine/Documentation/Source/Engine/Editor/VR/Controls/VREditorControls.INT.udn
+++ b/Engine/Documentation/Source/Engine/Editor/VR/Controls/VREditorControls.INT.udn
@@ -121,7 +121,7 @@ Each Controller emits a beam called the laser pointer which is used to determine
 | Control | Description |
 | ------- | ----------- |
 | ![Aim Controller at Actor + Squeeze Trigger + Move Controller](MeshPaint_2_NoModifier.png)(convert:false) | Aim the controller at a Static Mesh in the world and squeeze the Trigger to begin painting. Move the controller to move the brush along the Mesh. |
-| ![Modifier + Aim Controller at Actor + Squeeze Trigger + Move Controller](ModifierMeshPaint.png)(convert:false) | To erase instead of paint, hold the modifier button. Aim the controller at a Static Mesh in the world and squeeze the Trigger to begin erasing. Move the controller to move the brush along the Mesh.  |
+| ![Modifier + Aim Controller at Actor + Squeeze Trigger + Move Controller](ModifierMeshPaint.png)(convert:false) | To erase instead of paint, hold the modifier button. Aim the controller at a Static Mesh in the world and squeeze the Trigger to begin erasing. Move the controller to move the eraser along the Mesh.  |
 [/REGION]
 
 
