#UE4 docs updated description to actually describe what lighting channels do

diff --git a/Engine/Documentation/Source/Engine/Rendering/LightingAndShadows/LightingChannels/LightingChannels.INT.udn b/Engine/Documentation/Source/Engine/Rendering/LightingAndShadows/LightingChannels/LightingChannels.INT.udn
--- a/Engine/Documentation/Source/Engine/Rendering/LightingAndShadows/LightingChannels/LightingChannels.INT.udn
+++ b/Engine/Documentation/Source/Engine/Rendering/LightingAndShadows/LightingChannels/LightingChannels.INT.udn
@@ -1,7 +1,7 @@
 Availability:Docs
 Title: Lighting Channels
 Crumbs: %ROOT%, Engine
-Description:The Lighting Channels feature overview.
+Description:Lighting Channels ia a rendering feature that enables you to choose what renderable surfaces are lit by specific lights.
 Navigation:topic
 Version: 4.13
 
