#UE4 docs minor spelling and consistency fixes

diff --git a/Engine/Documentation/Source/Engine/Sequencer/HowTo/ImportExportEDLs/ImportExportEDLs.INT.udn b/Engine/Documentation/Source/Engine/Sequencer/HowTo/ImportExportEDLs/ImportExportEDLs.INT.udn
--- a/Engine/Documentation/Source/Engine/Sequencer/HowTo/ImportExportEDLs/ImportExportEDLs.INT.udn
+++ b/Engine/Documentation/Source/Engine/Sequencer/HowTo/ImportExportEDLs/ImportExportEDLs.INT.udn
@@ -8,13 +8,13 @@ parent:Engine/Sequencer/HowTo
 order:1
 checkpoint: editorqs
 Related: Engine/Sequencer/Overview
-Related: Engine\Sequencer\ReferenceEditor\
+Related: Engine/Sequencer/ReferenceEditor
 
 [EXCERPT:intro]
 Sequencer in Unreal Engine 4 not only allows you to render and export out your entire cinematic, but you can also export your cinematic split up into each of your shots along with an **Edit Decision List** or **EDL** which is a file that can be used with most video editing applications like Premiere, Avid or Final Cut. 
 [/EXCERPT:intro]
 
-The EDL contains an ordered list of reel and time code data representing where each video clip can be obtained in order to produce the final cut. 
+The EDL contains an ordered list of reel and timecode data representing where each video clip can be obtained in order to produce the final cut. 
 In this how-to we export a cinematic and EDL file, take it into an external program like Premiere, make edits to it and re-import the edited cinematic back into Unreal along with those changes. 
 
 ## Steps
@@ -75,7 +75,7 @@ We are also using Adobe Premiere Pro CS6 as our external editing software; howev
 	[/PARAMLITERAL]
 	[/OBJECT]
 
-	In Premeire Pro, **Right-click** the media folder and select **Link Media** then point each file to its corresponding media file. 
+	In Premiere Pro, **Right-click** the media folder and select **Link Media** then point each file to its corresponding media file. 
 
 1.	Modify the shots in your editing software as desired, then **Export** as an **EDL**. 
 
@@ -126,3 +126,5 @@ By exporting your sequence as an EDL and working externally, you have the abilit
 
 
 
+
+
