First draft of BuildGraph Introduction NoPublish page.

#UE4doc

diff --git a/Engine/Documentation/Source/Programming/Development/BuildGraphs/BuildGraphIntroduction/BuildGraphIntroduction.INT.udn b/Engine/Documentation/Source/Programming/Development/BuildGraphs/BuildGraphIntroduction/BuildGraphIntroduction.INT.udn
new file mode 100644
--- /dev/null
+++ b/Engine/Documentation/Source/Programming/Development/BuildGraphs/BuildGraphIntroduction/BuildGraphIntroduction.INT.udn
@@ -0,0 +1,31 @@
+Availability: NoPublish
+Title: Introduction to BuildGraphs
+Crumbs: %ROOT%, Programming, Programming/Development, Programming/Development/BuildGraphs
+Description:This page introduces you to our BuildGraph scripting system.
+version: 4.13
+parent:Programming/Development/BuildGraphs
+type:Reference
+
+**BuildGraph** is a scripting language for build automation with building blocks common to Unreal Engine 4 (UE4) projects. 
+BuildGraph integrates with UnrealBuildTool, AutomationTool, and the editor. Which means that BuildGraph can be 
+extended and customized for your projects.
+
+BuildGraph scripts are written with [XML](https://www.w3.org/TR/NOTE-dcd]), and specify a graph of user-defined nodes with 
+dependencies between them. Each node consists of tasks executed in sequence to produce some sort of output (i.e., compile 
+this project, then cook, run this test, and so on). When asked to build a target (that is, a node or named output) BuildGraph 
+will execute all of the nodes in the graph required to make that happen.
+
+Unlike other build tools, BuildGraph is designed as a hybrid between a makefile-like scripting language and a build farm 
+configuration script. It allows annotations for the type of machine that steps are supposed to be executed on, providing a 
+list of recipients for failure notifications if a step fails, and groups nodes that should only be executed after an explicit 
+user trigger. It also tracks the creation of output files from task execution in a way that lends to graph execution being 
+distributed across a farm of machines (with nodes running in parallel where possible), and intermediate artifacts being 
+transferred to and from a network share automatically.
+
+Epic uses BuildGraph to prepare the UE4 binary release, package samples for the marketplace, and implement pipelines for our 
+own games (among other things). Several example BuildGraph scripts are provided in the 
+`[UE4Root]/Engine/Build/Graph/Examples` directory, and the script for creating a binary UE4 distribution can be found at 
+`[UE4Root]/Engine/Build/InstalledEngineBuild.xml`.
+
+Opening a script with Visual Studio will use the schema located at `[UE4Root]Engine/Build/Graph/Schema.xsd` to provide rich 
+tooltips, validation, and Intellisense features while editing.
\ No newline at end of file
