First draft of BuildGraph Tasks reference page.

#UE4doc

diff --git a/Engine/Documentation/Source/Programming/Development/BuildGraphs/BuildGraphTasks/BuildGraphPredefinedTasks/BuildGraphPredefinedTasks.INT.udn b/Engine/Documentation/Source/Programming/Development/BuildGraphs/BuildGraphTasks/BuildGraphPredefinedTasks/BuildGraphPredefinedTasks.INT.udn
new file mode 100644
--- /dev/null
+++ b/Engine/Documentation/Source/Programming/Development/BuildGraphs/BuildGraphTasks/BuildGraphPredefinedTasks/BuildGraphPredefinedTasks.INT.udn
@@ -0,0 +1,7 @@
+Availability: NoPublish
+Title: BuildGraph Predefined Tasks
+Crumbs: %ROOT%, Programming, Programming/Development, Programming/Development/BuildGraphs, Programming/Development/BuildGraphs/BuildGraphTasks
+Description: This is a procedurally generated markdown page.
+version: 4.13
+parent:Programming/Development/BuildGraphs/BuildGraphTasks
+
diff --git a/Engine/Documentation/Source/Programming/Development/BuildGraphs/BuildGraphTasks/BuildGraphTasks.INT.udn b/Engine/Documentation/Source/Programming/Development/BuildGraphs/BuildGraphTasks/BuildGraphTasks.INT.udn
new file mode 100644
--- /dev/null
+++ b/Engine/Documentation/Source/Programming/Development/BuildGraphs/BuildGraphTasks/BuildGraphTasks.INT.udn
@@ -0,0 +1,58 @@
+Availability: Docs
+Title: BuildGraph Tasks
+Crumbs: %ROOT%, Programming, Programming/Development, Programming/Development/BuildGraphs
+Description: Learn how BuildGraph can help you create custom tasks.
+version: 4.13
+parent:Programming/Development/BuildGraphs
+type:Reference
+related: Programming/Development/BuildGraphs/BuildGraphScriptAnatomy
+tags:BuildGraph
+
+[VAR:TopicCompact]
+[OBJECT:TopicCompact]
+	[PARAM:image]
+		![%Programming/Development/BuildGraphs/BuildGraphTasks%](tasks_topic.png)
+	[/PARAM]
+	[PARAM:icon]
+		![](%ROOT%/reference_icon.png)(convert:false)
+	[/PARAM]
+	[PARAM:title]
+		%Programming/Development/BuildGraphs/BuildGraphTasks:title%
+	[/PARAM]
+	[PARAM:description]
+		%Programming/Development/BuildGraphs/BuildGraphTasks:description%
+	[/PARAM]
+	[PARAM:path]
+		[RELATIVE:Programming/Development/BuildGraphs/BuildGraphTasks]
+	[/PARAM]
+[/OBJECT]
+[/VAR]
+
+[TOC(start:1 end:2)]
+
+[EXCERPT:buildgraphtasks]
+
+Typically, project-specific utility scripts are implemented as AutomationTool BuildCommand classes. There are times; however, when you'll want 
+to automate custom tasks for your project. This is where BuildGraph can help. BuildGraph can be extended with any number of tasks, enabling you 
+to customize your build automation tasks in ways that suit your specific needs. The following section goes over ways of creating custom 
+tasks using BuildGraph, including a listing of predefined tasks that'll give you a great starting point for creating custom tasks.
+
+## Custom Tasks
+
+To create a new custom task, implement a class derived from the `CustomTask` class and apply the `[TaskElement]` attribute to it. 
+The `TaskElement` constructor takes two arguments; the name of the XML element that it's represented by, and the type of class 
+containing its parameters (which is passed to the constructor at load time).
+
+If you want to read parameter class fields from an XML file, attach the `[TaskParameter]` attribute to the task. Attaching the 
+`TaskParameter` attribute indicates if the parameter is required or optional, including additional validation that should be 
+applied to the argument.
+
+## Predefined Tasks
+
+If you'd like to have a starting point for creating tasks, we're providing a variety of predefined tasks as templates for you to work from. 
+If you're looking for a simple predefined task to start from, we recommend using `LogTask` as a good starting point. If you want to keep up with 
+our latest improvements to BuildGraph; bookmark this page, because the following list of predefined tasks will be updated periodically. 
+
+[INCLUDE:Programming/Development/BuildGraphs/BuildGraphTasks/BuildGraphPredefinedTasks]
+
+[/EXCERPT:buildgraphtasks]
\ No newline at end of file
