First draft of BuildGraph Usage reference page.

#UE4doc

diff --git a/Engine/Documentation/Source/Programming/Development/BuildGraphs/BuildGraphUsage/BuildGraphUsage.INT.udn b/Engine/Documentation/Source/Programming/Development/BuildGraphs/BuildGraphUsage/BuildGraphUsage.INT.udn
new file mode 100644
--- /dev/null
+++ b/Engine/Documentation/Source/Programming/Development/BuildGraphs/BuildGraphUsage/BuildGraphUsage.INT.udn
@@ -0,0 +1,84 @@
+Availability: Docs
+Title: BuildGraph Usage
+Crumbs: %ROOT%, Programming, Programming/Development, Programming/Development/BuildGraphs
+Description:This page shows you how to use BuildGraph Scripts.
+version: 4.13
+parent:Programming/Development/BuildGraphs
+type:Overview
+related: Programming/Development/BuildGraphs/BuildGraphScriptAnatomy
+tags:BuildGraph
+
+[VAR:TopicCompact]
+[OBJECT:TopicCompact]
+	[PARAM:image]
+		![%Programming/Development/BuildGraphs/BuildGraphUsage%](usage_topic.png)
+	[/PARAM]
+	[PARAM:icon]
+		![](%ROOT%/reference_icon.png)(convert:false)
+	[/PARAM]
+	[PARAM:title]
+		%Programming/Development/BuildGraphs/BuildGraphUsage:title%
+	[/PARAM]
+	[PARAM:description]
+		%Programming/Development/BuildGraphs/BuildGraphUsage:description%
+	[/PARAM]
+	[PARAM:path]
+		[RELATIVE:Programming/Development/BuildGraphs/BuildGraphUsage]
+	[/PARAM]
+[/OBJECT]
+[/VAR]
+
+[TOC(start:1 end:3)]
+
+BuildGraph is implemented as part of AutomationTool, which you can run directly from the Unreal Engine 4 solution on Windows, 
+or via the `Engine/Build/BatchFiles/RunUAT.bat` batch file from the command line. On Mac or Linux, it can be run with the 
+`Engine/Build/BatchFiles/RunUAT.sh` shell script. If you'd like to get an idea of how to use BuildGraph from the command line, 
+we've written sample usage commands for your reference.
+
+## Showing Options & Nodes
+
+To show all of the nodes and options available in a BuildGraph script:
+
+	BuildGraph -Script=Engine/Build/Graph/Examples/AllExamples -ListOnly
+
+## Setting an Option
+
+To set an option defined in the script:
+
+	BuildGraph -Script=Engine/Build/Graph/Examples/Properties.xml -ListOnly -Set:CommandLineProperty=ABCDEF
+
+## Seeing Node Execution
+
+To see what nodes will be executed for a target:
+
+	BuildGraph -Script=Engine\Build\Graph\Examples\AllExamples.xml -Target="All Examples" -ListOnly
+
+## Cleaning Previous Builds
+
+To clean any previous builds and execute a specific target:
+
+	BuildGraph -Script=Engine\Build\Graph\Examples\AllExamples.xml -Target="Property Examples" -Clean
+
+## Writing out an XML Schema
+
+To write out an XML schema containing all registered tasks:
+
+	BuildGraph -Schema=Engine\Build\Graph\Schema.xsd
+
+## Generating a Preprocessed Script
+
+To output a preprocessed version of your script (useful for debugging):
+
+	BuildGraph -Script=Engine\Build\Graph\Examples\AllExamples.xml -Target="All Examples" -ListOnly -Preprocess=D:\Preprocessed.xml
+
+## Exporting a Preprocessed Graph to JSON
+
+To export a preprocessed graph to JSON (for dynamic configuration on build systems):
+
+	BuildGraph -Script=Engine\Build\Graph\Examples\AllExamples.xml -Target="All Examples" -Export=D:\ExportedGraph.json
+
+## Executing a Single Node
+
+To execute a single node on a build system (while transferring intermediates to and from shared storage):
+
+	BuildGraph -Script=Engine\Build\Graph\Examples\AllExamples.xml -Target="All Examples" -Clean -SingleNode=”Manipulate Some Files” -SharedStorage=P:\Network-Share\Unique-Job-12345\
\ No newline at end of file
