First draft of BuildGraph Types reference page.

#UE4doc

diff --git a/Engine/Documentation/Source/Programming/Development/BuildGraphs/BuildGraphTypes/BuildGraphTypes.INT.udn b/Engine/Documentation/Source/Programming/Development/BuildGraphs/BuildGraphTypes/BuildGraphTypes.INT.udn
new file mode 100644
--- /dev/null
+++ b/Engine/Documentation/Source/Programming/Development/BuildGraphs/BuildGraphTypes/BuildGraphTypes.INT.udn
@@ -0,0 +1,61 @@
+Availability: Docs
+Title: BuildGraph Types
+Crumbs: %ROOT%, Programming, Programming/Development, Programming/Development/BuildGraphs
+Description:Learn about valid data types for BuildGraph attributes.
+version: 4.13
+parent:Programming/Development/BuildGraphs
+type:Reference
+related: Programming/Development/BuildGraphs/BuildGraphScriptAnatomy
+tags:BuildGraph
+
+[VAR:TopicCompact]
+[OBJECT:TopicCompact]
+	[PARAM:image]
+		![%Programming/Development/BuildGraphs/BuildGraphTypes%](types_topic.png)
+	[/PARAM]
+	[PARAM:icon]
+		![](%ROOT%/reference_icon.png)(convert:false)
+	[/PARAM]
+	[PARAM:title]
+		%Programming/Development/BuildGraphs/BuildGraphTypes:title%
+	[/PARAM]
+	[PARAM:description]
+		%Programming/Development/BuildGraphs/BuildGraphTypes:description%
+	[/PARAM]
+	[PARAM:path]
+		[RELATIVE:Programming/Development/BuildGraphs/BuildGraphTypes]
+	[/PARAM]
+[/OBJECT]
+[/VAR]
+
+[EXCERPT:buildgraphtypes]
+
+## Types
+
+The following table contains valid data types that can be contained by BuildGraph attributes:
+
+[INCLUDE:#TypesTable]
+
+[/EXCERPT:buildgraphtypes]
+
+<!--
+[EXCERPT:TypesTable]
+| **Type**                | **Description**                                                                                                                                                                                                                                                                            |
+| ----------------------- | ------------------------------------------------------------------------------------------------------------------------------------------------------------------------------------------------------------------------------------------------------------------------------------------ |
+| **String**              | An arbitrary string.                                                                                                                                                                                                                                                                       |
+| **String List**         | A list of arbitrary strings separated by semi-colons.                                                                                                                                                                                                                                      |
+| **Boolean**             | The constant `true` or `false`.                                                                                                                                                                                                                                                            |
+| **Integer**             | An integer constant.                                                                                                                                                                                                                                                                       |
+| **Regex**               | A regular expression, using [C#](http://www.iso.org/iso/home/store/catalogue_ics/catalogue_detail_ics.htm?csnumber=42926) syntax.                                                                                                                                                          |
+| **Name**                | A named entity. Any printable character, except `^  <  >  :  "  /  \  |  ?  *  ;`. Single spaces are allowed (except at the start and end of a name).                                                                                                                                      |
+| **Name List**           | A list of identifiers separated by semi-colons.                                                                                                                                                                                                                                            |
+| **Tag**                 | A label given to a list of files, beginning with the # character (i.e., `#My Files`).                                                                                                                                                                                                      |
+| **Tag List**            | A list of tags separated by semi-colons. (i.e., `#My Files;#Other Files`).                                                                                                                                                                                                                 |
+| **Target**              | A node name, aggregate name, agent name, or tag name. Indicates a sequence of nodes that need to be executed. Note that this overlaps with the meaning of `Target` as it applies to [UnrealBuildTool](Programming/UnrealBuildSystem) (which refers to a single program being built).       |
+| **Target list**         | A list of target names separated by semi-colons.                                                                                                                                                                                                                                           |
+| **File Spec**           | A set of file and tag names, separated by semi-colons. Wildcards such as `"..."`, `"*"`, and `"?"` are permitted when referencing files (i.e., `Engine/.../*.bat`). Unless otherwise specified, relative paths are resolved relative to the working root directory.                        |
+| **File Name**           | A path to a file. Unless otherwise specified, relative paths are resolved to the working root directory.                                                                                                                                                                                   |
+| **Directory Name**      | A path to a directory. Unless otherwise specified, relative paths are resolved relative to the working root directory.                                                                                                                                                                     |
+| **Condition**           | A [conditional expression](Programming/Development/BuildGraphs/BuildGraphConditions).                                                                                                                                                                                                      |
+[/EXCERPT:TypesTable]
+-->
\ No newline at end of file
