First draft of BuildGraph Anatomy reference page.

#UE4doc

diff --git a/Engine/Documentation/Source/Programming/Development/BuildGraphs/BuildGraphScriptAnatomy/BuildGraphScriptAnatomy.INT.udn b/Engine/Documentation/Source/Programming/Development/BuildGraphs/BuildGraphScriptAnatomy/BuildGraphScriptAnatomy.INT.udn
new file mode 100644
--- /dev/null
+++ b/Engine/Documentation/Source/Programming/Development/BuildGraphs/BuildGraphScriptAnatomy/BuildGraphScriptAnatomy.INT.udn
@@ -0,0 +1,72 @@
+Availability: Docs
+Title: BuildGraph Script Anatomy
+Crumbs: %ROOT%, Programming, Programming/Development, Programming/Development/BuildGraphs
+Description:This page introduces you to our BuildGraph scripting system.
+version: 4.13
+parent:Programming/Development/BuildGraphs
+type:Overview
+tags:BuildGraph
+
+[VAR:TopicCompact]
+[OBJECT:TopicCompact]
+	[PARAM:image]
+		![%Programming/Development/BuildGraphs/BuildGraphScriptAnatomy%](anatomy_topic.png)
+	[/PARAM]
+	[PARAM:icon]
+		![](%ROOT%/reference_icon.png)(convert:false)
+	[/PARAM]
+	[PARAM:title]
+		%Programming/Development/BuildGraphs/BuildGraphScriptAnatomy:title%
+	[/PARAM]
+	[PARAM:description]
+		%Programming/Development/BuildGraphs/BuildGraphScriptAnatomy:description%
+	[/PARAM]
+	[PARAM:path]
+		[RELATIVE:Programming/Development/BuildGraphs/BuildGraphScriptAnatomy]
+	[/PARAM]
+[/OBJECT]
+[/VAR]
+
+[TOC(start:1 end:4)]
+
+If you want to learn how to write your own **BuildGraph** scripts, it helps to know the various parts that go into making a Graph. 
+Graphs are created with the following elements:
+
+* Tasks: Actions that are executed as part of a build process (compiling, cooking, etc.).
+* Nodes: A named sequence of ordered tasks that are executed to produce outputs. Before they can be executed, nodes may depend 
+on other nodes executing their tasks first.
+* Agents: A group of nodes that are executed on the same machine (if running as part of a build system). Agents have no effect 
+when building locally. 
+* Triggers: Container for groups that should only be executed after manual intervention. 
+* Aggregates: Groups of nodes and named outputs that can be referred to with a single name.
+
+Scripts typically make frequent use of properties for reusable or conditionally defined values. Properties are declared with 
+the `<Property>` element, and are scoped to the point of their first declaration. Properties referenced with the `$(Property 
+Name)` notation are valid within all attribute strings, and will be expanded when the script is read. Properties that can be 
+supplied by the user on the command-line can be declared with the `<Option>` element, and environment variables can be 
+imported into properties using the `<EnvVar>` element.
+
+Any element may be conditionally defined via the "If" attribute. See below for the syntax of conditional expressions.
+
+BuildGraph is typically used for packaged games, so filtering and file manipulation are natively supported. Any attribute 
+that accepts a list of files may consist of a Perforce style wildcard (matching any number of "...", "*" and "?" patterns in 
+any location), a full path name, or a reference to a tagged collection of files. Attributes are denoted by prefixing a `#` 
+character. Files may be added to a tag set using the `<Tag>` task, which also allows the performance of union/difference style 
+operations. Each node can declare multiple outputs in the form of a list of named tags, which other nodes can then depend on.
+
+Graphs may be executed in parallel as part of a build system. To do so, the initial graph configuration is generated by 
+running with the `-Export=<Filename>` argument (producing a JSON file listing the nodes and dependencies to execute). Each 
+participating agent should be synced to the same changelist, and AutomationTool (UAT) should be re-run with the appropriate 
+`-SingleNode=<Name>` argument. Outputs from different nodes are transferred between agents via shared storage (typically, a 
+network share) the path to which can be specified on the command line using the `-SharedStorageDir=<Path>` argument. Note that 
+the allocation of machines (and coordination between them) is assumed to be managed by an external system.
+
+The syntax of elements used to define BuildGraph constructs is listed in the following sections.
+
+[INCLUDE:Programming/Development/BuildGraphs/BuildGraphTypes#buildgraphtypes]
+
+[INCLUDE:Programming/Development/BuildGraphs/BuildGraphElements#buildgraphelements]
+
+[INCLUDE:Programming/Development/BuildGraphs/BuildGraphConditions#buildgraphconditions]
+
+[INCLUDE:Programming/Development/BuildGraphs/BuildGraphTasks#buildgraphtasks]
\ No newline at end of file
