First draft of BuildGraph Elements reference page.

#UE4doc

diff --git a/Engine/Documentation/Source/Programming/Development/BuildGraphs/BuildGraphElements/BuildGraphElements.INT.udn b/Engine/Documentation/Source/Programming/Development/BuildGraphs/BuildGraphElements/BuildGraphElements.INT.udn
new file mode 100644
--- /dev/null
+++ b/Engine/Documentation/Source/Programming/Development/BuildGraphs/BuildGraphElements/BuildGraphElements.INT.udn
@@ -0,0 +1,322 @@
+Availability: Docs
+Title: BuildGraph Elements
+Crumbs: %ROOT%, Programming, Programming/Development, Programming/Development/BuildGraphs
+Description:Learn about the elements of a BuildGraph.
+version: 4.13
+parent:Programming/Development/BuildGraphs
+type:Reference
+related: Programming/Development/BuildGraphs/BuildGraphScriptAnatomy
+tags:BuildGraph
+
+[VAR:TopicCompact]
+[OBJECT:TopicCompact]
+	[PARAM:image]
+		![%Programming/Development/BuildGraphs/BuildGraphElements%](elements_topic.png)
+	[/PARAM]
+	[PARAM:icon]
+		![](%ROOT%/reference_icon.png)(convert:false)
+	[/PARAM]
+	[PARAM:title]
+		%Programming/Development/BuildGraphs/BuildGraphElements:title%
+	[/PARAM]
+	[PARAM:description]
+		%Programming/Development/BuildGraphs/BuildGraphElements:description%
+	[/PARAM]
+	[PARAM:path]
+		[RELATIVE:Programming/Development/BuildGraphs/BuildGraphElements]
+	[/PARAM]
+[/OBJECT]
+[/VAR]
+
+[TOC(start:1 end:3)]
+
+[EXCERPT:buildgraphelements]
+
+BuildGraph scripts have elements that are similar to the sort of elements that you'd expect to find in an XML document. 
+As you read through this section, you'll learn how to define a graph along with the syntax that goes into manipulating 
+properties. You'll also learn the syntax needed to control the flow of your logic and finally, we'll show you elements designed 
+to help you with script diagnostics.
+
+## Elements
+
+Elements describe the data that they contain, making them one of the basic building blocks of BuildGraph scripts.
+The following tables contain items that are provided as metadata, which are exported when running on the build system. Although 
+they're not used directly by BuildGraph when executing tasks locally, they're tagged with `[META]`.
+
+## Defining a Graph
+
+A BuildGraph script is typically defined with the following elements:
+
+* `<Node>`
+* `<Aggregate>`
+* `<Agent>`
+* `<Trigger>`
+
+### Node
+
+`<Node>` is the smallest unit of execution in BuildGraph, having a set of inputs and outputs. Each `<Node>` consists of a 
+sequence of tasks that are executed in order.
+
+[INCLUDE:#NodeTable]
+
+### Aggregate
+
+`<Aggregate>` declares a named aggregate that can be used as a synonym for a set of other nodes (or produced tag sets).
+
+[INCLUDE:#AggregateTable]
+
+### Agent
+
+`<Agent>` defines requirements for an agent on nodes that will be executed in sequence (without cleaning intermediate directories). 
+`<Agent>` requirements are ignored when building locally, but must be specified.
+
+[INCLUDE:#AgentTable]
+
+### Trigger
+
+`<Trigger>` serves as a container for portions of the graph that should only be executed with explicit user intervention. 
+To execute nodes behind a trigger, pass `-Trigger=<Name>` to the command line. 
+
+[INCLUDE:#TriggerTable]
+
+## Property Manipulation
+
+BuildGraph properties can be modified with the following elements:
+
+* `<Property>`
+* `<Option>`
+* `<EnvVar>`
+
+### Property
+
+`<Property>` sets the value of a property. If a property in an outer scope has already been declared with the same name,
+`<Property>` overwrites it. Otherwise, a new property is declared at the current scope.
+
+[INCLUDE:#PropertyTable]
+
+### Option
+
+`<Option>` defines a user option that can be set from the command line. `<Option>` may only appear at a global scope.
+
+[INCLUDE:#OptionTable]
+
+### EnvVar
+
+`<EnvVar>` declares a property to contain the contents of an environment variable (or an empty string if it's not set).
+`<EnvVar>` may only appear at a global scope.
+
+[INCLUDE:#EnvVarTable]
+
+## Flow Control
+
+You can control the flow of your BuildGraph scripts with the following elements:
+
+* `<Include>`
+* `<Do>`
+* `<Switch>`
+* `<ForEach>`
+
+### Include
+
+`<Include>` processes the contents of another BuildGraph script as if it appeared within this file.
+`<Include>` may only appear at a global scope.
+
+[INCLUDE:#IncludeTable]
+
+### Do
+
+`<Do>` groups a sequence of elements and only processes them if a condition evaluates as true.
+
+[INCLUDE:#DoTable]
+
+### Switch
+
+`<Switch>` is similar to a [C](http://www.open-std.org/jtc1/sc22/wg14/www/standards) switch statement because it evaluates 
+the conditions on a sequence of branches, processing the first statement with a condition that evaluates as true.
+
+Example `<Switch>` statement:
+
+	<Switch>
+		<Case If=”X”>
+			<Elements go here>
+		</Case>
+		<Case If=”Y”>
+			<Elements go here>
+		</Case>
+		<Default>
+			<Elements go here>
+		</Default>
+	</Switch>
+
+[INCLUDE:#SwitchTable]
+
+### ForEach
+
+`<ForEach>` expands the body of the element with a given property assigned to each item in a list that is 
+separated by semicolons.
+
+Example `<ForEach>` statement:
+
+	<ForEach Name="Counter" Values="1;2;3;4;5">
+		<Log Message="Counter=$(Counter)"/>
+	</ForEach>
+
+[INCLUDE:#ForEachTable]
+
+## Diagnostics
+
+You can perform diagnostics on your BuildGraph scripts with the following elements:
+
+* `<Warning>`
+* `<Error>`
+
+### Warning
+
+`<Warning>` outputs a warning message before executing the graph. `<Warning>` may be included in nodes, agents, triggers, or 
+at global scope. `<Warning>` will only provide an output if it's still part of the graph after it has been trimmed down to 
+the target being executed.
+
+[INCLUDE:#WarningTable]
+
+### Error
+
+`<Error>`outputs an error message. `<Error>` may be included in nodes, agents, triggers, or at global scope. 
+`<Error>` will only provide an output if it's still part of the graph after it has been trimmed down to the target being executed.
+
+[INCLUDE:#ErrorTable]
+
+[/EXCERPT:buildgraphelements]
+
+<!--Tables for the 'Defining a Graph' section.-->
+
+<!--
+[EXCERPT:NodeTable]
+|                  |              |              |                                                                                                                                             |
+| ---------------- | ------------ | ------------ | ------------------------------------------------------------------------------------------------------------------------------------------- |
+| Name             | Name         | Required     | Name of the node.                                                                                                                           |
+| Requires         | Target List  | Optional     | List of nodes, aggregates, or tagged file sets produced by other nodes that this node requires to execute, separated by semicolons.         |
+| Produces         | Tag List     | Optional     | Tagged file sets that this node makes available to other nodes, separated by semicolons.                                                    |
+| After            | Target List  | Optional     | List of nodes that this node should run after, if they are part of the current target (ignored if they are not). Separated by semicolons.   |
+| NotifyOnWarnings | Boolean      | Optional     | If false, this node will not produce notifications on warnings. Only used by build systems. Defaults to `true`. `[META]`                    |
+| If               | Condition    | Optional     | Condition to be evaluated. The element is ignored unless the condition evaluates to `true`.                                                 |
+[/EXCERPT:NodeTable]
+-->
+
+<!--
+[EXCERPT:AggregateTable]
+|                  |              |              |                                                                                             |
+| ---------------- | ------------ | ------------ | ------------------------------------------------------------------------------------------- |
+| Name             | Name         | Required     | Name of the aggregate.                                                                      |
+| Requires         | Target List  | Required     | List of dependencies for this aggregate. May be nodes, tagged file sets, or agent groups.   |
+| If               | Condition    | Optional     | Condition to be evaluated. The element is ignored unless the condition evaluates to `true`. |
+[/EXCERPT:AggregateTable]
+-->
+
+<!--
+[EXCERPT:AgentTable]
+|                  |                 |              |                                                                                                                                            |
+| ---------------- | --------------- | ------------ | ------------------------------------------------------------------------------------------------------------------------------------------ |
+| Name             | Name            | Required     | Name of the group.                                                                                                                         |
+| Type             | Identifier List | Optional     | Type of agent to run on. The meaning of this string is inferred by the host build system; it does not have any intrinsic meaning. `[META]` |
+| If               | Condition       | Optional     | Condition to be evaluated. The element is ignored unless the condition evaluates to `true`.                                                |
+[/EXCERPT:AgentTable]
+-->
+
+<!--
+[EXCERPT:TriggerTable]
+|                  |                 |              |                                                                                             |
+| ---------------- | --------------- | ------------ | ------------------------------------------------------------------------------------------- |
+| Name             | Name            | Required     | Name of the trigger.                                                                        |
+| If               | Condition       | Optional     | Condition to be evaluated. The element is ignored unless the condition evaluates to `true`. |
+[/EXCERPT:TriggerTable]
+-->
+
+<!--Tables for the 'Property Manipulation' section.-->
+
+<!--
+[EXCERPT:PropertyTable]
+|                  |            |          |                                                                                             |
+| ---------------- | ---------- | -------- | ------------------------------------------------------------------------------------------- |
+| Name             | Name       | Required | Name of the property to set.                                                                |
+| Value            | String     | Required | New value for the property.                                                                 |
+| If               | Condition  | Optional | Condition to be evaluated. The element is ignored unless the condition evaluates to `true`. |
+[/EXCERPT:PropertyTable]
+-->
+
+<!--
+[EXCERPT:OptionTable]
+|              |           |              |                                                                                             |
+| ------------ | --------- | ------------ | ------------------------------------------------------------------------------------------- |
+| Name         | Name      | Required     | Name of the option (and property) to initialize with the option's value.                    |
+| Description  | String    | Required     | Description for the option to display when running BuildGraph with the -ListOnly argument.  |
+| Restrict     | Regex     | Optional     | Regex that matches valid values for this option (eg. `[a-zA-Z]+`, `true | false`).          |
+| DefaultValue | String    | Required     | Default value for the option if the user does not set it explicitly.                        |
+| If           | Condition | Optional     | Condition to be evaluated. The element is ignored unless the condition evaluates to `true`. |
+[/EXCERPT:OptionTable]
+-->
+
+<!--
+[EXCERPT:EnvVarTable]
+|                  |             |          |                                                                                             |
+| ---------------- | ----------- | -------- | ------------------------------------------------------------------------------------------- |
+| Name             | Name        | Required | Name of the environment variable to introduce as a property.                                |
+| If               | Condition   | Optional | Condition to be evaluated. The element is ignored unless the condition evaluates to `true`. |
+[/EXCERPT:EnvVarTable]
+-->
+
+<!--Tables for the 'Flow Control' section.-->
+
+<!--
+[EXCERPT:IncludeTable]
+|                  |                 |              |                                                                                             |
+| ---------------- | --------------- | ------------ | ------------------------------------------------------------------------------------------- |
+| Script           | String          | Required     | Path to the script you wan to include (relative to the current script).                     |
+| If               | Condition       | Optional     | Condition to be evaluated. The element is ignored unless the condition evaluates to `true`. |
+[/EXCERPT:IncludeTable]
+-->
+
+<!--
+[EXCERPT:DoTable]
+|    |           |          |                                                                                             |
+| -- | ----------| ---------| ------------------------------------------------------------------------------------------- |
+| If | Condition | Optional | Condition to be evaluated. The element is ignored unless the condition evaluates to `true`. |
+[/EXCERPT:DoTable]
+-->
+
+<!--
+[EXCERPT:SwitchTable]
+|    |           |          |                                                                                             |
+| -- | ----------| ---------| ------------------------------------------------------------------------------------------- |
+| If | Condition | Optional | Condition to be evaluated. The element is ignored unless the condition evaluates to `true`. |
+[/EXCERPT:SwitchTable]
+-->
+
+<!--
+[EXCERPT:ForEachTable]
+|                  |                 |              |                                                                                             |
+| ---------------- | --------------- | ------------ | ------------------------------------------------------------------------------------------- |
+| Name             | Name            | Required     | Property to assign to each value in the list.                                               |
+| Values           | String List     | Required     | List of values that are separated by semicolons.                                            |
+| If               | Condition       | Optional     | Condition to be evaluated. The element is ignored unless the condition evaluates to `true`. |
+[/EXCERPT:ForEachTable]
+-->
+
+<!--Tables for the 'Diagnostics' section.-->
+
+<!--
+[EXCERPT:WarningTable]
+|                  |                 |              |                                                                                             |
+| ---------------- | --------------- | ------------ | ------------------------------------------------------------------------------------------- |
+| Message          | String          | Required     | Text to be printed to the log.                                                              |
+| If               | Condition       | Optional     | Condition to be evaluated. The element is ignored unless the condition evaluates to `true`. |
+[/EXCERPT:WarningTable]
+-->
+
+<!--
+[EXCERPT:ErrorTable]
+|                  |                 |              |                                                                                             |
+| ---------------- | --------------- | ------------ | ------------------------------------------------------------------------------------------- |
+| Message          | String          | Required     | Text to be printed to the log.                                                              |
+| If               | Condition       | Optional     | Condition to be evaluated. The element is ignored unless the condition evaluates to `true`. |
+[/EXCERPT:ErrorTable]
+-->
\ No newline at end of file
