First draft of BuildGraph Conditions reference page.

#UE4doc

diff --git a/Engine/Documentation/Source/Programming/Development/BuildGraphs/BuildGraphConditions/BuildGraphConditions.INT.udn b/Engine/Documentation/Source/Programming/Development/BuildGraphs/BuildGraphConditions/BuildGraphConditions.INT.udn
new file mode 100644
--- /dev/null
+++ b/Engine/Documentation/Source/Programming/Development/BuildGraphs/BuildGraphConditions/BuildGraphConditions.INT.udn
@@ -0,0 +1,83 @@
+Availability: Docs
+Title: BuildGraph Conditions
+Crumbs: %ROOT%, Programming, Programming/Development, Programming/Development/BuildGraphs
+Description:Learn the syntax needed to write BuildGraph script conditions.
+version: 4.13
+parent:Programming/Development/BuildGraphs
+type:Reference
+related: Programming/Development/BuildGraphs/BuildGraphScriptAnatomy
+tags:BuildGraph
+
+[VAR:TopicCompact]
+[OBJECT:TopicCompact]
+	[PARAM:image]
+		![%Programming/Development/BuildGraphs/BuildGraphConditions%](conditional_topic.png)
+	[/PARAM]
+	[PARAM:icon]
+		![](%ROOT%/reference_icon.png)(convert:false)
+	[/PARAM]
+	[PARAM:title]
+		%Programming/Development/BuildGraphs/BuildGraphConditions:title%
+	[/PARAM]
+	[PARAM:description]
+		%Programming/Development/BuildGraphs/BuildGraphConditions:description%
+	[/PARAM]
+	[PARAM:path]
+		[RELATIVE:Programming/Development/BuildGraphs/BuildGraphConditions]
+	[/PARAM]
+[/OBJECT]
+[/VAR]
+
+[TOC(start:1 end:3)]
+
+[EXCERPT:buildgraphconditions]
+
+If you want to build logical complexity into your BuildGraph scripts, you'll need to work with conditionals. The following 
+section introduces you to how BuildGraph conditions are written, including a list of conditional operators. 
+
+## Conditions
+
+BuildGraph script conditions consist of atoms and operators that evaluate to `true` or `false`.
+
+### Atoms
+
+Atoms can be numbers, strings, or identifiers that are coerced to their appropriate type for the oprator using them. 
+Atoms may be quoted with single (') or double (") quotes. They can also be an unquoted sequence of letters, digits, and 
+underscore characters. 
+
+### Literals
+
+All literals are considered identical, regardless of how they're declared. Additionally, literals are considered 
+case-insensitive for comparisons, which means that the strings "True" and 'true' are identical to the identifier `true` 
+(despite the presence of quotes and differences in case).
+
+## Operators
+
+A list of operators is specified below: 
+
+[INCLUDE:#OperatorTable]
+
+[REGION:note]
+The `'<'` and `'>'` characters must be escaped as `"&lt;"` and `"&gt;"` in XML (referencing [XML Schema v1.0](https://www.w3.org/TR/NOTE-dcd])).
+[/REGION]
+
+[/EXCERPT:buildgraphconditions]
+
+<!--
+[EXCERPT:OperatorTable]
+| **Operator**            | **Description**                                                             | **Precedence**  |
+| ----------------------- | --------------------------------------------------------------------------- | --------------- |
+| **(x)**                 | Subexpression                                                               | 1               |
+| **!x**                  | Not operator                                                                | 1               |
+| **Exists(x)**           | True if the x file exists.                                                  | 1               |
+| **HasTrailingSlash(x)** | True if x ends with a slash or backslash.                                   | 1               |
+| **x == y**              | Tests two atoms for equality (case insensitive).                            | 2               |
+| **x != y**              | Tests two atoms for inequality (case insensitive).                          | 2               |
+| **x < y**               | Compares whether the integer x is less than the integer y.                  | 2               |
+| **x <= y**              | Compares whether the integer x is less than or equal to the integer y.      | 2               |
+| **x > y**               | Compares whether the integer x is greater than the integer y.               | 2               |
+| **x >= y**              | Compares whether the integer x is greater than or equal to the integer y.   | 2               |
+| **x and y**             | True if both arguments are `true`.                                          | 3               |
+| **x or y**              | True if either argument is `true`.                                          | 4               |
+[/EXCERPT:OperatorTable]
+-->
\ No newline at end of file
