Adding topic image and TopicCompact markup for inclusion in BuildGraph Landing page.

#UE4doc

diff --git a/Engine/Documentation/Source/Programming/Development/InstalledBuildReference/InstalledBuildReference.INT.udn b/Engine/Documentation/Source/Programming/Development/InstalledBuildReference/InstalledBuildReference.INT.udn
--- a/Engine/Documentation/Source/Programming/Development/InstalledBuildReference/InstalledBuildReference.INT.udn
+++ b/Engine/Documentation/Source/Programming/Development/InstalledBuildReference/InstalledBuildReference.INT.udn
@@ -10,6 +10,26 @@ related:Programming/UnrealBuildSystem
 
 [TOC (start:1 end:4)]
 
+[VAR:TopicCompact]
+[OBJECT:TopicCompact]
+	[PARAM:image]
+		![%Programming/Development/InstalledBuildReference%](installedbuild_topic.png)
+	[/PARAM]
+	[PARAM:icon]
+		![](%ROOT%/reference_icon.png)(convert:false)
+	[/PARAM]
+	[PARAM:title]
+		%Programming/Development/InstalledBuildReference:title%
+	[/PARAM]
+	[PARAM:description]
+		%Programming/Development/InstalledBuildReference:description%
+	[/PARAM]
+	[PARAM:path]
+		[RELATIVE:Programming/Development/InstalledBuildReference]
+	[/PARAM]
+[/OBJECT]
+[/VAR]
+
 Developers are always looking for ways to efficiently deliver build solutions for a variety of hardware profiles and operating systems. 
 With the release of our BuildGraph scripting system, you'll be able to create custom sequences of Unreal Engine 4 (UE4) build commands, 
 harnessing the Installed Build process that we've developed for your use. As you read through the following reference page, you'll learn 
