Rename/move file(s)

diff --git a/Engine/Documentation/Source/Programming/Development/BuildGraph/BuildGraph.INT.udn b/Engine/Documentation/Source/Programming/Development/BuildGraph/BuildGraph.INT.udn
new file mode 100644
--- /dev/null
+++ b/Engine/Documentation/Source/Programming/Development/BuildGraph/BuildGraph.INT.udn
@@ -0,0 +1,32 @@
+Availability: Docs
+Title: BuildGraph Scripting
+Crumbs: %ROOT%, Programming, Programming/Development
+Description: Customize your Installed Builds by using our BuildGraph scripting system.
+Version: 4.13
+type:landing
+parent:Programming/Development
+tags:BuildGraph
+
+[INCLUDE:Programming/Development/BuildGraphs/BuildGraphIntroduction]
+
+## Writing BuildGraph Scripts
+
+[REGION:topics third]
+%Programming/Development/BuildGraphs/BuildGraphScriptAnatomy:topiccompact%
+[/REGION]
+
+## Using BuildGraph
+
+[REGION:topics third]
+%Programming/Development/BuildGraphs/BuildGraphUsage:topiccompact%
+%Programming/Development/InstalledBuildReference:topiccompact%
+[/REGION]
+
+## Reference
+
+[REGION:topics third]
+%Programming/Development/BuildGraphs/BuildGraphTypes:topiccompact%
+%Programming/Development/BuildGraphs/BuildGraphElements:topiccompact%
+%Programming/Development/BuildGraphs/BuildGraphConditions:topiccompact%
+%Programming/Development/BuildGraphs/BuildGraphTasks:topiccompact%
+[/REGION]
\ No newline at end of file
diff --git a/Engine/Documentation/Source/Programming/Development/BuildGraph/BuildGraphConditions/BuildGraphConditions.INT.udn b/Engine/Documentation/Source/Programming/Development/BuildGraph/BuildGraphConditions/BuildGraphConditions.INT.udn
new file mode 100644
--- /dev/null
+++ b/Engine/Documentation/Source/Programming/Development/BuildGraph/BuildGraphConditions/BuildGraphConditions.INT.udn
@@ -0,0 +1,83 @@
+Availability: Docs
+Title: BuildGraph Conditions
+Crumbs: %ROOT%, Programming, Programming/Development, Programming/Development/BuildGraphs
+Description:Learn the syntax needed to write BuildGraph script conditions.
+version: 4.13
+parent:Programming/Development/BuildGraphs
+type:Reference
+related: Programming/Development/BuildGraphs/BuildGraphScriptAnatomy
+tags:BuildGraph
+
+[VAR:TopicCompact]
+[OBJECT:TopicCompact]
+	[PARAM:image]
+		![%Programming/Development/BuildGraphs/BuildGraphConditions%](conditional_topic.png)
+	[/PARAM]
+	[PARAM:icon]
+		![](%ROOT%/reference_icon.png)(convert:false)
+	[/PARAM]
+	[PARAM:title]
+		%Programming/Development/BuildGraphs/BuildGraphConditions:title%
+	[/PARAM]
+	[PARAM:description]
+		%Programming/Development/BuildGraphs/BuildGraphConditions:description%
+	[/PARAM]
+	[PARAM:path]
+		[RELATIVE:Programming/Development/BuildGraphs/BuildGraphConditions]
+	[/PARAM]
+[/OBJECT]
+[/VAR]
+
+[TOC(start:1 end:3)]
+
+[EXCERPT:buildgraphconditions]
+
+If you want to build logical complexity into your BuildGraph scripts, you'll need to work with conditionals. The following 
+section introduces you to how BuildGraph conditions are written, including a list of conditional operators. 
+
+## Conditions
+
+BuildGraph script conditions consist of atoms and operators that evaluate to `true` or `false`.
+
+### Atoms
+
+Atoms can be numbers, strings, or identifiers that are coerced to their appropriate type for the oprator using them. 
+Atoms may be quoted with single (') or double (") quotes. They can also be an unquoted sequence of letters, digits, and 
+underscore characters. 
+
+### Literals
+
+All literals are considered identical, regardless of how they're declared. Additionally, literals are considered 
+case-insensitive for comparisons, which means that the strings "True" and 'true' are identical to the identifier `true` 
+(despite the presence of quotes and differences in case).
+
+## Operators
+
+A list of operators is specified below: 
+
+[INCLUDE:#OperatorTable]
+
+[REGION:note]
+The `'<'` and `'>'` characters must be escaped as `"&lt;"` and `"&gt;"` in XML (referencing [XML Schema v1.0](https://www.w3.org/TR/NOTE-dcd])).
+[/REGION]
+
+[/EXCERPT:buildgraphconditions]
+
+<!--
+[EXCERPT:OperatorTable]
+| **Operator**            | **Description**                                                             | **Precedence**  |
+| ----------------------- | --------------------------------------------------------------------------- | --------------- |
+| **(x)**                 | Subexpression                                                               | 1               |
+| **!x**                  | Not operator                                                                | 1               |
+| **Exists(x)**           | True if the x file exists.                                                  | 1               |
+| **HasTrailingSlash(x)** | True if x ends with a slash or backslash.                                   | 1               |
+| **x == y**              | Tests two atoms for equality (case insensitive).                            | 2               |
+| **x != y**              | Tests two atoms for inequality (case insensitive).                          | 2               |
+| **x < y**               | Compares whether the integer x is less than the integer y.                  | 2               |
+| **x <= y**              | Compares whether the integer x is less than or equal to the integer y.      | 2               |
+| **x > y**               | Compares whether the integer x is greater than the integer y.               | 2               |
+| **x >= y**              | Compares whether the integer x is greater than or equal to the integer y.   | 2               |
+| **x and y**             | True if both arguments are `true`.                                          | 3               |
+| **x or y**              | True if either argument is `true`.                                          | 4               |
+[/EXCERPT:OperatorTable]
+-->
\ No newline at end of file
diff --git a/Engine/Documentation/Source/Programming/Development/BuildGraph/BuildGraphElements/BuildGraphElements.INT.udn b/Engine/Documentation/Source/Programming/Development/BuildGraph/BuildGraphElements/BuildGraphElements.INT.udn
new file mode 100644
--- /dev/null
+++ b/Engine/Documentation/Source/Programming/Development/BuildGraph/BuildGraphElements/BuildGraphElements.INT.udn
@@ -0,0 +1,322 @@
+Availability: Docs
+Title: BuildGraph Elements
+Crumbs: %ROOT%, Programming, Programming/Development, Programming/Development/BuildGraphs
+Description:Learn about the elements of a BuildGraph.
+version: 4.13
+parent:Programming/Development/BuildGraphs
+type:Reference
+related: Programming/Development/BuildGraphs/BuildGraphScriptAnatomy
+tags:BuildGraph
+
+[VAR:TopicCompact]
+[OBJECT:TopicCompact]
+	[PARAM:image]
+		![%Programming/Development/BuildGraphs/BuildGraphElements%](elements_topic.png)
+	[/PARAM]
+	[PARAM:icon]
+		![](%ROOT%/reference_icon.png)(convert:false)
+	[/PARAM]
+	[PARAM:title]
+		%Programming/Development/BuildGraphs/BuildGraphElements:title%
+	[/PARAM]
+	[PARAM:description]
+		%Programming/Development/BuildGraphs/BuildGraphElements:description%
+	[/PARAM]
+	[PARAM:path]
+		[RELATIVE:Programming/Development/BuildGraphs/BuildGraphElements]
+	[/PARAM]
+[/OBJECT]
+[/VAR]
+
+[TOC(start:1 end:3)]
+
+[EXCERPT:buildgraphelements]
+
+BuildGraph scripts have elements that are similar to the sort of elements that you'd expect to find in an XML document. 
+As you read through this section, you'll learn how to define a graph along with the syntax that goes into manipulating 
+properties. You'll also learn the syntax needed to control the flow of your logic and finally, we'll show you elements designed 
+to help you with script diagnostics.
+
+## Elements
+
+Elements describe the data that they contain, making them one of the basic building blocks of BuildGraph scripts.
+The following tables contain items that are provided as metadata, which are exported when running on the build system. Although 
+they're not used directly by BuildGraph when executing tasks locally, they're tagged with `[META]`.
+
+## Defining a Graph
+
+A BuildGraph script is typically defined with the following elements:
+
+* `<Node>`
+* `<Aggregate>`
+* `<Agent>`
+* `<Trigger>`
+
+### Node
+
+`<Node>` is the smallest unit of execution in BuildGraph, having a set of inputs and outputs. Each `<Node>` consists of a 
+sequence of tasks that are executed in order.
+
+[INCLUDE:#NodeTable]
+
+### Aggregate
+
+`<Aggregate>` declares a named aggregate that can be used as a synonym for a set of other nodes (or produced tag sets).
+
+[INCLUDE:#AggregateTable]
+
+### Agent
+
+`<Agent>` defines requirements for an agent on nodes that will be executed in sequence (without cleaning intermediate directories). 
+`<Agent>` requirements are ignored when building locally, but must be specified.
+
+[INCLUDE:#AgentTable]
+
+### Trigger
+
+`<Trigger>` serves as a container for portions of the graph that should only be executed with explicit user intervention. 
+To execute nodes behind a trigger, pass `-Trigger=<Name>` to the command line. 
+
+[INCLUDE:#TriggerTable]
+
+## Property Manipulation
+
+BuildGraph properties can be modified with the following elements:
+
+* `<Property>`
+* `<Option>`
+* `<EnvVar>`
+
+### Property
+
+`<Property>` sets the value of a property. If a property in an outer scope has already been declared with the same name,
+`<Property>` overwrites it. Otherwise, a new property is declared at the current scope.
+
+[INCLUDE:#PropertyTable]
+
+### Option
+
+`<Option>` defines a user option that can be set from the command line. `<Option>` may only appear at a global scope.
+
+[INCLUDE:#OptionTable]
+
+### EnvVar
+
+`<EnvVar>` declares a property to contain the contents of an environment variable (or an empty string if it's not set).
+`<EnvVar>` may only appear at a global scope.
+
+[INCLUDE:#EnvVarTable]
+
+## Flow Control
+
+You can control the flow of your BuildGraph scripts with the following elements:
+
+* `<Include>`
+* `<Do>`
+* `<Switch>`
+* `<ForEach>`
+
+### Include
+
+`<Include>` processes the contents of another BuildGraph script as if it appeared within this file.
+`<Include>` may only appear at a global scope.
+
+[INCLUDE:#IncludeTable]
+
+### Do
+
+`<Do>` groups a sequence of elements and only processes them if a condition evaluates as true.
+
+[INCLUDE:#DoTable]
+
+### Switch
+
+`<Switch>` is similar to a [C](http://www.open-std.org/jtc1/sc22/wg14/www/standards) switch statement because it evaluates 
+the conditions on a sequence of branches, processing the first statement with a condition that evaluates as true.
+
+Example `<Switch>` statement:
+
+	<Switch>
+		<Case If=”X”>
+			<Elements go here>
+		</Case>
+		<Case If=”Y”>
+			<Elements go here>
+		</Case>
+		<Default>
+			<Elements go here>
+		</Default>
+	</Switch>
+
+[INCLUDE:#SwitchTable]
+
+### ForEach
+
+`<ForEach>` expands the body of the element with a given property assigned to each item in a list that is 
+separated by semicolons.
+
+Example `<ForEach>` statement:
+
+	<ForEach Name="Counter" Values="1;2;3;4;5">
+		<Log Message="Counter=$(Counter)"/>
+	</ForEach>
+
+[INCLUDE:#ForEachTable]
+
+## Diagnostics
+
+You can perform diagnostics on your BuildGraph scripts with the following elements:
+
+* `<Warning>`
+* `<Error>`
+
+### Warning
+
+`<Warning>` outputs a warning message before executing the graph. `<Warning>` may be included in nodes, agents, triggers, or 
+at global scope. `<Warning>` will only provide an output if it's still part of the graph after it has been trimmed down to 
+the target being executed.
+
+[INCLUDE:#WarningTable]
+
+### Error
+
+`<Error>`outputs an error message. `<Error>` may be included in nodes, agents, triggers, or at global scope. 
+`<Error>` will only provide an output if it's still part of the graph after it has been trimmed down to the target being executed.
+
+[INCLUDE:#ErrorTable]
+
+[/EXCERPT:buildgraphelements]
+
+<!--Tables for the 'Defining a Graph' section.-->
+
+<!--
+[EXCERPT:NodeTable]
+|                  |              |              |                                                                                                                                             |
+| ---------------- | ------------ | ------------ | ------------------------------------------------------------------------------------------------------------------------------------------- |
+| Name             | Name         | Required     | Name of the node.                                                                                                                           |
+| Requires         | Target List  | Optional     | List of nodes, aggregates, or tagged file sets produced by other nodes that this node requires to execute, separated by semicolons.         |
+| Produces         | Tag List     | Optional     | Tagged file sets that this node makes available to other nodes, separated by semicolons.                                                    |
+| After            | Target List  | Optional     | List of nodes that this node should run after, if they are part of the current target (ignored if they are not). Separated by semicolons.   |
+| NotifyOnWarnings | Boolean      | Optional     | If false, this node will not produce notifications on warnings. Only used by build systems. Defaults to `true`. `[META]`                    |
+| If               | Condition    | Optional     | Condition to be evaluated. The element is ignored unless the condition evaluates to `true`.                                                 |
+[/EXCERPT:NodeTable]
+-->
+
+<!--
+[EXCERPT:AggregateTable]
+|                  |              |              |                                                                                             |
+| ---------------- | ------------ | ------------ | ------------------------------------------------------------------------------------------- |
+| Name             | Name         | Required     | Name of the aggregate.                                                                      |
+| Requires         | Target List  | Required     | List of dependencies for this aggregate. May be nodes, tagged file sets, or agent groups.   |
+| If               | Condition    | Optional     | Condition to be evaluated. The element is ignored unless the condition evaluates to `true`. |
+[/EXCERPT:AggregateTable]
+-->
+
+<!--
+[EXCERPT:AgentTable]
+|                  |                 |              |                                                                                                                                            |
+| ---------------- | --------------- | ------------ | ------------------------------------------------------------------------------------------------------------------------------------------ |
+| Name             | Name            | Required     | Name of the group.                                                                                                                         |
+| Type             | Identifier List | Optional     | Type of agent to run on. The meaning of this string is inferred by the host build system; it does not have any intrinsic meaning. `[META]` |
+| If               | Condition       | Optional     | Condition to be evaluated. The element is ignored unless the condition evaluates to `true`.                                                |
+[/EXCERPT:AgentTable]
+-->
+
+<!--
+[EXCERPT:TriggerTable]
+|                  |                 |              |                                                                                             |
+| ---------------- | --------------- | ------------ | ------------------------------------------------------------------------------------------- |
+| Name             | Name            | Required     | Name of the trigger.                                                                        |
+| If               | Condition       | Optional     | Condition to be evaluated. The element is ignored unless the condition evaluates to `true`. |
+[/EXCERPT:TriggerTable]
+-->
+
+<!--Tables for the 'Property Manipulation' section.-->
+
+<!--
+[EXCERPT:PropertyTable]
+|                  |            |          |                                                                                             |
+| ---------------- | ---------- | -------- | ------------------------------------------------------------------------------------------- |
+| Name             | Name       | Required | Name of the property to set.                                                                |
+| Value            | String     | Required | New value for the property.                                                                 |
+| If               | Condition  | Optional | Condition to be evaluated. The element is ignored unless the condition evaluates to `true`. |
+[/EXCERPT:PropertyTable]
+-->
+
+<!--
+[EXCERPT:OptionTable]
+|              |           |              |                                                                                             |
+| ------------ | --------- | ------------ | ------------------------------------------------------------------------------------------- |
+| Name         | Name      | Required     | Name of the option (and property) to initialize with the option's value.                    |
+| Description  | String    | Required     | Description for the option to display when running BuildGraph with the -ListOnly argument.  |
+| Restrict     | Regex     | Optional     | Regex that matches valid values for this option (eg. `[a-zA-Z]+`, `true | false`).          |
+| DefaultValue | String    | Required     | Default value for the option if the user does not set it explicitly.                        |
+| If           | Condition | Optional     | Condition to be evaluated. The element is ignored unless the condition evaluates to `true`. |
+[/EXCERPT:OptionTable]
+-->
+
+<!--
+[EXCERPT:EnvVarTable]
+|                  |             |          |                                                                                             |
+| ---------------- | ----------- | -------- | ------------------------------------------------------------------------------------------- |
+| Name             | Name        | Required | Name of the environment variable to introduce as a property.                                |
+| If               | Condition   | Optional | Condition to be evaluated. The element is ignored unless the condition evaluates to `true`. |
+[/EXCERPT:EnvVarTable]
+-->
+
+<!--Tables for the 'Flow Control' section.-->
+
+<!--
+[EXCERPT:IncludeTable]
+|                  |                 |              |                                                                                             |
+| ---------------- | --------------- | ------------ | ------------------------------------------------------------------------------------------- |
+| Script           | String          | Required     | Path to the script you wan to include (relative to the current script).                     |
+| If               | Condition       | Optional     | Condition to be evaluated. The element is ignored unless the condition evaluates to `true`. |
+[/EXCERPT:IncludeTable]
+-->
+
+<!--
+[EXCERPT:DoTable]
+|    |           |          |                                                                                             |
+| -- | ----------| ---------| ------------------------------------------------------------------------------------------- |
+| If | Condition | Optional | Condition to be evaluated. The element is ignored unless the condition evaluates to `true`. |
+[/EXCERPT:DoTable]
+-->
+
+<!--
+[EXCERPT:SwitchTable]
+|    |           |          |                                                                                             |
+| -- | ----------| ---------| ------------------------------------------------------------------------------------------- |
+| If | Condition | Optional | Condition to be evaluated. The element is ignored unless the condition evaluates to `true`. |
+[/EXCERPT:SwitchTable]
+-->
+
+<!--
+[EXCERPT:ForEachTable]
+|                  |                 |              |                                                                                             |
+| ---------------- | --------------- | ------------ | ------------------------------------------------------------------------------------------- |
+| Name             | Name            | Required     | Property to assign to each value in the list.                                               |
+| Values           | String List     | Required     | List of values that are separated by semicolons.                                            |
+| If               | Condition       | Optional     | Condition to be evaluated. The element is ignored unless the condition evaluates to `true`. |
+[/EXCERPT:ForEachTable]
+-->
+
+<!--Tables for the 'Diagnostics' section.-->
+
+<!--
+[EXCERPT:WarningTable]
+|                  |                 |              |                                                                                             |
+| ---------------- | --------------- | ------------ | ------------------------------------------------------------------------------------------- |
+| Message          | String          | Required     | Text to be printed to the log.                                                              |
+| If               | Condition       | Optional     | Condition to be evaluated. The element is ignored unless the condition evaluates to `true`. |
+[/EXCERPT:WarningTable]
+-->
+
+<!--
+[EXCERPT:ErrorTable]
+|                  |                 |              |                                                                                             |
+| ---------------- | --------------- | ------------ | ------------------------------------------------------------------------------------------- |
+| Message          | String          | Required     | Text to be printed to the log.                                                              |
+| If               | Condition       | Optional     | Condition to be evaluated. The element is ignored unless the condition evaluates to `true`. |
+[/EXCERPT:ErrorTable]
+-->
\ No newline at end of file
diff --git a/Engine/Documentation/Source/Programming/Development/BuildGraph/BuildGraphIntroduction/BuildGraphIntroduction.INT.udn b/Engine/Documentation/Source/Programming/Development/BuildGraph/BuildGraphIntroduction/BuildGraphIntroduction.INT.udn
new file mode 100644
--- /dev/null
+++ b/Engine/Documentation/Source/Programming/Development/BuildGraph/BuildGraphIntroduction/BuildGraphIntroduction.INT.udn
@@ -0,0 +1,31 @@
+Availability: NoPublish
+Title: Introduction to BuildGraphs
+Crumbs: %ROOT%, Programming, Programming/Development, Programming/Development/BuildGraphs
+Description:This page introduces you to our BuildGraph scripting system.
+version: 4.13
+parent:Programming/Development/BuildGraphs
+type:Reference
+
+**BuildGraph** is a scripting language for build automation with building blocks common to Unreal Engine 4 (UE4) projects. 
+BuildGraph integrates with UnrealBuildTool, AutomationTool, and the editor. Which means that BuildGraph can be 
+extended and customized for your projects.
+
+BuildGraph scripts are written with [XML](https://www.w3.org/TR/NOTE-dcd]), and specify a graph of user-defined nodes with 
+dependencies between them. Each node consists of tasks executed in sequence to produce some sort of output (i.e., compile 
+this project, then cook, run this test, and so on). When asked to build a target (that is, a node or named output) BuildGraph 
+will execute all of the nodes in the graph required to make that happen.
+
+Unlike other build tools, BuildGraph is designed as a hybrid between a makefile-like scripting language and a build farm 
+configuration script. It allows annotations for the type of machine that steps are supposed to be executed on, providing a 
+list of recipients for failure notifications if a step fails, and groups nodes that should only be executed after an explicit 
+user trigger. It also tracks the creation of output files from task execution in a way that lends to graph execution being 
+distributed across a farm of machines (with nodes running in parallel where possible), and intermediate artifacts being 
+transferred to and from a network share automatically.
+
+Epic uses BuildGraph to prepare the UE4 binary release, package samples for the marketplace, and implement pipelines for our 
+own games (among other things). Several example BuildGraph scripts are provided in the 
+`[UE4Root]/Engine/Build/Graph/Examples` directory, and the script for creating a binary UE4 distribution can be found at 
+`[UE4Root]/Engine/Build/InstalledEngineBuild.xml`.
+
+Opening a script with Visual Studio will use the schema located at `[UE4Root]Engine/Build/Graph/Schema.xsd` to provide rich 
+tooltips, validation, and Intellisense features while editing.
\ No newline at end of file
diff --git a/Engine/Documentation/Source/Programming/Development/BuildGraph/BuildGraphScriptAnatomy/BuildGraphScriptAnatomy.INT.udn b/Engine/Documentation/Source/Programming/Development/BuildGraph/BuildGraphScriptAnatomy/BuildGraphScriptAnatomy.INT.udn
new file mode 100644
--- /dev/null
+++ b/Engine/Documentation/Source/Programming/Development/BuildGraph/BuildGraphScriptAnatomy/BuildGraphScriptAnatomy.INT.udn
@@ -0,0 +1,72 @@
+Availability: Docs
+Title: BuildGraph Script Anatomy
+Crumbs: %ROOT%, Programming, Programming/Development, Programming/Development/BuildGraphs
+Description:This page introduces you to our BuildGraph scripting system.
+version: 4.13
+parent:Programming/Development/BuildGraphs
+type:Overview
+tags:BuildGraph
+
+[VAR:TopicCompact]
+[OBJECT:TopicCompact]
+	[PARAM:image]
+		![%Programming/Development/BuildGraphs/BuildGraphScriptAnatomy%](anatomy_topic.png)
+	[/PARAM]
+	[PARAM:icon]
+		![](%ROOT%/reference_icon.png)(convert:false)
+	[/PARAM]
+	[PARAM:title]
+		%Programming/Development/BuildGraphs/BuildGraphScriptAnatomy:title%
+	[/PARAM]
+	[PARAM:description]
+		%Programming/Development/BuildGraphs/BuildGraphScriptAnatomy:description%
+	[/PARAM]
+	[PARAM:path]
+		[RELATIVE:Programming/Development/BuildGraphs/BuildGraphScriptAnatomy]
+	[/PARAM]
+[/OBJECT]
+[/VAR]
+
+[TOC(start:1 end:4)]
+
+If you want to learn how to write your own **BuildGraph** scripts, it helps to know the various parts that go into making a Graph. 
+Graphs are created with the following elements:
+
+* Tasks: Actions that are executed as part of a build process (compiling, cooking, etc.).
+* Nodes: A named sequence of ordered tasks that are executed to produce outputs. Before they can be executed, nodes may depend 
+on other nodes executing their tasks first.
+* Agents: A group of nodes that are executed on the same machine (if running as part of a build system). Agents have no effect 
+when building locally. 
+* Triggers: Container for groups that should only be executed after manual intervention. 
+* Aggregates: Groups of nodes and named outputs that can be referred to with a single name.
+
+Scripts typically make frequent use of properties for reusable or conditionally defined values. Properties are declared with 
+the `<Property>` element, and are scoped to the point of their first declaration. Properties referenced with the `$(Property 
+Name)` notation are valid within all attribute strings, and will be expanded when the script is read. Properties that can be 
+supplied by the user on the command-line can be declared with the `<Option>` element, and environment variables can be 
+imported into properties using the `<EnvVar>` element.
+
+Any element may be conditionally defined via the "If" attribute. See below for the syntax of conditional expressions.
+
+BuildGraph is typically used for packaged games, so filtering and file manipulation are natively supported. Any attribute 
+that accepts a list of files may consist of a Perforce style wildcard (matching any number of "...", "*" and "?" patterns in 
+any location), a full path name, or a reference to a tagged collection of files. Attributes are denoted by prefixing a `#` 
+character. Files may be added to a tag set using the `<Tag>` task, which also allows the performance of union/difference style 
+operations. Each node can declare multiple outputs in the form of a list of named tags, which other nodes can then depend on.
+
+Graphs may be executed in parallel as part of a build system. To do so, the initial graph configuration is generated by 
+running with the `-Export=<Filename>` argument (producing a JSON file listing the nodes and dependencies to execute). Each 
+participating agent should be synced to the same changelist, and AutomationTool (UAT) should be re-run with the appropriate 
+`-SingleNode=<Name>` argument. Outputs from different nodes are transferred between agents via shared storage (typically, a 
+network share) the path to which can be specified on the command line using the `-SharedStorageDir=<Path>` argument. Note that 
+the allocation of machines (and coordination between them) is assumed to be managed by an external system.
+
+The syntax of elements used to define BuildGraph constructs is listed in the following sections.
+
+[INCLUDE:Programming/Development/BuildGraphs/BuildGraphTypes#buildgraphtypes]
+
+[INCLUDE:Programming/Development/BuildGraphs/BuildGraphElements#buildgraphelements]
+
+[INCLUDE:Programming/Development/BuildGraphs/BuildGraphConditions#buildgraphconditions]
+
+[INCLUDE:Programming/Development/BuildGraphs/BuildGraphTasks#buildgraphtasks]
\ No newline at end of file
diff --git a/Engine/Documentation/Source/Programming/Development/BuildGraph/BuildGraphTasks/BuildGraphPredefinedTasks/BuildGraphPredefinedTasks.INT.udn b/Engine/Documentation/Source/Programming/Development/BuildGraph/BuildGraphTasks/BuildGraphPredefinedTasks/BuildGraphPredefinedTasks.INT.udn
new file mode 100644
--- /dev/null
+++ b/Engine/Documentation/Source/Programming/Development/BuildGraph/BuildGraphTasks/BuildGraphPredefinedTasks/BuildGraphPredefinedTasks.INT.udn
@@ -0,0 +1,7 @@
+Availability: NoPublish
+Title: BuildGraph Predefined Tasks
+Crumbs: %ROOT%, Programming, Programming/Development, Programming/Development/BuildGraphs, Programming/Development/BuildGraphs/BuildGraphTasks
+Description: This is a procedurally generated markdown page.
+version: 4.13
+parent:Programming/Development/BuildGraphs/BuildGraphTasks
+
diff --git a/Engine/Documentation/Source/Programming/Development/BuildGraph/BuildGraphTasks/BuildGraphTasks.INT.udn b/Engine/Documentation/Source/Programming/Development/BuildGraph/BuildGraphTasks/BuildGraphTasks.INT.udn
new file mode 100644
--- /dev/null
+++ b/Engine/Documentation/Source/Programming/Development/BuildGraph/BuildGraphTasks/BuildGraphTasks.INT.udn
@@ -0,0 +1,58 @@
+Availability: Docs
+Title: BuildGraph Tasks
+Crumbs: %ROOT%, Programming, Programming/Development, Programming/Development/BuildGraphs
+Description: Learn how BuildGraph can help you create custom tasks.
+version: 4.13
+parent:Programming/Development/BuildGraphs
+type:Reference
+related: Programming/Development/BuildGraphs/BuildGraphScriptAnatomy
+tags:BuildGraph
+
+[VAR:TopicCompact]
+[OBJECT:TopicCompact]
+	[PARAM:image]
+		![%Programming/Development/BuildGraphs/BuildGraphTasks%](tasks_topic.png)
+	[/PARAM]
+	[PARAM:icon]
+		![](%ROOT%/reference_icon.png)(convert:false)
+	[/PARAM]
+	[PARAM:title]
+		%Programming/Development/BuildGraphs/BuildGraphTasks:title%
+	[/PARAM]
+	[PARAM:description]
+		%Programming/Development/BuildGraphs/BuildGraphTasks:description%
+	[/PARAM]
+	[PARAM:path]
+		[RELATIVE:Programming/Development/BuildGraphs/BuildGraphTasks]
+	[/PARAM]
+[/OBJECT]
+[/VAR]
+
+[TOC(start:1 end:2)]
+
+[EXCERPT:buildgraphtasks]
+
+Typically, project-specific utility scripts are implemented as AutomationTool BuildCommand classes. There are times; however, when you'll want 
+to automate custom tasks for your project. This is where BuildGraph can help. BuildGraph can be extended with any number of tasks, enabling you 
+to customize your build automation tasks in ways that suit your specific needs. The following section goes over ways of creating custom 
+tasks using BuildGraph, including a listing of predefined tasks that'll give you a great starting point for creating custom tasks.
+
+## Custom Tasks
+
+To create a new custom task, implement a class derived from the `CustomTask` class and apply the `[TaskElement]` attribute to it. 
+The `TaskElement` constructor takes two arguments; the name of the XML element that it's represented by, and the type of class 
+containing its parameters (which is passed to the constructor at load time).
+
+If you want to read parameter class fields from an XML file, attach the `[TaskParameter]` attribute to the task. Attaching the 
+`TaskParameter` attribute indicates if the parameter is required or optional, including additional validation that should be 
+applied to the argument.
+
+## Predefined Tasks
+
+If you'd like to have a starting point for creating tasks, we're providing a variety of predefined tasks as templates for you to work from. 
+If you're looking for a simple predefined task to start from, we recommend using `LogTask` as a good starting point. If you want to keep up with 
+our latest improvements to BuildGraph; bookmark this page, because the following list of predefined tasks will be updated periodically. 
+
+[INCLUDE:Programming/Development/BuildGraphs/BuildGraphTasks/BuildGraphPredefinedTasks]
+
+[/EXCERPT:buildgraphtasks]
\ No newline at end of file
diff --git a/Engine/Documentation/Source/Programming/Development/BuildGraph/BuildGraphTypes/BuildGraphTypes.INT.udn b/Engine/Documentation/Source/Programming/Development/BuildGraph/BuildGraphTypes/BuildGraphTypes.INT.udn
new file mode 100644
--- /dev/null
+++ b/Engine/Documentation/Source/Programming/Development/BuildGraph/BuildGraphTypes/BuildGraphTypes.INT.udn
@@ -0,0 +1,61 @@
+Availability: Docs
+Title: BuildGraph Types
+Crumbs: %ROOT%, Programming, Programming/Development, Programming/Development/BuildGraphs
+Description:Learn about valid data types for BuildGraph attributes.
+version: 4.13
+parent:Programming/Development/BuildGraphs
+type:Reference
+related: Programming/Development/BuildGraphs/BuildGraphScriptAnatomy
+tags:BuildGraph
+
+[VAR:TopicCompact]
+[OBJECT:TopicCompact]
+	[PARAM:image]
+		![%Programming/Development/BuildGraphs/BuildGraphTypes%](types_topic.png)
+	[/PARAM]
+	[PARAM:icon]
+		![](%ROOT%/reference_icon.png)(convert:false)
+	[/PARAM]
+	[PARAM:title]
+		%Programming/Development/BuildGraphs/BuildGraphTypes:title%
+	[/PARAM]
+	[PARAM:description]
+		%Programming/Development/BuildGraphs/BuildGraphTypes:description%
+	[/PARAM]
+	[PARAM:path]
+		[RELATIVE:Programming/Development/BuildGraphs/BuildGraphTypes]
+	[/PARAM]
+[/OBJECT]
+[/VAR]
+
+[EXCERPT:buildgraphtypes]
+
+## Types
+
+The following table contains valid data types that can be contained by BuildGraph attributes:
+
+[INCLUDE:#TypesTable]
+
+[/EXCERPT:buildgraphtypes]
+
+<!--
+[EXCERPT:TypesTable]
+| **Type**                | **Description**                                                                                                                                                                                                                                                                            |
+| ----------------------- | ------------------------------------------------------------------------------------------------------------------------------------------------------------------------------------------------------------------------------------------------------------------------------------------ |
+| **String**              | An arbitrary string.                                                                                                                                                                                                                                                                       |
+| **String List**         | A list of arbitrary strings separated by semi-colons.                                                                                                                                                                                                                                      |
+| **Boolean**             | The constant `true` or `false`.                                                                                                                                                                                                                                                            |
+| **Integer**             | An integer constant.                                                                                                                                                                                                                                                                       |
+| **Regex**               | A regular expression, using [C#](http://www.iso.org/iso/home/store/catalogue_ics/catalogue_detail_ics.htm?csnumber=42926) syntax.                                                                                                                                                          |
+| **Name**                | A named entity. Any printable character, except `^  <  >  :  "  /  \  |  ?  *  ;`. Single spaces are allowed (except at the start and end of a name).                                                                                                                                      |
+| **Name List**           | A list of identifiers separated by semi-colons.                                                                                                                                                                                                                                            |
+| **Tag**                 | A label given to a list of files, beginning with the # character (i.e., `#My Files`).                                                                                                                                                                                                      |
+| **Tag List**            | A list of tags separated by semi-colons. (i.e., `#My Files;#Other Files`).                                                                                                                                                                                                                 |
+| **Target**              | A node name, aggregate name, agent name, or tag name. Indicates a sequence of nodes that need to be executed. Note that this overlaps with the meaning of `Target` as it applies to [UnrealBuildTool](Programming/UnrealBuildSystem) (which refers to a single program being built).       |
+| **Target list**         | A list of target names separated by semi-colons.                                                                                                                                                                                                                                           |
+| **File Spec**           | A set of file and tag names, separated by semi-colons. Wildcards such as `"..."`, `"*"`, and `"?"` are permitted when referencing files (i.e., `Engine/.../*.bat`). Unless otherwise specified, relative paths are resolved relative to the working root directory.                        |
+| **File Name**           | A path to a file. Unless otherwise specified, relative paths are resolved to the working root directory.                                                                                                                                                                                   |
+| **Directory Name**      | A path to a directory. Unless otherwise specified, relative paths are resolved relative to the working root directory.                                                                                                                                                                     |
+| **Condition**           | A [conditional expression](Programming/Development/BuildGraphs/BuildGraphConditions).                                                                                                                                                                                                      |
+[/EXCERPT:TypesTable]
+-->
\ No newline at end of file
diff --git a/Engine/Documentation/Source/Programming/Development/BuildGraph/BuildGraphUsage/BuildGraphUsage.INT.udn b/Engine/Documentation/Source/Programming/Development/BuildGraph/BuildGraphUsage/BuildGraphUsage.INT.udn
new file mode 100644
--- /dev/null
+++ b/Engine/Documentation/Source/Programming/Development/BuildGraph/BuildGraphUsage/BuildGraphUsage.INT.udn
@@ -0,0 +1,84 @@
+Availability: Docs
+Title: BuildGraph Usage
+Crumbs: %ROOT%, Programming, Programming/Development, Programming/Development/BuildGraphs
+Description:This page shows you how to use BuildGraph Scripts.
+version: 4.13
+parent:Programming/Development/BuildGraphs
+type:Overview
+related: Programming/Development/BuildGraphs/BuildGraphScriptAnatomy
+tags:BuildGraph
+
+[VAR:TopicCompact]
+[OBJECT:TopicCompact]
+	[PARAM:image]
+		![%Programming/Development/BuildGraphs/BuildGraphUsage%](usage_topic.png)
+	[/PARAM]
+	[PARAM:icon]
+		![](%ROOT%/reference_icon.png)(convert:false)
+	[/PARAM]
+	[PARAM:title]
+		%Programming/Development/BuildGraphs/BuildGraphUsage:title%
+	[/PARAM]
+	[PARAM:description]
+		%Programming/Development/BuildGraphs/BuildGraphUsage:description%
+	[/PARAM]
+	[PARAM:path]
+		[RELATIVE:Programming/Development/BuildGraphs/BuildGraphUsage]
+	[/PARAM]
+[/OBJECT]
+[/VAR]
+
+[TOC(start:1 end:3)]
+
+BuildGraph is implemented as part of AutomationTool, which you can run directly from the Unreal Engine 4 solution on Windows, 
+or via the `Engine/Build/BatchFiles/RunUAT.bat` batch file from the command line. On Mac or Linux, it can be run with the 
+`Engine/Build/BatchFiles/RunUAT.sh` shell script. If you'd like to get an idea of how to use BuildGraph from the command line, 
+we've written sample usage commands for your reference.
+
+## Showing Options & Nodes
+
+To show all of the nodes and options available in a BuildGraph script:
+
+	BuildGraph -Script=Engine/Build/Graph/Examples/AllExamples -ListOnly
+
+## Setting an Option
+
+To set an option defined in the script:
+
+	BuildGraph -Script=Engine/Build/Graph/Examples/Properties.xml -ListOnly -Set:CommandLineProperty=ABCDEF
+
+## Seeing Node Execution
+
+To see what nodes will be executed for a target:
+
+	BuildGraph -Script=Engine\Build\Graph\Examples\AllExamples.xml -Target="All Examples" -ListOnly
+
+## Cleaning Previous Builds
+
+To clean any previous builds and execute a specific target:
+
+	BuildGraph -Script=Engine\Build\Graph\Examples\AllExamples.xml -Target="Property Examples" -Clean
+
+## Writing out an XML Schema
+
+To write out an XML schema containing all registered tasks:
+
+	BuildGraph -Schema=Engine\Build\Graph\Schema.xsd
+
+## Generating a Preprocessed Script
+
+To output a preprocessed version of your script (useful for debugging):
+
+	BuildGraph -Script=Engine\Build\Graph\Examples\AllExamples.xml -Target="All Examples" -ListOnly -Preprocess=D:\Preprocessed.xml
+
+## Exporting a Preprocessed Graph to JSON
+
+To export a preprocessed graph to JSON (for dynamic configuration on build systems):
+
+	BuildGraph -Script=Engine\Build\Graph\Examples\AllExamples.xml -Target="All Examples" -Export=D:\ExportedGraph.json
+
+## Executing a Single Node
+
+To execute a single node on a build system (while transferring intermediates to and from shared storage):
+
+	BuildGraph -Script=Engine\Build\Graph\Examples\AllExamples.xml -Target="All Examples" -Clean -SingleNode=”Manipulate Some Files” -SharedStorage=P:\Network-Share\Unique-Job-12345\
\ No newline at end of file
diff --git a/Engine/Documentation/Source/Programming/Development/BuildGraphs/BuildGraph.INT.udn b/Engine/Documentation/Source/Programming/Development/BuildGraphs/BuildGraph.INT.udn
deleted file mode 100644
diff --git a/Engine/Documentation/Source/Programming/Development/BuildGraphs/BuildGraphConditions/BuildGraphConditions.INT.udn b/Engine/Documentation/Source/Programming/Development/BuildGraphs/BuildGraphConditions/BuildGraphConditions.INT.udn
deleted file mode 100644
diff --git a/Engine/Documentation/Source/Programming/Development/BuildGraphs/BuildGraphElements/BuildGraphElements.INT.udn b/Engine/Documentation/Source/Programming/Development/BuildGraphs/BuildGraphElements/BuildGraphElements.INT.udn
deleted file mode 100644
diff --git a/Engine/Documentation/Source/Programming/Development/BuildGraphs/BuildGraphIntroduction/BuildGraphIntroduction.INT.udn b/Engine/Documentation/Source/Programming/Development/BuildGraphs/BuildGraphIntroduction/BuildGraphIntroduction.INT.udn
deleted file mode 100644
diff --git a/Engine/Documentation/Source/Programming/Development/BuildGraphs/BuildGraphScriptAnatomy/BuildGraphScriptAnatomy.INT.udn b/Engine/Documentation/Source/Programming/Development/BuildGraphs/BuildGraphScriptAnatomy/BuildGraphScriptAnatomy.INT.udn
deleted file mode 100644
diff --git a/Engine/Documentation/Source/Programming/Development/BuildGraphs/BuildGraphTasks/BuildGraphPredefinedTasks/BuildGraphPredefinedTasks.INT.udn b/Engine/Documentation/Source/Programming/Development/BuildGraphs/BuildGraphTasks/BuildGraphPredefinedTasks/BuildGraphPredefinedTasks.INT.udn
deleted file mode 100644
diff --git a/Engine/Documentation/Source/Programming/Development/BuildGraphs/BuildGraphTasks/BuildGraphTasks.INT.udn b/Engine/Documentation/Source/Programming/Development/BuildGraphs/BuildGraphTasks/BuildGraphTasks.INT.udn
deleted file mode 100644
diff --git a/Engine/Documentation/Source/Programming/Development/BuildGraphs/BuildGraphTypes/BuildGraphTypes.INT.udn b/Engine/Documentation/Source/Programming/Development/BuildGraphs/BuildGraphTypes/BuildGraphTypes.INT.udn
deleted file mode 100644
diff --git a/Engine/Documentation/Source/Programming/Development/BuildGraphs/BuildGraphUsage/BuildGraphUsage.INT.udn b/Engine/Documentation/Source/Programming/Development/BuildGraphs/BuildGraphUsage/BuildGraphUsage.INT.udn
deleted file mode 100644
