#UE4 Docs: First check in on this new document.

diff --git a/Engine/Documentation/Source/Platforms/Mobile/MobilePackagingWizard/MobilePackagingWizard.INT.udn b/Engine/Documentation/Source/Platforms/Mobile/MobilePackagingWizard/MobilePackagingWizard.INT.udn
new file mode 100644
--- /dev/null
+++ b/Engine/Documentation/Source/Platforms/Mobile/MobilePackagingWizard/MobilePackagingWizard.INT.udn
@@ -0,0 +1,114 @@
+Title: Mobile Packaging Wizard 
+Description:Using the Mobile Packaging Wizard to reduce the install size of your UE4 mobile project
+Crumbs:%ROOT%, Platforms, Platforms/Mobile, Platforms/Mobile/MobilePackagingWizard/
+Availability:Docs
+version:4.12
+Parent:Platforms/Mobile
+tags:Mobile
+type:overview
+order:8
+SkillLevel: Intermediate
+Related:Engine/Deployment
+Related:Engine/Deployment/Cooking
+Related:Engine/Deployment/Patching
+Related:Engine/Deployment/UnrealFrontend
+
+[TOC(start:2 end:3)]
+
+
+The Unreal Engine 4 (UE4) **Mobile Packaging Wizard** will help to reduce the size of your mobile project by packaging only the content that is needed to upload your project to the App store.  The rest of the content will be download from the cloud when the user launches the game on their device for the first time. In the following document, we will take a look at how you set up your UE4 mobile project to use the Mobile Packaging Wizard.
+
+
+## Mobile Packaging Wizard Tool Location
+
+To access the Mobile Packaging Wizard you will first have to open the **Project Launcher** by going to the **Main Toolbar** and then clicking on the  **Window** option and selecting the **Project Launcher** option from the list.
+
+![](MPW_Location.png)
+
+
+## Setting Up a Mobile Packaging Wizard Profile
+
+Before the Mobile Packaging Wizard can be used a new profile that is specific to your projects needs has to be created. In the following section, we will take a look at what needs to be done in order to create a custom launch profile for your UE4 mobile project.
+
+1. Click on the small white triangle next to Custom Launch Profiles and then from the list select the **Minimal Android APK + DLC…** option.
+
+	![](MPW_Select_Android_Option.png)
+
+	[region:note]
+	While we are using Android for this example the workflow will still be the same when using the Mobile Packaging Wizard for a game that is for IOS.
+	[/REGION]
+
+1. When the **Minimal Android APK + DLC** packaging window is displayed,click on the **Browse** button and then select a location for the packaged content to be placed. 
+
+	![](MPW_Press_Browse.png)
+
+1. Once you have selected a location for the packaged content to be placed, press the **Next** button to continue the packaging process.
+
+	![](MPW_Press_Next.png)
+
+1. In the **Distributable APK** menu set the **Build Configuration** to **Development** and then select the map you want to use, pressing the **Next** button when finished to continue. 
+
+	![](MPW_Dis_APK.png)
+
+	[region:note]
+	Only the contents referenced by the map(s) you choose will be packaged into the application for upload to the Play Store. Typically you should choose only a single map which contains a Level Blueprint to kick off the BuildPatchServices code, and some UMG user interface to show the download progress and any error conditions.
+	[/region]
+
+1. In the **Downloadable Content** menu, select the Android texture formats your application will support and the maps that will need to be downloaded. When that has been completed press the **Create Profile** button to create a new Custom Launch Profile.
+
+	![](MPW_Downloadable_Content.png)
+
+## Building the Application 
+
+Now that all of the options for the Mobile Packing Wizard have been setup and saved into a Custom Launch Profile, it is now time to run that profile and package your project. In the following section, we will take a look at what needs to be done in order for the Mobile Packing Wizard to successfully package your UE4 project.
+
+1. If the Project Launcher is not already open, open it up and then look at the bottom of the Project Launcher in the **Custom Launch Profile** section. In this section you should have two profiles that start with your project name and are followed by either **Android APK** or **Android DLC**. 
+
+	![](MPW_Custom_Launch_Profiles.png)
+
+	[region:note]
+	When you first open up the Project Launcher the Custom Launch Profiles section, could have a few items in it by default. To remove these items Right - Click on the item you want to remove and then select the **Delete** option to remove it.
+	[/region]
+
+1. Locate the **APK** Profile option and then click on the **Launch this Profile** button to begin the packaging process. 
+
+	![](MPW_Launch_APK_Profile.png)
+
+	[region:warning]
+	You must package your projects files before you build the DLC. Trying to package the DLC first and then packing your project will not.
+	[/region]
+
+1. Once the Launch this Profile button has been pressed, the Project Launcher will begin the packaging process.  You can see the progress of the packaging process in the Project Launcher window.
+
+	![](MPW_Packaging_Progress.png)
+
+1. When the packaging has completed, click on the **Done** button to bring the Project Launcher back up.
+
+	![](MPW_Packaging_Completed.png)
+
+## Building the DLC
+
+Once your project has finished the cooking process for the initial download from the App store,  it is now time to cook the content that is to be downloaded from the cloud after your App has been installed to a mobile device. 
+
+1. From the Project Launcher go to the **Custom Launch Profiles** section and look for your projects **DLC** option. When the DLC option for your project has been located, click on the **Launch this Profile** button to begin the DLC packaging process. 
+
+	![](MPW_Build_DLC.png)
+
+1. After pressing the Launch this Profile button the Project Launcher will begin to build the DLC for your project. When the DLC has finished being built, click on the **Done** button to continue. 
+
+	![](MPW_DLC_Completed.png)
+
+## Packaging Results 
+
+Once everything has finished building, open up the folder that was specified in the 2nd step of the Setting up a mobile packaging Wizard Profile section and you should see the following two folders.
+
+![](MPW_App_DLC_Location.png)
+
+* **App**: This folder contains the version of the application should be uploaded to the App store.
+
+* **HTTPchunks**: This folder contains the content that is meant to be downloaded from the cloud.
+
+You can now upload the contents of the **App** folder to the App store and then place the HTTPchunks folder into the cloud service of your choice so users can download the rest of your projects files.
+
+
+
