Fixing broken links due to renaming path per BenM's request.

#UE4doc

diff --git a/Engine/Documentation/Source/Programming/Development/BuildGraph/BuildGraph.INT.udn b/Engine/Documentation/Source/Programming/Development/BuildGraph/BuildGraph.INT.udn
--- a/Engine/Documentation/Source/Programming/Development/BuildGraph/BuildGraph.INT.udn
+++ b/Engine/Documentation/Source/Programming/Development/BuildGraph/BuildGraph.INT.udn
@@ -1,5 +1,5 @@
 Availability: Docs
-Title: BuildGraph Scripting
+Title: BuildGraph
 Crumbs: %ROOT%, Programming, Programming/Development
 Description: Customize your Installed Builds by using our BuildGraph scripting system.
 Version: 4.13
@@ -7,26 +7,26 @@ type:landing
 parent:Programming/Development
 tags:BuildGraph
 
-[INCLUDE:Programming/Development/BuildGraphs/BuildGraphIntroduction]
+[INCLUDE:Programming/Development/BuildGraph/BuildGraphIntroduction]
 
 ## Writing BuildGraph Scripts
 
 [REGION:topics third]
-%Programming/Development/BuildGraphs/BuildGraphScriptAnatomy:topiccompact%
+%Programming/Development/BuildGraph/BuildGraphScriptAnatomy:topiccompact%
 [/REGION]
 
 ## Using BuildGraph
 
 [REGION:topics third]
-%Programming/Development/BuildGraphs/BuildGraphUsage:topiccompact%
+%Programming/Development/BuildGraph/BuildGraphUsage:topiccompact%
 %Programming/Development/InstalledBuildReference:topiccompact%
 [/REGION]
 
-## Reference
+## BuildGraph Script Reference
 
 [REGION:topics third]
-%Programming/Development/BuildGraphs/BuildGraphTypes:topiccompact%
-%Programming/Development/BuildGraphs/BuildGraphElements:topiccompact%
-%Programming/Development/BuildGraphs/BuildGraphConditions:topiccompact%
-%Programming/Development/BuildGraphs/BuildGraphTasks:topiccompact%
+%Programming/Development/BuildGraph/BuildGraphTypes:topiccompact%
+%Programming/Development/BuildGraph/BuildGraphElements:topiccompact%
+%Programming/Development/BuildGraph/BuildGraphConditions:topiccompact%
+%Programming/Development/BuildGraph/BuildGraphTasks:topiccompact%
 [/REGION]
\ No newline at end of file
diff --git a/Engine/Documentation/Source/Programming/Development/BuildGraph/BuildGraphConditions/BuildGraphConditions.INT.udn b/Engine/Documentation/Source/Programming/Development/BuildGraph/BuildGraphConditions/BuildGraphConditions.INT.udn
--- a/Engine/Documentation/Source/Programming/Development/BuildGraph/BuildGraphConditions/BuildGraphConditions.INT.udn
+++ b/Engine/Documentation/Source/Programming/Development/BuildGraph/BuildGraphConditions/BuildGraphConditions.INT.udn
@@ -1,29 +1,29 @@
 Availability: Docs
-Title: BuildGraph Conditions
-Crumbs: %ROOT%, Programming, Programming/Development, Programming/Development/BuildGraphs
+Title: BuildGraph Script Conditions
+Crumbs: %ROOT%, Programming, Programming/Development, Programming/Development/BuildGraph
 Description:Learn the syntax needed to write BuildGraph script conditions.
 version: 4.13
-parent:Programming/Development/BuildGraphs
+parent:Programming/Development/BuildGraph
 type:Reference
-related: Programming/Development/BuildGraphs/BuildGraphScriptAnatomy
+related: Programming/Development/BuildGraph/BuildGraphScriptAnatomy
 tags:BuildGraph
 
 [VAR:TopicCompact]
 [OBJECT:TopicCompact]
 	[PARAM:image]
-		![%Programming/Development/BuildGraphs/BuildGraphConditions%](conditional_topic.png)
+		![%Programming/Development/BuildGraph/BuildGraphConditions%](conditional_topic.png)
 	[/PARAM]
 	[PARAM:icon]
 		![](%ROOT%/reference_icon.png)(convert:false)
 	[/PARAM]
 	[PARAM:title]
-		%Programming/Development/BuildGraphs/BuildGraphConditions:title%
+		%Programming/Development/BuildGraph/BuildGraphConditions:title%
 	[/PARAM]
 	[PARAM:description]
-		%Programming/Development/BuildGraphs/BuildGraphConditions:description%
+		%Programming/Development/BuildGraph/BuildGraphConditions:description%
 	[/PARAM]
 	[PARAM:path]
-		[RELATIVE:Programming/Development/BuildGraphs/BuildGraphConditions]
+		[RELATIVE:Programming/Development/BuildGraph/BuildGraphConditions]
 	[/PARAM]
 [/OBJECT]
 [/VAR]
diff --git a/Engine/Documentation/Source/Programming/Development/BuildGraph/BuildGraphElements/BuildGraphElements.INT.udn b/Engine/Documentation/Source/Programming/Development/BuildGraph/BuildGraphElements/BuildGraphElements.INT.udn
--- a/Engine/Documentation/Source/Programming/Development/BuildGraph/BuildGraphElements/BuildGraphElements.INT.udn
+++ b/Engine/Documentation/Source/Programming/Development/BuildGraph/BuildGraphElements/BuildGraphElements.INT.udn
@@ -1,29 +1,29 @@
 Availability: Docs
-Title: BuildGraph Elements
-Crumbs: %ROOT%, Programming, Programming/Development, Programming/Development/BuildGraphs
+Title: BuildGraph Script Elements
+Crumbs: %ROOT%, Programming, Programming/Development, Programming/Development/BuildGraph
 Description:Learn about the elements of a BuildGraph.
 version: 4.13
-parent:Programming/Development/BuildGraphs
+parent:Programming/Development/BuildGraph
 type:Reference
-related: Programming/Development/BuildGraphs/BuildGraphScriptAnatomy
+related: Programming/Development/BuildGraph/BuildGraphScriptAnatomy
 tags:BuildGraph
 
 [VAR:TopicCompact]
 [OBJECT:TopicCompact]
 	[PARAM:image]
-		![%Programming/Development/BuildGraphs/BuildGraphElements%](elements_topic.png)
+		![%Programming/Development/BuildGraph/BuildGraphElements%](elements_topic.png)
 	[/PARAM]
 	[PARAM:icon]
 		![](%ROOT%/reference_icon.png)(convert:false)
 	[/PARAM]
 	[PARAM:title]
-		%Programming/Development/BuildGraphs/BuildGraphElements:title%
+		%Programming/Development/BuildGraph/BuildGraphElements:title%
 	[/PARAM]
 	[PARAM:description]
-		%Programming/Development/BuildGraphs/BuildGraphElements:description%
+		%Programming/Development/BuildGraph/BuildGraphElements:description%
 	[/PARAM]
 	[PARAM:path]
-		[RELATIVE:Programming/Development/BuildGraphs/BuildGraphElements]
+		[RELATIVE:Programming/Development/BuildGraph/BuildGraphElements]
 	[/PARAM]
 [/OBJECT]
 [/VAR]
diff --git a/Engine/Documentation/Source/Programming/Development/BuildGraph/BuildGraphIntroduction/BuildGraphIntroduction.INT.udn b/Engine/Documentation/Source/Programming/Development/BuildGraph/BuildGraphIntroduction/BuildGraphIntroduction.INT.udn
--- a/Engine/Documentation/Source/Programming/Development/BuildGraph/BuildGraphIntroduction/BuildGraphIntroduction.INT.udn
+++ b/Engine/Documentation/Source/Programming/Development/BuildGraph/BuildGraphIntroduction/BuildGraphIntroduction.INT.udn
@@ -1,14 +1,13 @@
 Availability: NoPublish
-Title: Introduction to BuildGraphs
-Crumbs: %ROOT%, Programming, Programming/Development, Programming/Development/BuildGraphs
+Title: Introduction to BuildGraph
+Crumbs: %ROOT%, Programming, Programming/Development, Programming/Development/BuildGraph
 Description:This page introduces you to our BuildGraph scripting system.
 version: 4.13
-parent:Programming/Development/BuildGraphs
+parent:Programming/Development/BuildGraph
 type:Reference
 
-**BuildGraph** is a scripting language for build automation with building blocks common to Unreal Engine 4 (UE4) projects. 
-BuildGraph integrates with UnrealBuildTool, AutomationTool, and the editor. Which means that BuildGraph can be 
-extended and customized for your projects.
+**BuildGraph** is a system that runs BuildGraph scripts for build automation (with building blocks common to Unreal Engine 4 (UE4) projects). 
+BuildGraph integrates with UnrealBuildTool, AutomationTool, and the editor. Which means that BuildGraph can be extended and customized for your projects.
 
 BuildGraph scripts are written with [XML](https://www.w3.org/TR/NOTE-dcd]), and specify a graph of user-defined nodes with 
 dependencies between them. Each node consists of tasks executed in sequence to produce some sort of output (i.e., compile 
diff --git a/Engine/Documentation/Source/Programming/Development/BuildGraph/BuildGraphScriptAnatomy/BuildGraphScriptAnatomy.INT.udn b/Engine/Documentation/Source/Programming/Development/BuildGraph/BuildGraphScriptAnatomy/BuildGraphScriptAnatomy.INT.udn
--- a/Engine/Documentation/Source/Programming/Development/BuildGraph/BuildGraphScriptAnatomy/BuildGraphScriptAnatomy.INT.udn
+++ b/Engine/Documentation/Source/Programming/Development/BuildGraph/BuildGraphScriptAnatomy/BuildGraphScriptAnatomy.INT.udn
@@ -1,28 +1,28 @@
 Availability: Docs
 Title: BuildGraph Script Anatomy
-Crumbs: %ROOT%, Programming, Programming/Development, Programming/Development/BuildGraphs
+Crumbs: %ROOT%, Programming, Programming/Development, Programming/Development/BuildGraph
 Description:This page introduces you to our BuildGraph scripting system.
 version: 4.13
-parent:Programming/Development/BuildGraphs
+parent:Programming/Development/BuildGraph
 type:Overview
 tags:BuildGraph
 
 [VAR:TopicCompact]
 [OBJECT:TopicCompact]
 	[PARAM:image]
-		![%Programming/Development/BuildGraphs/BuildGraphScriptAnatomy%](anatomy_topic.png)
+		![%Programming/Development/BuildGraph/BuildGraphScriptAnatomy%](anatomy_topic.png)
 	[/PARAM]
 	[PARAM:icon]
 		![](%ROOT%/reference_icon.png)(convert:false)
 	[/PARAM]
 	[PARAM:title]
-		%Programming/Development/BuildGraphs/BuildGraphScriptAnatomy:title%
+		%Programming/Development/BuildGraph/BuildGraphScriptAnatomy:title%
 	[/PARAM]
 	[PARAM:description]
-		%Programming/Development/BuildGraphs/BuildGraphScriptAnatomy:description%
+		%Programming/Development/BuildGraph/BuildGraphScriptAnatomy:description%
 	[/PARAM]
 	[PARAM:path]
-		[RELATIVE:Programming/Development/BuildGraphs/BuildGraphScriptAnatomy]
+		[RELATIVE:Programming/Development/BuildGraph/BuildGraphScriptAnatomy]
 	[/PARAM]
 [/OBJECT]
 [/VAR]
@@ -63,10 +63,10 @@ the allocation of machines (and coordination between them) is assumed to be mana
 
 The syntax of elements used to define BuildGraph constructs is listed in the following sections.
 
-[INCLUDE:Programming/Development/BuildGraphs/BuildGraphTypes#buildgraphtypes]
+[INCLUDE:Programming/Development/BuildGraph/BuildGraphTypes#buildgraphtypes]
 
-[INCLUDE:Programming/Development/BuildGraphs/BuildGraphElements#buildgraphelements]
+[INCLUDE:Programming/Development/BuildGraph/BuildGraphElements#buildgraphelements]
 
-[INCLUDE:Programming/Development/BuildGraphs/BuildGraphConditions#buildgraphconditions]
+[INCLUDE:Programming/Development/BuildGraph/BuildGraphConditions#buildgraphconditions]
 
-[INCLUDE:Programming/Development/BuildGraphs/BuildGraphTasks#buildgraphtasks]
\ No newline at end of file
+[INCLUDE:Programming/Development/BuildGraph/BuildGraphTasks#buildgraphtasks]
\ No newline at end of file
diff --git a/Engine/Documentation/Source/Programming/Development/BuildGraph/BuildGraphTasks/BuildGraphPredefinedTasks/BuildGraphPredefinedTasks.INT.udn b/Engine/Documentation/Source/Programming/Development/BuildGraph/BuildGraphTasks/BuildGraphPredefinedTasks/BuildGraphPredefinedTasks.INT.udn
--- a/Engine/Documentation/Source/Programming/Development/BuildGraph/BuildGraphTasks/BuildGraphPredefinedTasks/BuildGraphPredefinedTasks.INT.udn
+++ b/Engine/Documentation/Source/Programming/Development/BuildGraph/BuildGraphTasks/BuildGraphPredefinedTasks/BuildGraphPredefinedTasks.INT.udn
@@ -1,7 +1,7 @@
 Availability: NoPublish
 Title: BuildGraph Predefined Tasks
-Crumbs: %ROOT%, Programming, Programming/Development, Programming/Development/BuildGraphs, Programming/Development/BuildGraphs/BuildGraphTasks
+Crumbs: %ROOT%, Programming, Programming/Development, Programming/Development/BuildGraph, Programming/Development/BuildGraph/BuildGraphTasks
 Description: This is a procedurally generated markdown page.
 version: 4.13
-parent:Programming/Development/BuildGraphs/BuildGraphTasks
+parent:Programming/Development/BuildGraph/BuildGraphTasks
 
diff --git a/Engine/Documentation/Source/Programming/Development/BuildGraph/BuildGraphTasks/BuildGraphTasks.INT.udn b/Engine/Documentation/Source/Programming/Development/BuildGraph/BuildGraphTasks/BuildGraphTasks.INT.udn
--- a/Engine/Documentation/Source/Programming/Development/BuildGraph/BuildGraphTasks/BuildGraphTasks.INT.udn
+++ b/Engine/Documentation/Source/Programming/Development/BuildGraph/BuildGraphTasks/BuildGraphTasks.INT.udn
@@ -1,29 +1,29 @@
 Availability: Docs
-Title: BuildGraph Tasks
-Crumbs: %ROOT%, Programming, Programming/Development, Programming/Development/BuildGraphs
+Title: BuildGraph Script Tasks
+Crumbs: %ROOT%, Programming, Programming/Development, Programming/Development/BuildGraph
 Description: Learn how BuildGraph can help you create custom tasks.
 version: 4.13
-parent:Programming/Development/BuildGraphs
+parent:Programming/Development/BuildGraph
 type:Reference
-related: Programming/Development/BuildGraphs/BuildGraphScriptAnatomy
+related: Programming/Development/BuildGraph/BuildGraphScriptAnatomy
 tags:BuildGraph
 
 [VAR:TopicCompact]
 [OBJECT:TopicCompact]
 	[PARAM:image]
-		![%Programming/Development/BuildGraphs/BuildGraphTasks%](tasks_topic.png)
+		![%Programming/Development/BuildGraph/BuildGraphTasks%](tasks_topic.png)
 	[/PARAM]
 	[PARAM:icon]
 		![](%ROOT%/reference_icon.png)(convert:false)
 	[/PARAM]
 	[PARAM:title]
-		%Programming/Development/BuildGraphs/BuildGraphTasks:title%
+		%Programming/Development/BuildGraph/BuildGraphTasks:title%
 	[/PARAM]
 	[PARAM:description]
-		%Programming/Development/BuildGraphs/BuildGraphTasks:description%
+		%Programming/Development/BuildGraph/BuildGraphTasks:description%
 	[/PARAM]
 	[PARAM:path]
-		[RELATIVE:Programming/Development/BuildGraphs/BuildGraphTasks]
+		[RELATIVE:Programming/Development/BuildGraph/BuildGraphTasks]
 	[/PARAM]
 [/OBJECT]
 [/VAR]
@@ -53,6 +53,6 @@ If you'd like to have a starting point for creating tasks, we're providing a var
 If you're looking for a simple predefined task to start from, we recommend using `LogTask` as a good starting point. If you want to keep up with 
 our latest improvements to BuildGraph; bookmark this page, because the following list of predefined tasks will be updated periodically. 
 
-[INCLUDE:Programming/Development/BuildGraphs/BuildGraphTasks/BuildGraphPredefinedTasks]
+[INCLUDE:Programming/Development/BuildGraph/BuildGraphTasks/BuildGraphPredefinedTasks]
 
 [/EXCERPT:buildgraphtasks]
\ No newline at end of file
diff --git a/Engine/Documentation/Source/Programming/Development/BuildGraph/BuildGraphTypes/BuildGraphTypes.INT.udn b/Engine/Documentation/Source/Programming/Development/BuildGraph/BuildGraphTypes/BuildGraphTypes.INT.udn
--- a/Engine/Documentation/Source/Programming/Development/BuildGraph/BuildGraphTypes/BuildGraphTypes.INT.udn
+++ b/Engine/Documentation/Source/Programming/Development/BuildGraph/BuildGraphTypes/BuildGraphTypes.INT.udn
@@ -1,29 +1,29 @@
 Availability: Docs
-Title: BuildGraph Types
-Crumbs: %ROOT%, Programming, Programming/Development, Programming/Development/BuildGraphs
+Title: BuildGraph Script Types
+Crumbs: %ROOT%, Programming, Programming/Development, Programming/Development/BuildGraph
 Description:Learn about valid data types for BuildGraph attributes.
 version: 4.13
-parent:Programming/Development/BuildGraphs
+parent:Programming/Development/BuildGraph
 type:Reference
-related: Programming/Development/BuildGraphs/BuildGraphScriptAnatomy
+related: Programming/Development/BuildGraph/BuildGraphScriptAnatomy
 tags:BuildGraph
 
 [VAR:TopicCompact]
 [OBJECT:TopicCompact]
 	[PARAM:image]
-		![%Programming/Development/BuildGraphs/BuildGraphTypes%](types_topic.png)
+		![%Programming/Development/BuildGraph/BuildGraphTypes%](types_topic.png)
 	[/PARAM]
 	[PARAM:icon]
 		![](%ROOT%/reference_icon.png)(convert:false)
 	[/PARAM]
 	[PARAM:title]
-		%Programming/Development/BuildGraphs/BuildGraphTypes:title%
+		%Programming/Development/BuildGraph/BuildGraphTypes:title%
 	[/PARAM]
 	[PARAM:description]
-		%Programming/Development/BuildGraphs/BuildGraphTypes:description%
+		%Programming/Development/BuildGraph/BuildGraphTypes:description%
 	[/PARAM]
 	[PARAM:path]
-		[RELATIVE:Programming/Development/BuildGraphs/BuildGraphTypes]
+		[RELATIVE:Programming/Development/BuildGraph/BuildGraphTypes]
 	[/PARAM]
 [/OBJECT]
 [/VAR]
@@ -56,6 +56,6 @@ The following table contains valid data types that can be contained by BuildGrap
 | **File Spec**           | A set of file and tag names, separated by semi-colons. Wildcards such as `"..."`, `"*"`, and `"?"` are permitted when referencing files (i.e., `Engine/.../*.bat`). Unless otherwise specified, relative paths are resolved relative to the working root directory.                        |
 | **File Name**           | A path to a file. Unless otherwise specified, relative paths are resolved to the working root directory.                                                                                                                                                                                   |
 | **Directory Name**      | A path to a directory. Unless otherwise specified, relative paths are resolved relative to the working root directory.                                                                                                                                                                     |
-| **Condition**           | A [conditional expression](Programming/Development/BuildGraphs/BuildGraphConditions).                                                                                                                                                                                                      |
+| **Condition**           | A [conditional expression](Programming/Development/BuildGraph/BuildGraphConditions).                                                                                                                                                                                                      |
 [/EXCERPT:TypesTable]
 -->
\ No newline at end of file
diff --git a/Engine/Documentation/Source/Programming/Development/BuildGraph/BuildGraphUsage/BuildGraphUsage.INT.udn b/Engine/Documentation/Source/Programming/Development/BuildGraph/BuildGraphUsage/BuildGraphUsage.INT.udn
--- a/Engine/Documentation/Source/Programming/Development/BuildGraph/BuildGraphUsage/BuildGraphUsage.INT.udn
+++ b/Engine/Documentation/Source/Programming/Development/BuildGraph/BuildGraphUsage/BuildGraphUsage.INT.udn
@@ -1,29 +1,29 @@
 Availability: Docs
 Title: BuildGraph Usage
-Crumbs: %ROOT%, Programming, Programming/Development, Programming/Development/BuildGraphs
+Crumbs: %ROOT%, Programming, Programming/Development, Programming/Development/BuildGraph
 Description:This page shows you how to use BuildGraph Scripts.
 version: 4.13
-parent:Programming/Development/BuildGraphs
+parent:Programming/Development/BuildGraph
 type:Overview
-related: Programming/Development/BuildGraphs/BuildGraphScriptAnatomy
+related: Programming/Development/BuildGraph/BuildGraphScriptAnatomy
 tags:BuildGraph
 
 [VAR:TopicCompact]
 [OBJECT:TopicCompact]
 	[PARAM:image]
-		![%Programming/Development/BuildGraphs/BuildGraphUsage%](usage_topic.png)
+		![%Programming/Development/BuildGraph/BuildGraphUsage%](usage_topic.png)
 	[/PARAM]
 	[PARAM:icon]
 		![](%ROOT%/reference_icon.png)(convert:false)
 	[/PARAM]
 	[PARAM:title]
-		%Programming/Development/BuildGraphs/BuildGraphUsage:title%
+		%Programming/Development/BuildGraph/BuildGraphUsage:title%
 	[/PARAM]
 	[PARAM:description]
-		%Programming/Development/BuildGraphs/BuildGraphUsage:description%
+		%Programming/Development/BuildGraph/BuildGraphUsage:description%
 	[/PARAM]
 	[PARAM:path]
-		[RELATIVE:Programming/Development/BuildGraphs/BuildGraphUsage]
+		[RELATIVE:Programming/Development/BuildGraph/BuildGraphUsage]
 	[/PARAM]
 [/OBJECT]
 [/VAR]
diff --git a/Engine/Documentation/Source/Programming/Development/InstalledBuildReference/InstalledBuildReference.INT.udn b/Engine/Documentation/Source/Programming/Development/InstalledBuildReference/InstalledBuildReference.INT.udn
--- a/Engine/Documentation/Source/Programming/Development/InstalledBuildReference/InstalledBuildReference.INT.udn
+++ b/Engine/Documentation/Source/Programming/Development/InstalledBuildReference/InstalledBuildReference.INT.udn
@@ -7,6 +7,7 @@ parent:Programming/Development
 type:Reference
 related:Programming/Development/BuildConfigurations
 related:Programming/UnrealBuildSystem
+related:Programming/Development/BuildGraph
 
 [TOC (start:1 end:4)]
 
