Fixing broken links caused by round 2 of renaming path per BenM's specification.

#UE4doc

diff --git a/Engine/Documentation/Source/Programming/Development/BuildGraph/BuildGraph.INT.udn b/Engine/Documentation/Source/Programming/Development/BuildGraph/BuildGraph.INT.udn
--- a/Engine/Documentation/Source/Programming/Development/BuildGraph/BuildGraph.INT.udn
+++ b/Engine/Documentation/Source/Programming/Development/BuildGraph/BuildGraph.INT.udn
@@ -22,11 +22,11 @@ tags:BuildGraph
 %Programming/Development/InstalledBuildReference:topiccompact%
 [/REGION]
 
-## BuildGraph Script Reference
+## BuildGraph Scripting Reference
 
 [REGION:topics third]
-%Programming/Development/BuildGraph/BuildGraphTypes:topiccompact%
-%Programming/Development/BuildGraph/BuildGraphElements:topiccompact%
-%Programming/Development/BuildGraph/BuildGraphConditions:topiccompact%
-%Programming/Development/BuildGraph/BuildGraphTasks:topiccompact%
+%Programming/Development/BuildGraph/BuildGraphScriptTypes:topiccompact%
+%Programming/Development/BuildGraph/BuildGraphScriptElements:topiccompact%
+%Programming/Development/BuildGraph/BuildGraphScriptConditions:topiccompact%
+%Programming/Development/BuildGraph/BuildGraphScriptTasks:topiccompact%
 [/REGION]
\ No newline at end of file
diff --git a/Engine/Documentation/Source/Programming/Development/BuildGraph/BuildGraphScriptAnatomy/BuildGraphScriptAnatomy.INT.udn b/Engine/Documentation/Source/Programming/Development/BuildGraph/BuildGraphScriptAnatomy/BuildGraphScriptAnatomy.INT.udn
--- a/Engine/Documentation/Source/Programming/Development/BuildGraph/BuildGraphScriptAnatomy/BuildGraphScriptAnatomy.INT.udn
+++ b/Engine/Documentation/Source/Programming/Development/BuildGraph/BuildGraphScriptAnatomy/BuildGraphScriptAnatomy.INT.udn
@@ -63,10 +63,10 @@ the allocation of machines (and coordination between them) is assumed to be mana
 
 The syntax of elements used to define BuildGraph constructs is listed in the following sections.
 
-[INCLUDE:Programming/Development/BuildGraph/BuildGraphTypes#buildgraphtypes]
+[INCLUDE:Programming/Development/BuildGraph/BuildGraphScriptTypes#buildgraphscripttypes]
 
-[INCLUDE:Programming/Development/BuildGraph/BuildGraphElements#buildgraphelements]
+[INCLUDE:Programming/Development/BuildGraph/BuildGraphScriptElements#buildgraphscriptelements]
 
-[INCLUDE:Programming/Development/BuildGraph/BuildGraphConditions#buildgraphconditions]
+[INCLUDE:Programming/Development/BuildGraph/BuildGraphScriptConditions#buildgraphscriptconditions]
 
-[INCLUDE:Programming/Development/BuildGraph/BuildGraphTasks#buildgraphtasks]
\ No newline at end of file
+[INCLUDE:Programming/Development/BuildGraph/BuildGraphScriptTasks#buildgraphscripttasks]
\ No newline at end of file
diff --git a/Engine/Documentation/Source/Programming/Development/BuildGraph/BuildGraphScriptConditions/BuildGraphScriptConditions.INT.udn b/Engine/Documentation/Source/Programming/Development/BuildGraph/BuildGraphScriptConditions/BuildGraphScriptConditions.INT.udn
--- a/Engine/Documentation/Source/Programming/Development/BuildGraph/BuildGraphScriptConditions/BuildGraphScriptConditions.INT.udn
+++ b/Engine/Documentation/Source/Programming/Development/BuildGraph/BuildGraphScriptConditions/BuildGraphScriptConditions.INT.udn
@@ -11,26 +11,26 @@ tags:BuildGraph
 [VAR:TopicCompact]
 [OBJECT:TopicCompact]
 	[PARAM:image]
-		![%Programming/Development/BuildGraph/BuildGraphConditions%](conditional_topic.png)
+		![%Programming/Development/BuildGraph/BuildGraphScriptConditions%](conditional_topic.png)
 	[/PARAM]
 	[PARAM:icon]
 		![](%ROOT%/reference_icon.png)(convert:false)
 	[/PARAM]
 	[PARAM:title]
-		%Programming/Development/BuildGraph/BuildGraphConditions:title%
+		%Programming/Development/BuildGraph/BuildGraphScriptConditions:title%
 	[/PARAM]
 	[PARAM:description]
-		%Programming/Development/BuildGraph/BuildGraphConditions:description%
+		%Programming/Development/BuildGraph/BuildGraphScriptConditions:description%
 	[/PARAM]
 	[PARAM:path]
-		[RELATIVE:Programming/Development/BuildGraph/BuildGraphConditions]
+		[RELATIVE:Programming/Development/BuildGraph/BuildGraphScriptConditions]
 	[/PARAM]
 [/OBJECT]
 [/VAR]
 
 [TOC(start:1 end:3)]
 
-[EXCERPT:buildgraphconditions]
+[EXCERPT:BuildGraphScriptConditions]
 
 If you want to build logical complexity into your BuildGraph scripts, you'll need to work with conditionals. The following 
 section introduces you to how BuildGraph conditions are written, including a list of conditional operators. 
@@ -61,7 +61,7 @@ A list of operators is specified below:
 The `'<'` and `'>'` characters must be escaped as `"&lt;"` and `"&gt;"` in XML (referencing [XML Schema v1.0](https://www.w3.org/TR/NOTE-dcd])).
 [/REGION]
 
-[/EXCERPT:buildgraphconditions]
+[/EXCERPT:BuildGraphScriptConditions]
 
 <!--
 [EXCERPT:OperatorTable]
diff --git a/Engine/Documentation/Source/Programming/Development/BuildGraph/BuildGraphScriptElements/BuildGraphScriptElements.INT.udn b/Engine/Documentation/Source/Programming/Development/BuildGraph/BuildGraphScriptElements/BuildGraphScriptElements.INT.udn
--- a/Engine/Documentation/Source/Programming/Development/BuildGraph/BuildGraphScriptElements/BuildGraphScriptElements.INT.udn
+++ b/Engine/Documentation/Source/Programming/Development/BuildGraph/BuildGraphScriptElements/BuildGraphScriptElements.INT.udn
@@ -11,26 +11,26 @@ tags:BuildGraph
 [VAR:TopicCompact]
 [OBJECT:TopicCompact]
 	[PARAM:image]
-		![%Programming/Development/BuildGraph/BuildGraphElements%](elements_topic.png)
+		![%Programming/Development/BuildGraph/BuildGraphScriptElements%](elements_topic.png)
 	[/PARAM]
 	[PARAM:icon]
 		![](%ROOT%/reference_icon.png)(convert:false)
 	[/PARAM]
 	[PARAM:title]
-		%Programming/Development/BuildGraph/BuildGraphElements:title%
+		%Programming/Development/BuildGraph/BuildGraphScriptElements:title%
 	[/PARAM]
 	[PARAM:description]
-		%Programming/Development/BuildGraph/BuildGraphElements:description%
+		%Programming/Development/BuildGraph/BuildGraphScriptElements:description%
 	[/PARAM]
 	[PARAM:path]
-		[RELATIVE:Programming/Development/BuildGraph/BuildGraphElements]
+		[RELATIVE:Programming/Development/BuildGraph/BuildGraphScriptElements]
 	[/PARAM]
 [/OBJECT]
 [/VAR]
 
 [TOC(start:1 end:3)]
 
-[EXCERPT:buildgraphelements]
+[EXCERPT:BuildGraphScriptElements]
 
 BuildGraph scripts have elements that are similar to the sort of elements that you'd expect to find in an XML document. 
 As you read through this section, you'll learn how to define a graph along with the syntax that goes into manipulating 
@@ -185,7 +185,7 @@ the target being executed.
 
 [INCLUDE:#ErrorTable]
 
-[/EXCERPT:buildgraphelements]
+[/EXCERPT:BuildGraphScriptElements]
 
 <!--Tables for the 'Defining a Graph' section.-->
 
diff --git a/Engine/Documentation/Source/Programming/Development/BuildGraph/BuildGraphScriptTasks/BuildGraphScriptPredefinedTasks/BuildGraphScriptPredefinedTasks.INT.udn b/Engine/Documentation/Source/Programming/Development/BuildGraph/BuildGraphScriptTasks/BuildGraphScriptPredefinedTasks/BuildGraphScriptPredefinedTasks.INT.udn
--- a/Engine/Documentation/Source/Programming/Development/BuildGraph/BuildGraphScriptTasks/BuildGraphScriptPredefinedTasks/BuildGraphScriptPredefinedTasks.INT.udn
+++ b/Engine/Documentation/Source/Programming/Development/BuildGraph/BuildGraphScriptTasks/BuildGraphScriptPredefinedTasks/BuildGraphScriptPredefinedTasks.INT.udn
@@ -1,7 +1,7 @@
 Availability: NoPublish
 Title: BuildGraph Predefined Tasks
-Crumbs: %ROOT%, Programming, Programming/Development, Programming/Development/BuildGraph, Programming/Development/BuildGraph/BuildGraphTasks
+Crumbs: %ROOT%, Programming, Programming/Development, Programming/Development/BuildGraph, Programming/Development/BuildGraph/BuildGraphScriptTasks
 Description: This is a procedurally generated markdown page.
 version: 4.13
-parent:Programming/Development/BuildGraph/BuildGraphTasks
+parent:Programming/Development/BuildGraph/BuildGraphScriptTasks
 
diff --git a/Engine/Documentation/Source/Programming/Development/BuildGraph/BuildGraphScriptTasks/BuildGraphScriptTasks.INT.udn b/Engine/Documentation/Source/Programming/Development/BuildGraph/BuildGraphScriptTasks/BuildGraphScriptTasks.INT.udn
--- a/Engine/Documentation/Source/Programming/Development/BuildGraph/BuildGraphScriptTasks/BuildGraphScriptTasks.INT.udn
+++ b/Engine/Documentation/Source/Programming/Development/BuildGraph/BuildGraphScriptTasks/BuildGraphScriptTasks.INT.udn
@@ -11,26 +11,26 @@ tags:BuildGraph
 [VAR:TopicCompact]
 [OBJECT:TopicCompact]
 	[PARAM:image]
-		![%Programming/Development/BuildGraph/BuildGraphTasks%](tasks_topic.png)
+		![%Programming/Development/BuildGraph/BuildGraphScriptTasks%](tasks_topic.png)
 	[/PARAM]
 	[PARAM:icon]
 		![](%ROOT%/reference_icon.png)(convert:false)
 	[/PARAM]
 	[PARAM:title]
-		%Programming/Development/BuildGraph/BuildGraphTasks:title%
+		%Programming/Development/BuildGraph/BuildGraphScriptTasks:title%
 	[/PARAM]
 	[PARAM:description]
-		%Programming/Development/BuildGraph/BuildGraphTasks:description%
+		%Programming/Development/BuildGraph/BuildGraphScriptTasks:description%
 	[/PARAM]
 	[PARAM:path]
-		[RELATIVE:Programming/Development/BuildGraph/BuildGraphTasks]
+		[RELATIVE:Programming/Development/BuildGraph/BuildGraphScriptTasks]
 	[/PARAM]
 [/OBJECT]
 [/VAR]
 
 [TOC(start:1 end:2)]
 
-[EXCERPT:buildgraphtasks]
+[EXCERPT:BuildGraphScriptTasks]
 
 Typically, project-specific utility scripts are implemented as AutomationTool BuildCommand classes. There are times; however, when you'll want 
 to automate custom tasks for your project. This is where BuildGraph can help. BuildGraph can be extended with any number of tasks, enabling you 
@@ -53,6 +53,6 @@ If you'd like to have a starting point for creating tasks, we're providing a var
 If you're looking for a simple predefined task to start from, we recommend using `LogTask` as a good starting point. If you want to keep up with 
 our latest improvements to BuildGraph; bookmark this page, because the following list of predefined tasks will be updated periodically. 
 
-[INCLUDE:Programming/Development/BuildGraph/BuildGraphTasks/BuildGraphPredefinedTasks]
+[INCLUDE:Programming/Development/BuildGraph/BuildGraphScriptTasks/BuildGraphScriptPredefinedTasks]
 
-[/EXCERPT:buildgraphtasks]
\ No newline at end of file
+[/EXCERPT:BuildGraphScriptTasks]
\ No newline at end of file
diff --git a/Engine/Documentation/Source/Programming/Development/BuildGraph/BuildGraphScriptTypes/BuildGraphScriptTypes.INT.udn b/Engine/Documentation/Source/Programming/Development/BuildGraph/BuildGraphScriptTypes/BuildGraphScriptTypes.INT.udn
--- a/Engine/Documentation/Source/Programming/Development/BuildGraph/BuildGraphScriptTypes/BuildGraphScriptTypes.INT.udn
+++ b/Engine/Documentation/Source/Programming/Development/BuildGraph/BuildGraphScriptTypes/BuildGraphScriptTypes.INT.udn
@@ -11,24 +11,24 @@ tags:BuildGraph
 [VAR:TopicCompact]
 [OBJECT:TopicCompact]
 	[PARAM:image]
-		![%Programming/Development/BuildGraph/BuildGraphTypes%](types_topic.png)
+		![%Programming/Development/BuildGraph/BuildGraphScriptTypes%](types_topic.png)
 	[/PARAM]
 	[PARAM:icon]
 		![](%ROOT%/reference_icon.png)(convert:false)
 	[/PARAM]
 	[PARAM:title]
-		%Programming/Development/BuildGraph/BuildGraphTypes:title%
+		%Programming/Development/BuildGraph/BuildGraphScriptTypes:title%
 	[/PARAM]
 	[PARAM:description]
-		%Programming/Development/BuildGraph/BuildGraphTypes:description%
+		%Programming/Development/BuildGraph/BuildGraphScriptTypes:description%
 	[/PARAM]
 	[PARAM:path]
-		[RELATIVE:Programming/Development/BuildGraph/BuildGraphTypes]
+		[RELATIVE:Programming/Development/BuildGraph/BuildGraphScriptTypes]
 	[/PARAM]
 [/OBJECT]
 [/VAR]
 
-[EXCERPT:buildgraphtypes]
+[EXCERPT:BuildGraphScriptTypes]
 
 ## Types
 
@@ -36,7 +36,7 @@ The following table contains valid data types that can be contained by BuildGrap
 
 [INCLUDE:#TypesTable]
 
-[/EXCERPT:buildgraphtypes]
+[/EXCERPT:BuildGraphScriptTypes]
 
 <!--
 [EXCERPT:TypesTable]
@@ -56,6 +56,6 @@ The following table contains valid data types that can be contained by BuildGrap
 | **File Spec**           | A set of file and tag names, separated by semi-colons. Wildcards such as `"..."`, `"*"`, and `"?"` are permitted when referencing files (i.e., `Engine/.../*.bat`). Unless otherwise specified, relative paths are resolved relative to the working root directory.                        |
 | **File Name**           | A path to a file. Unless otherwise specified, relative paths are resolved to the working root directory.                                                                                                                                                                                   |
 | **Directory Name**      | A path to a directory. Unless otherwise specified, relative paths are resolved relative to the working root directory.                                                                                                                                                                     |
-| **Condition**           | A [conditional expression](Programming/Development/BuildGraph/BuildGraphConditions).                                                                                                                                                                                                      |
+| **Condition**           | A [conditional expression](Programming/Development/BuildGraph/BuildGraphScriptConditions).                                                                                                                                                                                                      |
 [/EXCERPT:TypesTable]
 -->
\ No newline at end of file
