#UE4 Docs: Accepting Richards changes to the document.

diff --git a/Engine/Documentation/Source/Platforms/Android/OpenGLES31MobileRenderer/OpenGLES31MobileRenderer.INT.udn b/Engine/Documentation/Source/Platforms/Android/OpenGLES31MobileRenderer/OpenGLES31MobileRenderer.INT.udn
--- a/Engine/Documentation/Source/Platforms/Android/OpenGLES31MobileRenderer/OpenGLES31MobileRenderer.INT.udn
+++ b/Engine/Documentation/Source/Platforms/Android/OpenGLES31MobileRenderer/OpenGLES31MobileRenderer.INT.udn
@@ -21,13 +21,13 @@ While Unreal Engine 4 (UE4) has long supported many OpenGL ES 3.0 and 3.1 featur
 
 ## Enabling the OpenGL ES 3.1 UE4 Editor Preview 
 
-In order to preview what OpenGL ES 3.1 will look like in the UE4 Editor you will need to first enable it by doing the following:
+In order to preview what OpenGL ES 3.1 will look like in the UE4 Editor, you will need to first enable it by doing the following:
 
-1. From the **Main Toolbar** go to **Edit** and then select the **Editor Preferences** option.
+1. From the **Main Toolbar**, go to **Edit** and then select the **Editor Preferences** option.
 
 	![](OGL3_01.png)
 
-1. Under the **General** section locate the **Experimental** category and then look for the **Rendering** section and click on the checkmark box next to **Enable Metal/Vulkan/High-end mobile Preview Rendering Level in editor** to enable the option to preview OpenGL ES 3.1. 
+1. Under the **General** section, locate the **Experimental** category and then look for the **Rendering** section. Click on the checkbox next to **Enable Metal/Vulkan/High-end mobile Preview Rendering Level in editor** to enable the option to preview OpenGL ES 3.1. 
 
 	[REGION:lightbox]
 	[![](Platforms\Android\VulkanMobileRenderer\VMP_03.png)(w:447 h:308)](Platforms\Android\VulkanMobileRenderer\VMP_03.png)
@@ -38,7 +38,7 @@ In order to preview what OpenGL ES 3.1 will look like in the UE4 Editor you will
 	[/REGION]
 
 	[region:note]
-	For a more detailed description on how to enable Metal/Vulkan/High-End  previewing in the UE4 Editor, check out the [Enabling Vulkan Preview Rendering in Editor](Platforms/Android/VulkanMobileRenderer/#enablingvulkanpreviewrenderingineditor) document.
+	For a more detailed description on how to enable Metal/Vulkan/High-End previewing in the UE4 Editor, check out the [Enabling Vulkan Preview Rendering in Editor](Platforms/Android/VulkanMobileRenderer/#enablingvulkanpreviewrenderingineditor) document.
 	[/region]
 
 1. Then on the **Main Tool Bar** click on the **Settings** option and go to **Preview Rendering Level** option and select the **High-End Mobile / Metal** option.
@@ -66,14 +66,14 @@ In order to preview what OpenGL ES 3.1 will look like in the UE4 Editor you will
 To use OpenGL ES 3.1 in your UE4 Android project, you will need to make sure your UE4 project has the following options set:
 
 [region:warning]
-Please note that OpenGL ES 3.1 Support will only work on code based projects that are downloaded and compiled from Github. If you are not familiar with how to do this check out the [Download Unreal Engine Source Code](GettingStarted/DownloadingUnrealEngine) for detailed instructions on how to do this.
+Please note that OpenGL ES 3.1 Support will only work on code based projects that are downloaded and compiled from Github. If you are not familiar with how to do this, check out the [Download Unreal Engine Source Code](GettingStarted/DownloadingUnrealEngine) for detailed instructions.
 [/region]
 
-1. First, go to the **Main Toolbar** and click on the **Edit** option and then select the **Project Settings** option. 
+1. Go to the **Main Toolbar**, click on the **Edit** option, and select the **Project Settings** option. 
 
 	![](OGL3_03.png)
 
-1. From the **Project Settings** menu go to the **Platforms** section and click on the **Android** section.
+1. From the **Project Settings** menu, go to the **Platforms** section and click on the **Android** section.
 
 	[REGION:lightbox]
 	[![](Platforms\Android\VulkanMobileRenderer\VMP_Android_Settings_00.png)(w:639 h:385)](Platforms\Android\VulkanMobileRenderer\VMP_Android_Settings_00.png)
@@ -93,7 +93,7 @@ Please note that OpenGL ES 3.1 Support will only work on code based projects tha
 	Click for full image.
 	[/REGION]
 
-1. Under the **APKPackaging** section, locate the  **Build** section and make sure that the following options are checked. 
+1. Under the **APKPackaging** section, locate the  **Build** section and make sure that the following options are checked: 
 	
 	* Support armv7
 	* Support OpenGL ES3.1
@@ -124,7 +124,7 @@ Please note that OpenGL ES 3.1 Support will only work on code based projects tha
 Now that the project has been setup to work with OpenGL ES 3.1, it is now time to build the project and deploy it to a Android based Smartphone. To build a project that has OpenGL ES 3.1 support you will need to do the following:
 
 [region:note]
-Before you begin with the next part, make sure that you have your Android based Smartphone plugged into your development PC via a USB cable. Failing to do this will result in your project not being able to be deployed to your Android Smartphone.
+Before beginning the next part, make sure that you have your Android Smartphone plugged into your development PC via a USB cable. Failing to do this will result in your project not being able to be deployed.
 [/region]
 
 
@@ -138,11 +138,11 @@ Before you begin with the next part, make sure that you have your Android based
 	Click for full image.
 	[/REGION]
 
-1. Click on the small white triangle that is on the left-hand side of the **Launch** button to expose the devices that your project can be launched on and then click on your Android device to start the build process.
+1. Click on the small white triangle that is on the left-hand side of the **Launch** button to expose the devices that your project can be launched on, and then click on your Android device to start the build process.
 
 	![](OGL3_08.png)
 
-1. After this is completed, a build status bar will be displayed in the lower right-hand corner of the UE4 Editor that will show you how the build is progressing. To see a more detailed breakdown of what is going on with the build, click on the  **Show Output Log** to open the UE4 log.
+1. After this is completed, a build status progress bar will be displayed in the lower right-hand corner of the UE4 Editor. To see a more detailed breakdown of what is going on with the build, click on **Show Output Log**.
 
 	![](OGL3_09.png)
 
