Add generated documentation for predefined BuildGraph tasks.

diff --git a/Engine/Documentation/Source/Programming/Development/BuildGraph/BuildGraphScriptTasks/BuildGraphScriptPredefinedTasks/BuildGraphScriptPredefinedTasks.INT.udn b/Engine/Documentation/Source/Programming/Development/BuildGraph/BuildGraphScriptTasks/BuildGraphScriptPredefinedTasks/BuildGraphScriptPredefinedTasks.INT.udn
--- a/Engine/Documentation/Source/Programming/Development/BuildGraph/BuildGraphScriptTasks/BuildGraphScriptPredefinedTasks/BuildGraphScriptPredefinedTasks.INT.udn
+++ b/Engine/Documentation/Source/Programming/Development/BuildGraph/BuildGraphScriptTasks/BuildGraphScriptPredefinedTasks/BuildGraphScriptPredefinedTasks.INT.udn
@@ -5,3 +5,267 @@ Description: This is a procedurally generated markdown page.
 version: 4.13
 parent:Programming/Development/BuildGraph/BuildGraphScriptTasks
 
+## Command
+
+Invokes an AutomationTool child process to run the given command.
+
+|                          |           |          |                                                                                                                                                                   |
+| ------------------------ | --------- | -------- | ----------------------------------------------------------------------------------------------------------------------------------------------------------------- |
+| Name                     | String    | Required | The command name to execute                                                                                                                                       |
+| Arguments                | String    | Optional | Arguments to be passed to the command                                                                                                                             |
+| MergeTelemetryWithPrefix | String    | Optional | If non-null, instructs telemetry from the command to be merged into the telemetry for this UAT instance with the given prefix. May be an empty (non-null) string. |
+| If                       | Condition | Optional | Whether to execute this task. It is ignored if this condition evaluates to false.                                                                                 |
+
+## Commandlet
+
+Spawns the editor to run a commandlet.
+
+|           |           |          |                                                                                                          |
+| --------- | --------- | -------- | -------------------------------------------------------------------------------------------------------- |
+| Name      | String    | Required | The commandlet name to execute                                                                           |
+| Project   | File Spec | Optional | The project to run the editor with                                                                       |
+| Arguments | String    | Optional | Arguments to be passed to the commandlet                                                                 |
+| EditorExe | File Name | Optional | The editor executable to use. Defaults to the development UE4Editor executable for the current platform. |
+| If        | Condition | Optional | Whether to execute this task. It is ignored if this condition evaluates to false.                        |
+
+## Compile
+
+Compiles a target with UnrealBuildTool.
+
+|               |                           |          |                                                                                   |
+| ------------- | ------------------------- | -------- | --------------------------------------------------------------------------------- |
+| Target        | String                    | Required | The target to compile                                                             |
+| Configuration | UnrealTargetConfiguration | Required | The configuration to compile                                                      |
+| Platform      | UnrealTargetPlatform      | Required | The platform to compile for                                                       |
+| Arguments     | String                    | Optional | Additional arguments for UnrealBuildTool                                          |
+| Tag           | Tag List                  | Optional | Tag to be applied to build products of this task                                  |
+| If            | Condition                 | Optional | Whether to execute this task. It is ignored if this condition evaluates to false. |
+
+## Cook
+
+Cook a selection of maps for a certain platform
+
+|           |           |          |                                                                                   |
+| --------- | --------- | -------- | --------------------------------------------------------------------------------- |
+| Project   | String    | Required | Project file to be cooked                                                         |
+| Platform  | String    | Required | The cook platform to target (eg. WindowsNoEditor)                                 |
+| Maps      | String    | Optional | List of maps to be cooked, separated by '+' characters                            |
+| Versioned | Boolean   | Optional | Additional arguments to be passed to the cooker                                   |
+| Arguments | String    | Optional | Additional arguments to be passed to the cooker                                   |
+| Tag       | Tag List  | Optional | Tag to be applied to build products of this task                                  |
+| If        | Condition | Optional | Whether to execute this task. It is ignored if this condition evaluates to false. |
+
+## Copy
+
+Copies files from one directory to another.
+
+|         |                |          |                                                                                                                                                  |
+| ------- | -------------- | -------- | ------------------------------------------------------------------------------------------------------------------------------------------------ |
+| Files   | File Spec      | Required | List of file specifications separated by semicolons (eg. *.cpp;Engine/.../*.bat), or the name of a tag set. Relative paths are based at FromDir. |
+| FromDir | Directory Name | Required | The base directory to copy from.                                                                                                                 |
+| ToDir   | Directory Name | Required | The directory to copy to                                                                                                                         |
+| Tag     | Tag List       | Optional | Tag to be applied to build products of this task                                                                                                 |
+| If      | Condition      | Optional | Whether to execute this task. It is ignored if this condition evaluates to false.                                                                |
+
+## CsCompile
+
+Compiles C# project files, and their dependencies.
+
+|               |           |          |                                                                                                                                                  |
+| ------------- | --------- | -------- | ------------------------------------------------------------------------------------------------------------------------------------------------ |
+| Project       | String    | Required | The C# project file to be compile. More than one project file can be specified by separating with semicolons.                                    |
+| Configuration | String    | Optional | The configuration to compile                                                                                                                     |
+| Platform      | String    | Optional | The platform to compile                                                                                                                          |
+| Arguments     | String    | Optional | Additional options to pass to the compiler                                                                                                       |
+| EnumerateOnly | Boolean   | Optional | Only enumerate build products; do not actually compile the projects.                                                                             |
+| Tag           | Tag List  | Optional | Tag to be applied to build products of this task                                                                                                 |
+| TagReferences | Tag List  | Optional | Tag to be applied to any non-private references the projects have
+            (i.e. those that are external and not copied into the output dir) |
+| If            | Condition | Optional | Whether to execute this task. It is ignored if this condition evaluates to false.                                                                |
+
+## Delete
+
+Delete a set of files.
+
+|                        |           |          |                                                                                                            |
+| ---------------------- | --------- | -------- | ---------------------------------------------------------------------------------------------------------- |
+| Files                  | File Spec | Required | List of file specifications separated by semicolons (eg. *.cpp;Engine/.../*.bat), or the name of a tag set |
+| DeleteEmptyDirectories | Boolean   | Optional | Whether to delete empty directories after deleting the files. Defaults to true.                            |
+| If                     | Condition | Optional | Whether to execute this task. It is ignored if this condition evaluates to false.                          |
+
+## Log
+
+Print a message (and other optional diagnostic information) to the output log
+
+|                 |           |          |                                                                                     |
+| --------------- | --------- | -------- | ----------------------------------------------------------------------------------- |
+| Message         | String    | Optional | Message to print out                                                                |
+| Files           | File Spec | Optional | If specified, causes the given list of files to be printed after the given message. |
+| IncludeContents | Boolean   | Optional | If specified, causes the contents of the given files to be printed out.             |
+| If              | Condition | Optional | Whether to execute this task. It is ignored if this condition evaluates to false.   |
+
+## PakFile
+
+Creates a PAK file from a given set of files.
+
+|              |                |          |                                                                                                                                   |
+| ------------ | -------------- | -------- | --------------------------------------------------------------------------------------------------------------------------------- |
+| Files        | File Spec      | Required | List of files, wildcards and tag sets to add to the pak file, separated by ';' characters.                                        |
+| Output       | File Name      | Required | PAK file to output                                                                                                                |
+| ResponseFile | File Name      | Optional | Path to a Response File that contains a list of files to add to the pak file, instead of specifying them individually             |
+| RebaseDir    | Directory Name | Optional | Directories to rebase the files relative to. If specified, the shortest path under a listed directory will be used for each file. |
+| Order        | File Name      | Optional | Script which gives the order of files                                                                                             |
+| Sign         | String         | Optional | Encryption keys for this pak file                                                                                                 |
+| Compress     | Boolean        | Optional | Whether to compress files                                                                                                         |
+| Arguments    | String         | Optional | Additional arguments to be passed to UnrealPak                                                                                    |
+| Tag          | Tag List       | Optional | Tag to be applied to build products of this task                                                                                  |
+| If           | Condition      | Optional | Whether to execute this task. It is ignored if this condition evaluates to false.                                                 |
+
+## Rename
+
+Renames a file, or group of files.
+
+|       |                |          |                                                                                                 |
+| ----- | -------------- | -------- | ----------------------------------------------------------------------------------------------- |
+| Files | File Spec      | Required | The file or files to rename                                                                     |
+| From  | Directory Name | Optional | The current file name, or pattern to match (eg. *.txt). Should not include any path separators. |
+| To    | Directory Name | Required | The new name for the file(s). Should not include any path separators.                           |
+| Tag   | Tag List       | Optional | Tag to be applied to the renamed files                                                          |
+| If    | Condition      | Optional | Whether to execute this task. It is ignored if this condition evaluates to false.               |
+
+## SetVersion
+
+Updates the local version files (Engine/Source/Runtime/Launch/Resources/Version.h, Engine/Build/Build.version, and Engine/Source/Programs/DotNETCommon/Metadata.cs) with the given version information.
+
+|                  |           |          |                                                                                                                           |
+| ---------------- | --------- | -------- | ------------------------------------------------------------------------------------------------------------------------- |
+| Change           | Int32     | Required | The changelist to set in the version files                                                                                |
+| CompatibleChange | Int32     | Optional | The engine compatible changelist to set in the version files                                                              |
+| Branch           | String    | Required | The branch string                                                                                                         |
+| Build            | String    | Optional | The build version string                                                                                                  |
+| Licensee         | Boolean   | Optional | Whether to set the IS_LICENSEE_VERSION flag to true                                                                       |
+| SkipWrite        | Boolean   | Optional | If set, don't actually write to the files - just return the version files that would be updated. Useful for local builds. |
+| Tag              | Tag List  | Optional | Tag to be applied to build products of this task                                                                          |
+| If               | Condition | Optional | Whether to execute this task. It is ignored if this condition evaluates to false.                                         |
+
+## Sign
+
+Signs a set of executable files with an installed certificate.
+
+|       |           |          |                                                                                                            |
+| ----- | --------- | -------- | ---------------------------------------------------------------------------------------------------------- |
+| Files | File Spec | Required | List of file specifications separated by semicolons (eg. *.cpp;Engine/.../*.bat), or the name of a tag set |
+| Tag   | Tag List  | Optional | Tag to be applied to build products of this task                                                           |
+| If    | Condition | Optional | Whether to execute this task. It is ignored if this condition evaluates to false.                          |
+
+## Spawn
+
+Spawns an external executable and waits for it to complete.
+
+|            |           |          |                                                                                   |
+| ---------- | --------- | -------- | --------------------------------------------------------------------------------- |
+| Exe        | File Name | Required | Executable to spawn                                                               |
+| Arguments  | String    | Optional | Arguments for the newly created process                                           |
+| ErrorLevel | Int32     | Optional | The minimum exit code which is treated as an error.                               |
+| If         | Condition | Optional | Whether to execute this task. It is ignored if this condition evaluates to false. |
+
+## Stage
+
+Stages files listed in a build receipt to an output directory.
+
+|               |                           |          |                                                                                   |
+| ------------- | ------------------------- | -------- | --------------------------------------------------------------------------------- |
+| Project       | File Name                 | Optional | The project that this target belongs to                                           |
+| Target        | String                    | Required | Name of the target to stage                                                       |
+| Platform      | UnrealTargetPlatform      | Required | Platform to stage                                                                 |
+| Configuration | UnrealTargetConfiguration | Required | Configuration to be staged                                                        |
+| Architecture  | String                    | Optional | Architecture to be staged                                                         |
+| ToDir         | String                    | Required | Directory the receipt files should be staged to                                   |
+| Overwrite     | Boolean                   | Optional | Whether to overwrite existing files                                               |
+| Tag           | Tag List                  | Optional | Tag to be applied to build products of this task                                  |
+| If            | Condition                 | Optional | Whether to execute this task. It is ignored if this condition evaluates to false. |
+
+## Strip
+
+Strips debugging information from a set of files.
+
+|           |                      |          |                                                                                                      |
+| --------- | -------------------- | -------- | ---------------------------------------------------------------------------------------------------- |
+| Platform  | UnrealTargetPlatform | Required | The platform toolchain to strip binaries                                                             |
+| BaseDir   | Directory Name       | Optional | The directory to find files in                                                                       |
+| Files     | File Spec            | Required | List of file specifications separated by semicolons (eg. Engine/.../*.pdb), or the name of a tag set |
+| OutputDir | Directory Name       | Required | Output directory for the stripped files. Defaults to the input path (overwriting the input files).   |
+| Tag       | Tag List             | Optional | Tag to be applied to build products of this task                                                     |
+| If        | Condition            | Optional | Whether to execute this task. It is ignored if this condition evaluates to false.                    |
+
+## Submit
+
+Creates a new changelist and submits a set of files to a Perforce stream.
+
+|             |                |          |                                                                                                                                                                              |
+| ----------- | -------------- | -------- | ---------------------------------------------------------------------------------------------------------------------------------------------------------------------------- |
+| Description | String         | Required | The description for the submitted changelist                                                                                                                                 |
+| Files       | File Spec      | Required | The files to submit                                                                                                                                                          |
+| FileType    | String         | Optional | The Perforce file type for the submitted files (eg. binary+FS32)                                                                                                             |
+| Workspace   | String         | Optional | The workspace name. If specified, a new workspace will be created using the given stream and root directory to submit the files. If not, the current workspace will be used. |
+| Stream      | String         | Optional | The stream for the workspace; defaults to the current stream. Ignored unless If the Workspace attribute is also specified.                                                   |
+| RootDir     | Directory Name | Optional | Root directory for the stream. If not specified, defaults to the current root directory.                                                                                     |
+| If          | Condition      | Optional | Whether to execute this task. It is ignored if this condition evaluates to false.                                                                                            |
+
+## Tag
+
+Applies a tag to a given set of files. The list of files is found by enumerating the tags and file specifications given by the 'Files' 
+            parameter. From this list, any files not matched by the 'Filter' parameter are removed, followed by any files matched by the 'Except' parameter.
+
+|         |                |          |                                                                                                                                                                                                                  |
+| ------- | -------------- | -------- | ---------------------------------------------------------------------------------------------------------------------------------------------------------------------------------------------------------------- |
+| BaseDir | Directory Name | Optional | Set the base directory to resolve relative paths and patterns against. If set, any absolute patterns (eg. /Engine/Build/...) are taken to be relative to this path. If not, they are taken to be truly absolute. |
+| Files   | File Spec      | Required | Set of files to work from, including wildcards and tag names, separated by semicolons. Resolved relative to BaseDir if set, otherwise to the branch root directory.                                              |
+| Filter  | File Spec      | Optional | Patterns to filter the list of files by, including tag names or wildcards. May include patterns that apply to the base directory if set. Defaults to all files if not specified.                                 |
+| Except  | File Spec      | Optional | Set of patterns to exclude from the matched list. May include tag names of patterns that apply to the base directory.                                                                                            |
+| With    | Tag List       | Required | Name of the tag to apply                                                                                                                                                                                         |
+| If      | Condition      | Optional | Whether to execute this task. It is ignored if this condition evaluates to false.                                                                                                                                |
+
+## TagReceipt
+
+Task which tags build products and/or runtime dependencies by reading from *.target files.
+
+|                                |                |          |                                                                                                                                                   |
+| ------------------------------ | -------------- | -------- | ------------------------------------------------------------------------------------------------------------------------------------------------- |
+| Files                          | File Spec      | Required | Set of receipt files (*.target) to read, including wildcards and tag names, separated by semicolons.                                              |
+| EngineDir                      | Directory Name | Optional | Path to the Engine folder, used to expand $(EngineDir) properties in receipt files. Defaults to the Engine directory for the current workspace.   |
+| ProjectDir                     | Directory Name | Optional | Path to the project folder, used to expand $(ProjectDir) properties in receipt files. Defaults to the Engine directory for the current workspace. |
+| BuildProducts                  | Boolean        | Optional | Whether to tag the Build Products listed in receipts                                                                                              |
+| BuildProductType               | String         | Optional | Which type of Build Products to tag (See TargetReceipt.cs - UnrealBuildTool.BuildProductType for valid values)                                    |
+| RuntimeDependencies            | Boolean        | Optional | Whether to tag the Runtime Dependencies listed in receipts                                                                                        |
+| StagedFileType                 | String         | Optional | Which type of Runtime Dependencies to tag (See TargetReceipt.cs - UnrealBuildTool.StagedFileType for valid values)                                |
+| PrecompiledBuildDependencies   | Boolean        | Optional | Whether to tag the Precompiled Build Dependencies listed in receipts                                                                              |
+| PrecompiledRuntimeDependencies | Boolean        | Optional | Whether to tag the Precompiled Runtime Dependencies listed in receipts                                                                            |
+| With                           | Tag List       | Required | Name of the tag to apply                                                                                                                          |
+| If                             | Condition      | Optional | Whether to execute this task. It is ignored if this condition evaluates to false.                                                                 |
+
+## Unzip
+
+Extract files from a zip archive.
+
+|         |                |          |                                                                                   |
+| ------- | -------------- | -------- | --------------------------------------------------------------------------------- |
+| ZipFile | File Name      | Required | Path to the zip file to extract                                                   |
+| ToDir   | Directory Name | Required | Output directory for the extracted files                                          |
+| Tag     | Tag List       | Optional | Tag to be applied to the extracted files                                          |
+| If      | Condition      | Optional | Whether to execute this task. It is ignored if this condition evaluates to false. |
+
+## Zip
+
+Compresses files into a zip archive.
+
+|         |                |          |                                                                                                                                                    |
+| ------- | -------------- | -------- | -------------------------------------------------------------------------------------------------------------------------------------------------- |
+| FromDir | Directory Name | Required | The directory to read compressed files from                                                                                                        |
+| Files   | File Spec      | Optional | List of file specifications separated by semicolons (eg. *.cpp;Engine/.../*.bat), or the name of a tag set. Relative paths are taken from FromDir. |
+| ZipFile | File Name      | Required | The zip file to create                                                                                                                             |
+| Tag     | Tag List       | Optional | Tag to be applied to the created zip file                                                                                                          |
+| If      | Condition      | Optional | Whether to execute this task. It is ignored if this condition evaluates to false.                                                                  |
+
+  </body>
+</html>
