Updated TOC depth.

#UE4doc

diff --git a/Engine/Documentation/Source/Programming/Development/BuildGraph/BuildGraphScriptTasks/BuildGraphScriptTasks.INT.udn b/Engine/Documentation/Source/Programming/Development/BuildGraph/BuildGraphScriptTasks/BuildGraphScriptTasks.INT.udn
--- a/Engine/Documentation/Source/Programming/Development/BuildGraph/BuildGraphScriptTasks/BuildGraphScriptTasks.INT.udn
+++ b/Engine/Documentation/Source/Programming/Development/BuildGraph/BuildGraphScriptTasks/BuildGraphScriptTasks.INT.udn
@@ -28,7 +28,7 @@ tags:BuildGraph
 [/OBJECT]
 [/VAR]
 
-[TOC(start:1 end:2)]
+[TOC(start:1 end:3)]
 
 [EXCERPT:BuildGraphScriptTasks]
 
