Repaired UDN file for BenM's reference.

diff --git a/Engine/Documentation/Source/Programming/Development/BuildGraph/BuildGraphScriptTasks/BuildGraphScriptPredefinedTasks/BuildGraphScriptPredefinedTasks.INT.udn b/Engine/Documentation/Source/Programming/Development/BuildGraph/BuildGraphScriptTasks/BuildGraphScriptPredefinedTasks/BuildGraphScriptPredefinedTasks.INT.udn
--- a/Engine/Documentation/Source/Programming/Development/BuildGraph/BuildGraphScriptTasks/BuildGraphScriptPredefinedTasks/BuildGraphScriptPredefinedTasks.INT.udn
+++ b/Engine/Documentation/Source/Programming/Development/BuildGraph/BuildGraphScriptTasks/BuildGraphScriptPredefinedTasks/BuildGraphScriptPredefinedTasks.INT.udn
@@ -5,7 +5,7 @@ Description: This is a procedurally generated markdown page.
 version: 4.13
 parent:Programming/Development/BuildGraph/BuildGraphScriptTasks
 
-## Command
+### Command
 
 Invokes an AutomationTool child process to run the given command.
 
@@ -16,7 +16,7 @@ Invokes an AutomationTool child process to run the given command.
 | MergeTelemetryWithPrefix | String    | Optional | If non-null, instructs telemetry from the command to be merged into the telemetry for this UAT instance with the given prefix. May be an empty (non-null) string. |
 | If                       | Condition | Optional | Whether to execute this task. It is ignored if this condition evaluates to false.                                                                                 |
 
-## Commandlet
+### Commandlet
 
 Spawns the editor to run a commandlet.
 
@@ -28,7 +28,7 @@ Spawns the editor to run a commandlet.
 | EditorExe | File Name | Optional | The editor executable to use. Defaults to the development UE4Editor executable for the current platform. |
 | If        | Condition | Optional | Whether to execute this task. It is ignored if this condition evaluates to false.                        |
 
-## Compile
+### Compile
 
 Compiles a target with UnrealBuildTool.
 
@@ -41,7 +41,7 @@ Compiles a target with UnrealBuildTool.
 | Tag           | Tag List                  | Optional | Tag to be applied to build products of this task                                  |
 | If            | Condition                 | Optional | Whether to execute this task. It is ignored if this condition evaluates to false. |
 
-## Cook
+### Cook
 
 Cook a selection of maps for a certain platform
 
@@ -55,7 +55,7 @@ Cook a selection of maps for a certain platform
 | Tag       | Tag List  | Optional | Tag to be applied to build products of this task                                  |
 | If        | Condition | Optional | Whether to execute this task. It is ignored if this condition evaluates to false. |
 
-## Copy
+### Copy
 
 Copies files from one directory to another.
 
@@ -67,7 +67,7 @@ Copies files from one directory to another.
 | Tag     | Tag List       | Optional | Tag to be applied to build products of this task                                                                                                 |
 | If      | Condition      | Optional | Whether to execute this task. It is ignored if this condition evaluates to false.                                                                |
 
-## CsCompile
+### CsCompile
 
 Compiles C# project files, and their dependencies.
 
@@ -79,11 +79,10 @@ Compiles C# project files, and their dependencies.
 | Arguments     | String    | Optional | Additional options to pass to the compiler                                                                                                       |
 | EnumerateOnly | Boolean   | Optional | Only enumerate build products; do not actually compile the projects.                                                                             |
 | Tag           | Tag List  | Optional | Tag to be applied to build products of this task                                                                                                 |
-| TagReferences | Tag List  | Optional | Tag to be applied to any non-private references the projects have
-            (i.e. those that are external and not copied into the output dir) |
+| TagReferences | Tag List  | Optional | Tag to be applied to any non-private references the projects have (i.e. those that are external and not copied into the output dir)              |
 | If            | Condition | Optional | Whether to execute this task. It is ignored if this condition evaluates to false.                                                                |
 
-## Delete
+### Delete
 
 Delete a set of files.
 
@@ -93,7 +92,7 @@ Delete a set of files.
 | DeleteEmptyDirectories | Boolean   | Optional | Whether to delete empty directories after deleting the files. Defaults to true.                            |
 | If                     | Condition | Optional | Whether to execute this task. It is ignored if this condition evaluates to false.                          |
 
-## Log
+### Log
 
 Print a message (and other optional diagnostic information) to the output log
 
@@ -104,7 +103,7 @@ Print a message (and other optional diagnostic information) to the output log
 | IncludeContents | Boolean   | Optional | If specified, causes the contents of the given files to be printed out.             |
 | If              | Condition | Optional | Whether to execute this task. It is ignored if this condition evaluates to false.   |
 
-## PakFile
+### PakFile
 
 Creates a PAK file from a given set of files.
 
@@ -121,7 +120,7 @@ Creates a PAK file from a given set of files.
 | Tag          | Tag List       | Optional | Tag to be applied to build products of this task                                                                                  |
 | If           | Condition      | Optional | Whether to execute this task. It is ignored if this condition evaluates to false.                                                 |
 
-## Rename
+### Rename
 
 Renames a file, or group of files.
 
@@ -133,7 +132,7 @@ Renames a file, or group of files.
 | Tag   | Tag List       | Optional | Tag to be applied to the renamed files                                                          |
 | If    | Condition      | Optional | Whether to execute this task. It is ignored if this condition evaluates to false.               |
 
-## SetVersion
+### SetVersion
 
 Updates the local version files (Engine/Source/Runtime/Launch/Resources/Version.h, Engine/Build/Build.version, and Engine/Source/Programs/DotNETCommon/Metadata.cs) with the given version information.
 
@@ -148,7 +147,7 @@ Updates the local version files (Engine/Source/Runtime/Launch/Resources/Version.
 | Tag              | Tag List  | Optional | Tag to be applied to build products of this task                                                                          |
 | If               | Condition | Optional | Whether to execute this task. It is ignored if this condition evaluates to false.                                         |
 
-## Sign
+### Sign
 
 Signs a set of executable files with an installed certificate.
 
@@ -158,7 +157,7 @@ Signs a set of executable files with an installed certificate.
 | Tag   | Tag List  | Optional | Tag to be applied to build products of this task                                                           |
 | If    | Condition | Optional | Whether to execute this task. It is ignored if this condition evaluates to false.                          |
 
-## Spawn
+### Spawn
 
 Spawns an external executable and waits for it to complete.
 
@@ -169,7 +168,7 @@ Spawns an external executable and waits for it to complete.
 | ErrorLevel | Int32     | Optional | The minimum exit code which is treated as an error.                               |
 | If         | Condition | Optional | Whether to execute this task. It is ignored if this condition evaluates to false. |
 
-## Stage
+### Stage
 
 Stages files listed in a build receipt to an output directory.
 
@@ -185,7 +184,7 @@ Stages files listed in a build receipt to an output directory.
 | Tag           | Tag List                  | Optional | Tag to be applied to build products of this task                                  |
 | If            | Condition                 | Optional | Whether to execute this task. It is ignored if this condition evaluates to false. |
 
-## Strip
+### Strip
 
 Strips debugging information from a set of files.
 
@@ -198,7 +197,7 @@ Strips debugging information from a set of files.
 | Tag       | Tag List             | Optional | Tag to be applied to build products of this task                                                     |
 | If        | Condition            | Optional | Whether to execute this task. It is ignored if this condition evaluates to false.                    |
 
-## Submit
+### Submit
 
 Creates a new changelist and submits a set of files to a Perforce stream.
 
@@ -212,7 +211,7 @@ Creates a new changelist and submits a set of files to a Perforce stream.
 | RootDir     | Directory Name | Optional | Root directory for the stream. If not specified, defaults to the current root directory.                                                                                     |
 | If          | Condition      | Optional | Whether to execute this task. It is ignored if this condition evaluates to false.                                                                                            |
 
-## Tag
+### Tag
 
 Applies a tag to a given set of files. The list of files is found by enumerating the tags and file specifications given by the 'Files' 
             parameter. From this list, any files not matched by the 'Filter' parameter are removed, followed by any files matched by the 'Except' parameter.
@@ -226,7 +225,7 @@ Applies a tag to a given set of files. The list of files is found by enumerating
 | With    | Tag List       | Required | Name of the tag to apply                                                                                                                                                                                         |
 | If      | Condition      | Optional | Whether to execute this task. It is ignored if this condition evaluates to false.                                                                                                                                |
 
-## TagReceipt
+### TagReceipt
 
 Task which tags build products and/or runtime dependencies by reading from *.target files.
 
@@ -244,7 +243,7 @@ Task which tags build products and/or runtime dependencies by reading from *.tar
 | With                           | Tag List       | Required | Name of the tag to apply                                                                                                                          |
 | If                             | Condition      | Optional | Whether to execute this task. It is ignored if this condition evaluates to false.                                                                 |
 
-## Unzip
+### Unzip
 
 Extract files from a zip archive.
 
@@ -255,7 +254,7 @@ Extract files from a zip archive.
 | Tag     | Tag List       | Optional | Tag to be applied to the extracted files                                          |
 | If      | Condition      | Optional | Whether to execute this task. It is ignored if this condition evaluates to false. |
 
-## Zip
+### Zip
 
 Compresses files into a zip archive.
 
@@ -267,5 +266,3 @@ Compresses files into a zip archive.
 | Tag     | Tag List       | Optional | Tag to be applied to the created zip file                                                                                                          |
 | If      | Condition      | Optional | Whether to execute this task. It is ignored if this condition evaluates to false.                                                                  |
 
-  </body>
-</html>
