Minor grammatical fix.

#UE4doc

diff --git a/Engine/Documentation/Source/Programming/Development/BuildGraph/BuildGraphIntroduction/BuildGraphIntroduction.INT.udn b/Engine/Documentation/Source/Programming/Development/BuildGraph/BuildGraphIntroduction/BuildGraphIntroduction.INT.udn
--- a/Engine/Documentation/Source/Programming/Development/BuildGraph/BuildGraphIntroduction/BuildGraphIntroduction.INT.udn
+++ b/Engine/Documentation/Source/Programming/Development/BuildGraph/BuildGraphIntroduction/BuildGraphIntroduction.INT.udn
@@ -9,7 +9,7 @@ type:Reference
 **BuildGraph** is a system that runs BuildGraph scripts for build automation (with building blocks common to Unreal Engine 4 (UE4) projects). 
 BuildGraph integrates with UnrealBuildTool, AutomationTool, and the editor. Which means that BuildGraph can be extended and customized for your projects.
 
-BuildGraph scripts are written with [XML](https://www.w3.org/TR/NOTE-dcd]), and specify a graph of user-defined nodes with 
+BuildGraph scripts are written with [XML](https://www.w3.org/TR/NOTE-dcd]), specifying a graph of user-defined nodes with 
 dependencies between them. Each node consists of tasks executed in sequence to produce some sort of output (i.e., compile 
 this project, then cook, run this test, and so on). When asked to build a target (that is, a node or named output) BuildGraph 
 will execute all of the nodes in the graph required to make that happen.
