Minor fix.

#UE4doc

diff --git a/Engine/Documentation/Source/Programming/Development/BuildGraph/BuildGraphScriptElements/BuildGraphScriptElements.INT.udn b/Engine/Documentation/Source/Programming/Development/BuildGraph/BuildGraphScriptElements/BuildGraphScriptElements.INT.udn
--- a/Engine/Documentation/Source/Programming/Development/BuildGraph/BuildGraphScriptElements/BuildGraphScriptElements.INT.udn
+++ b/Engine/Documentation/Source/Programming/Development/BuildGraph/BuildGraphScriptElements/BuildGraphScriptElements.INT.udn
@@ -137,10 +137,10 @@ the conditions on a sequence of branches, processing the first statement with a
 Example `<Switch>` statement:
 
 	<Switch>
-		<Case If=”X”>
+		<Case If=X>
 			<Elements go here>
 		</Case>
-		<Case If=”Y”>
+		<Case If=Y>
 			<Elements go here>
 		</Case>
 		<Default>
