Addressed SME review comments.

#UE4doc

diff --git a/Engine/Documentation/Source/Platforms/Linux/BeginnerLinuxDeveloper/SupportedGfxDrivers/SupportedGfxDrivers.INT.udn b/Engine/Documentation/Source/Platforms/Linux/BeginnerLinuxDeveloper/SupportedGfxDrivers/SupportedGfxDrivers.INT.udn
--- a/Engine/Documentation/Source/Platforms/Linux/BeginnerLinuxDeveloper/SupportedGfxDrivers/SupportedGfxDrivers.INT.udn
+++ b/Engine/Documentation/Source/Platforms/Linux/BeginnerLinuxDeveloper/SupportedGfxDrivers/SupportedGfxDrivers.INT.udn
@@ -42,6 +42,11 @@ Graphics drivers are typically provided by manufacturers for testing and integra
 Engine 4 (UE4). If you want to get the most out of UE4 while maintaining hardware performance, please 
 make sure to update your graphics driver appropriately.
 
+[REGION:note]
+If your system has Mesa driver version older than Mesa 10.5.x, we recommend that you update your Mesa 
+driver to be newer than Mesa 10.5.x (refer to Ubuntu's documentation).
+[/REGION]
+
 ## NVIDIA Graphics Drivers
 
 If you're running an NVIDIA graphics driver version older than version 347.09, consider updating your graphics driver 
@@ -56,10 +61,10 @@ your system from crashing while running the engine.
 
 <!--
 [EXCERPT:SupportedGfxDrivers]
-| **Manufacturer** | **Supported Graphics Driver Versions** | **Suggested Graphics Driver Version** |
-| ---------------- | -------------------------------------- | ------------------------------------- |
-| **NVIDIA**       | Between 347.09 and 361.43              | 361.75                                |
-| **AMD**          | 14.301.1001.0 or newer                 | 15.30.1025.1001                       |
-| **Intel**        | Any                                    | Any                                   |
+| **Manufacturer** | **Supported Graphics Driver Versions**    | **Suggested Graphics Driver Version**     |
+| ---------------- | ----------------------------------------- | ----------------------------------------- |
+| **NVIDIA**       | Between 347.09 and 361.43                 | 361.75                                    |
+| **AMD**          | 14.301.1001.0 or newer                    | 15.30.1025.1001                           |
+| **Intel**        | Latest available version for your system. | Latest available version for your system. |
 [/EXCERPT:SupportedGfxDrivers]
 -->
\ No newline at end of file
