#UE4 docs - fixing non-bolded trigger text

diff --git a/Engine/Documentation/Source/Engine/Editor/VR/UniversalGizmo/UniversalGizmo.INT.udn b/Engine/Documentation/Source/Engine/Editor/VR/UniversalGizmo/UniversalGizmo.INT.udn
--- a/Engine/Documentation/Source/Engine/Editor/VR/UniversalGizmo/UniversalGizmo.INT.udn
+++ b/Engine/Documentation/Source/Engine/Editor/VR/UniversalGizmo/UniversalGizmo.INT.udn
@@ -117,7 +117,7 @@ Each gizmo gives you access to different transform actions. By aiming the contro
 |[REGION:none](#rotate)[/REGION] **Rotate** | ![](LocalRotateGizmo_Highlight.png)(w:300) | Aim the controller at one of the rotation manipulators and squeeze the **Trigger** halfway to grab it. Move the controller to rotate the selected Actor(s) around the axis. |
 | [REGION:none](#uniformscale)[/REGION]**Uniform Scale** | ![](WorldScaleGizmo_Highlight.png)(w:300) | Aim the controller at the white cube in the center of the gizmo and squeeze the **Trigger** halfway to grab it. Move the controller to scale the selected Actor(s) along all axes uniformly. |
 | [REGION:none](#scale)[/REGION]**Scale** | ![](LocalScaleGizmo_Highlight.png)(w:300) | Aim the controller at one of the scale manipulators and squeeze the **Trigger** halfway to grab it. Move the controller to scale the selected Actor(s) along the selected axis or plane. |
-| [REGION:none](#stretch)[/REGION]**Stretch** |![](LocalScaleGizmo_Stretch.png)(w:300) | Aim the controller at one of the stretch manipulators and squeeze the Trigger halfway to grab it. Move the controller to scale the selected Actor(s) non-uniformly along all axes while remaining anchored to the opposite corner.| 
+| [REGION:none](#stretch)[/REGION]**Stretch** |![](LocalScaleGizmo_Stretch.png)(w:300) | Aim the controller at one of the stretch manipulators and squeeze the **Trigger** halfway to grab it. Move the controller to scale the selected Actor(s) non-uniformly along all axes while remaining anchored to the opposite corner.| 
 
 
 
