Addressed DmitryR peer review comment.

#UE4doc

diff --git a/Engine/Documentation/Source/Platforms/Linux/BeginnerLinuxDeveloper/SupportedGfxDrivers/SupportedGfxDrivers.INT.udn b/Engine/Documentation/Source/Platforms/Linux/BeginnerLinuxDeveloper/SupportedGfxDrivers/SupportedGfxDrivers.INT.udn
--- a/Engine/Documentation/Source/Platforms/Linux/BeginnerLinuxDeveloper/SupportedGfxDrivers/SupportedGfxDrivers.INT.udn
+++ b/Engine/Documentation/Source/Platforms/Linux/BeginnerLinuxDeveloper/SupportedGfxDrivers/SupportedGfxDrivers.INT.udn
@@ -42,11 +42,6 @@ Graphics drivers are typically provided by manufacturers for testing and integra
 Engine 4 (UE4). If you want to get the most out of UE4 while maintaining hardware performance, please 
 make sure to update your graphics driver appropriately.
 
-[REGION:note]
-If your system has Mesa driver version older than Mesa 10.5.x, we recommend that you update your Mesa 
-driver to be newer than Mesa 10.5.x (refer to Ubuntu's documentation).
-[/REGION]
-
 ## NVIDIA Graphics Drivers
 
 If you're running an NVIDIA graphics driver version older than version 347.09, consider updating your graphics driver 
@@ -59,6 +54,10 @@ drivers in order to verify that your graphics driver will work properly with the
 If you're running an AMD graphics driver version older than 14.301.1001.0, you should update your graphics driver to prevent 
 your system from crashing while running the engine.
 
+## Intel Graphics Drivers
+
+If your system has Mesa drivers, we recommend that you update your driver to be newer than Mesa 10.5.x.
+
 <!--
 [EXCERPT:SupportedGfxDrivers]
 | **Manufacturer** | **Supported Graphics Driver Versions**    | **Suggested Graphics Driver Version**     |
