Fixing broken links and removing introduction page per JeffW.

#UE4doc

diff --git a/Engine/Documentation/Source/Programming/Development/BuildGraph/BuildGraph.INT.udn b/Engine/Documentation/Source/Programming/Development/BuildGraph/BuildGraph.INT.udn
--- a/Engine/Documentation/Source/Programming/Development/BuildGraph/BuildGraph.INT.udn
+++ b/Engine/Documentation/Source/Programming/Development/BuildGraph/BuildGraph.INT.udn
@@ -7,26 +7,38 @@ type:landing
 parent:Programming/Development
 tags:BuildGraph
 
-[INCLUDE:Programming/Development/BuildGraph/BuildGraphIntroduction]
+**BuildGraph** is a system that runs BuildGraph scripts for build automation (with building blocks common to Unreal Engine 4 (UE4) projects). 
+BuildGraph integrates with UnrealBuildTool, AutomationTool, and the editor. Which means that BuildGraph can be extended and customized for your projects.
+
+BuildGraph scripts are written with [XML](https://www.w3.org/TR/NOTE-dcd]), specifying a graph of user-defined nodes with 
+dependencies between them. Each node consists of tasks executed in sequence to produce some sort of output (i.e., compile 
+this project, then cook, run this test, and so on). When asked to build a target (that is, a node or named output) BuildGraph 
+will execute all of the nodes in the graph required to make that happen.
+
+Unlike other build tools, BuildGraph is designed as a hybrid between a makefile-like scripting language and a build farm 
+configuration script. It allows annotations for the type of machine that steps are supposed to be executed on, providing a 
+list of recipients for failure notifications if a step fails, and groups nodes that should only be executed after an explicit 
+user trigger. It also tracks the creation of output files from task execution in a way that lends to graph execution being 
+distributed across a farm of machines (with nodes running in parallel where possible), and intermediate artifacts being 
+transferred to and from a network share automatically.
+
+Epic uses BuildGraph to prepare the UE4 binary release, package samples for the marketplace, and implement pipelines for our 
+own games (among other things). Several example BuildGraph scripts are provided in the 
+`[UE4Root]/Engine/Build/Graph/Examples` directory, and the script for creating a binary UE4 distribution can be found at 
+`[UE4Root]/Engine/Build/InstalledEngineBuild.xml`.
+
+Opening a script with Visual Studio will use the schema located at `[UE4Root]Engine/Build/Graph/Schema.xsd` to provide rich 
+tooltips, validation, and Intellisense features while editing.
 
 ## Writing BuildGraph Scripts
 
 [REGION:topics third]
-%Programming/Development/BuildGraph/BuildGraphScriptAnatomy:topiccompact%
+%Programming/Development/BuildGraph/ScriptAnatomy:topiccompact%
 [/REGION]
 
 ## Using BuildGraph
 
 [REGION:topics third]
-%Programming/Development/BuildGraph/BuildGraphUsage:topiccompact%
+%Programming/Development/BuildGraph/Usage:topiccompact%
 %Programming/Development/InstalledBuildReference:topiccompact%
-[/REGION]
-
-## BuildGraph Scripting Reference
-
-[REGION:topics third]
-%Programming/Development/BuildGraph/BuildGraphScriptTypes:topiccompact%
-%Programming/Development/BuildGraph/BuildGraphScriptElements:topiccompact%
-%Programming/Development/BuildGraph/BuildGraphScriptConditions:topiccompact%
-%Programming/Development/BuildGraph/BuildGraphScriptTasks:topiccompact%
 [/REGION]
\ No newline at end of file
diff --git a/Engine/Documentation/Source/Programming/Development/BuildGraph/Introduction/BuildGraphIntroduction.INT.udn b/Engine/Documentation/Source/Programming/Development/BuildGraph/Introduction/BuildGraphIntroduction.INT.udn
deleted file mode 100644
diff --git a/Engine/Documentation/Source/Programming/Development/BuildGraph/ScriptAnatomy/BuildGraphScriptAnatomy.INT.udn b/Engine/Documentation/Source/Programming/Development/BuildGraph/ScriptAnatomy/BuildGraphScriptAnatomy.INT.udn
--- a/Engine/Documentation/Source/Programming/Development/BuildGraph/ScriptAnatomy/BuildGraphScriptAnatomy.INT.udn
+++ b/Engine/Documentation/Source/Programming/Development/BuildGraph/ScriptAnatomy/BuildGraphScriptAnatomy.INT.udn
@@ -10,19 +10,19 @@ tags:BuildGraph
 [VAR:TopicCompact]
 [OBJECT:TopicCompact]
 	[PARAM:image]
-		![%Programming/Development/BuildGraph/BuildGraphScriptAnatomy%](anatomy_topic.png)
+		![%Programming/Development/BuildGraph/ScriptAnatomy%](anatomy_topic.png)
 	[/PARAM]
 	[PARAM:icon]
 		![](%ROOT%/reference_icon.png)(convert:false)
 	[/PARAM]
 	[PARAM:title]
-		%Programming/Development/BuildGraph/BuildGraphScriptAnatomy:title%
+		%Programming/Development/BuildGraph/ScriptAnatomy:title%
 	[/PARAM]
 	[PARAM:description]
-		%Programming/Development/BuildGraph/BuildGraphScriptAnatomy:description%
+		%Programming/Development/BuildGraph/ScriptAnatomy:description%
 	[/PARAM]
 	[PARAM:path]
-		[RELATIVE:Programming/Development/BuildGraph/BuildGraphScriptAnatomy]
+		[RELATIVE:Programming/Development/BuildGraph/ScriptAnatomy]
 	[/PARAM]
 [/OBJECT]
 [/VAR]
@@ -63,10 +63,11 @@ the allocation of machines (and coordination between them) is assumed to be mana
 
 The syntax of elements used to define BuildGraph constructs is listed in the following sections.
 
-[INCLUDE:Programming/Development/BuildGraph/BuildGraphScriptTypes#buildgraphscripttypes]
+## BuildGraph Scripting Reference
 
-[INCLUDE:Programming/Development/BuildGraph/BuildGraphScriptElements#buildgraphscriptelements]
-
-[INCLUDE:Programming/Development/BuildGraph/BuildGraphScriptConditions#buildgraphscriptconditions]
-
-[INCLUDE:Programming/Development/BuildGraph/BuildGraphScriptTasks#buildgraphscripttasks]
\ No newline at end of file
+[REGION:topics third]
+%Programming/Development/BuildGraph/ScriptAnatomy/Types:topiccompact%
+%Programming/Development/BuildGraph/ScriptAnatomy/Elements:topiccompact%
+%Programming/Development/BuildGraph/ScriptAnatomy/Conditions:topiccompact%
+%Programming/Development/BuildGraph/ScriptAnatomy/Tasks:topiccompact%
+[/REGION]
\ No newline at end of file
diff --git a/Engine/Documentation/Source/Programming/Development/BuildGraph/ScriptAnatomy/Conditions/BuildGraphScriptConditions.INT.udn b/Engine/Documentation/Source/Programming/Development/BuildGraph/ScriptAnatomy/Conditions/BuildGraphScriptConditions.INT.udn
--- a/Engine/Documentation/Source/Programming/Development/BuildGraph/ScriptAnatomy/Conditions/BuildGraphScriptConditions.INT.udn
+++ b/Engine/Documentation/Source/Programming/Development/BuildGraph/ScriptAnatomy/Conditions/BuildGraphScriptConditions.INT.udn
@@ -1,29 +1,29 @@
 Availability: Docs
 Title: BuildGraph Script Conditions
-Crumbs: %ROOT%, Programming, Programming/Development, Programming/Development/BuildGraph
+Crumbs: %ROOT%, Programming, Programming/Development, Programming/Development/BuildGraph, Programming/Development/BuildGraph/ScriptAnatomy
 Description:Learn the syntax needed to write BuildGraph script conditions.
 version: 4.13
 parent:Programming/Development/BuildGraph
 type:Reference
-related: Programming/Development/BuildGraph/BuildGraphScriptAnatomy
+related: Programming/Development/BuildGraph/ScriptAnatomy
 tags:BuildGraph
 
 [VAR:TopicCompact]
 [OBJECT:TopicCompact]
 	[PARAM:image]
-		![%Programming/Development/BuildGraph/BuildGraphScriptConditions%](conditional_topic.png)
+		![%Programming/Development/BuildGraph/ScriptAnatomy/Conditions%](conditional_topic.png)
 	[/PARAM]
 	[PARAM:icon]
 		![](%ROOT%/reference_icon.png)(convert:false)
 	[/PARAM]
 	[PARAM:title]
-		%Programming/Development/BuildGraph/BuildGraphScriptConditions:title%
+		%Programming/Development/BuildGraph/ScriptAnatomy/Conditions:title%
 	[/PARAM]
 	[PARAM:description]
-		%Programming/Development/BuildGraph/BuildGraphScriptConditions:description%
+		%Programming/Development/BuildGraph/ScriptAnatomy/Conditions:description%
 	[/PARAM]
 	[PARAM:path]
-		[RELATIVE:Programming/Development/BuildGraph/BuildGraphScriptConditions]
+		[RELATIVE:Programming/Development/BuildGraph/ScriptAnatomy/Conditions]
 	[/PARAM]
 [/OBJECT]
 [/VAR]
diff --git a/Engine/Documentation/Source/Programming/Development/BuildGraph/ScriptAnatomy/Elements/BuildGraphScriptElements.INT.udn b/Engine/Documentation/Source/Programming/Development/BuildGraph/ScriptAnatomy/Elements/BuildGraphScriptElements.INT.udn
--- a/Engine/Documentation/Source/Programming/Development/BuildGraph/ScriptAnatomy/Elements/BuildGraphScriptElements.INT.udn
+++ b/Engine/Documentation/Source/Programming/Development/BuildGraph/ScriptAnatomy/Elements/BuildGraphScriptElements.INT.udn
@@ -1,29 +1,29 @@
 Availability: Docs
 Title: BuildGraph Script Elements
-Crumbs: %ROOT%, Programming, Programming/Development, Programming/Development/BuildGraph
+Crumbs: %ROOT%, Programming, Programming/Development, Programming/Development/BuildGraph, Programming/Development/BuildGraph/ScriptAnatomy
 Description:Learn about the elements of a BuildGraph.
 version: 4.13
 parent:Programming/Development/BuildGraph
 type:Reference
-related: Programming/Development/BuildGraph/BuildGraphScriptAnatomy
+related: Programming/Development/BuildGraph/ScriptAnatomy
 tags:BuildGraph
 
 [VAR:TopicCompact]
 [OBJECT:TopicCompact]
 	[PARAM:image]
-		![%Programming/Development/BuildGraph/BuildGraphScriptElements%](elements_topic.png)
+		![%Programming/Development/BuildGraph/ScriptAnatomy/Elements%](elements_topic.png)
 	[/PARAM]
 	[PARAM:icon]
 		![](%ROOT%/reference_icon.png)(convert:false)
 	[/PARAM]
 	[PARAM:title]
-		%Programming/Development/BuildGraph/BuildGraphScriptElements:title%
+		%Programming/Development/BuildGraph/ScriptAnatomy/Elements:title%
 	[/PARAM]
 	[PARAM:description]
-		%Programming/Development/BuildGraph/BuildGraphScriptElements:description%
+		%Programming/Development/BuildGraph/ScriptAnatomy/Elements:description%
 	[/PARAM]
 	[PARAM:path]
-		[RELATIVE:Programming/Development/BuildGraph/BuildGraphScriptElements]
+		[RELATIVE:Programming/Development/BuildGraph/ScriptAnatomy/Elements]
 	[/PARAM]
 [/OBJECT]
 [/VAR]
diff --git a/Engine/Documentation/Source/Programming/Development/BuildGraph/ScriptAnatomy/Tasks/BuildGraphScriptTasks.INT.udn b/Engine/Documentation/Source/Programming/Development/BuildGraph/ScriptAnatomy/Tasks/BuildGraphScriptTasks.INT.udn
--- a/Engine/Documentation/Source/Programming/Development/BuildGraph/ScriptAnatomy/Tasks/BuildGraphScriptTasks.INT.udn
+++ b/Engine/Documentation/Source/Programming/Development/BuildGraph/ScriptAnatomy/Tasks/BuildGraphScriptTasks.INT.udn
@@ -1,29 +1,29 @@
 Availability: Docs
 Title: BuildGraph Script Tasks
-Crumbs: %ROOT%, Programming, Programming/Development, Programming/Development/BuildGraph
+Crumbs: %ROOT%, Programming, Programming/Development, Programming/Development/BuildGraph, Programming/Development/BuildGraph/ScriptAnatomy
 Description: Learn how BuildGraph can help you create custom tasks.
 version: 4.13
 parent:Programming/Development/BuildGraph
 type:Reference
-related: Programming/Development/BuildGraph/BuildGraphScriptAnatomy
+related: Programming/Development/BuildGraph/ScriptAnatomy
 tags:BuildGraph
 
 [VAR:TopicCompact]
 [OBJECT:TopicCompact]
 	[PARAM:image]
-		![%Programming/Development/BuildGraph/BuildGraphScriptTasks%](tasks_topic.png)
+		![%Programming/Development/BuildGraph/ScriptAnatomy/Tasks%](tasks_topic.png)
 	[/PARAM]
 	[PARAM:icon]
 		![](%ROOT%/reference_icon.png)(convert:false)
 	[/PARAM]
 	[PARAM:title]
-		%Programming/Development/BuildGraph/BuildGraphScriptTasks:title%
+		%Programming/Development/BuildGraph/ScriptAnatomy/Tasks:title%
 	[/PARAM]
 	[PARAM:description]
-		%Programming/Development/BuildGraph/BuildGraphScriptTasks:description%
+		%Programming/Development/BuildGraph/ScriptAnatomy/Tasks:description%
 	[/PARAM]
 	[PARAM:path]
-		[RELATIVE:Programming/Development/BuildGraph/BuildGraphScriptTasks]
+		[RELATIVE:Programming/Development/BuildGraph/ScriptAnatomy/Tasks]
 	[/PARAM]
 [/OBJECT]
 [/VAR]
@@ -53,6 +53,6 @@ If you'd like to have a starting point for creating tasks, we're providing a var
 If you're looking for a simple predefined task to start from, we recommend using `LogTask` as a good starting point. If you want to keep up with 
 our latest improvements to BuildGraph; bookmark this page, because the following list of predefined tasks will be updated periodically. 
 
-[INCLUDE:Programming/Development/BuildGraph/BuildGraphScriptTasks/BuildGraphScriptPredefinedTasks]
+[INCLUDE:Programming/Development/BuildGraph/ScriptAnatomy/Tasks/PredefinedTasks]
 
 [/EXCERPT:BuildGraphScriptTasks]
\ No newline at end of file
diff --git a/Engine/Documentation/Source/Programming/Development/BuildGraph/ScriptAnatomy/Tasks/PredefinedTasks/BuildGraphScriptPredefinedTasks.INT.udn b/Engine/Documentation/Source/Programming/Development/BuildGraph/ScriptAnatomy/Tasks/PredefinedTasks/BuildGraphScriptPredefinedTasks.INT.udn
--- a/Engine/Documentation/Source/Programming/Development/BuildGraph/ScriptAnatomy/Tasks/PredefinedTasks/BuildGraphScriptPredefinedTasks.INT.udn
+++ b/Engine/Documentation/Source/Programming/Development/BuildGraph/ScriptAnatomy/Tasks/PredefinedTasks/BuildGraphScriptPredefinedTasks.INT.udn
@@ -1,9 +1,9 @@
 Availability: NoPublish
 Title: BuildGraph Predefined Tasks
-Crumbs: %ROOT%, Programming, Programming/Development, Programming/Development/BuildGraph, Programming/Development/BuildGraph/BuildGraphScriptTasks
+Crumbs: %ROOT%, Programming, Programming/Development, Programming/Development/BuildGraph, Programming/Development/BuildGraph/ScriptAnatomy, Programming/Development/BuildGraph/ScriptAnatomy/Tasks
 Description: This is a procedurally generated markdown page.
 version: 4.13
-parent:Programming/Development/BuildGraph/BuildGraphScriptTasks
+parent:Programming/Development/BuildGraph/ScriptAnatomy/Tasks
 
 ### Command
 
diff --git a/Engine/Documentation/Source/Programming/Development/BuildGraph/ScriptAnatomy/Types/BuildGraphScriptTypes.INT.udn b/Engine/Documentation/Source/Programming/Development/BuildGraph/ScriptAnatomy/Types/BuildGraphScriptTypes.INT.udn
--- a/Engine/Documentation/Source/Programming/Development/BuildGraph/ScriptAnatomy/Types/BuildGraphScriptTypes.INT.udn
+++ b/Engine/Documentation/Source/Programming/Development/BuildGraph/ScriptAnatomy/Types/BuildGraphScriptTypes.INT.udn
@@ -1,29 +1,29 @@
 Availability: Docs
 Title: BuildGraph Script Types
-Crumbs: %ROOT%, Programming, Programming/Development, Programming/Development/BuildGraph
+Crumbs: %ROOT%, Programming, Programming/Development, Programming/Development/BuildGraph, Programming/Development/BuildGraph/ScriptAnatomy
 Description:Learn about valid data types for BuildGraph attributes.
 version: 4.13
 parent:Programming/Development/BuildGraph
 type:Reference
-related: Programming/Development/BuildGraph/BuildGraphScriptAnatomy
+related: Programming/Development/BuildGraph/ScriptAnatomy
 tags:BuildGraph
 
 [VAR:TopicCompact]
 [OBJECT:TopicCompact]
 	[PARAM:image]
-		![%Programming/Development/BuildGraph/BuildGraphScriptTypes%](types_topic.png)
+		![%Programming/Development/BuildGraph/ScriptAnatomy/Types%](types_topic.png)
 	[/PARAM]
 	[PARAM:icon]
 		![](%ROOT%/reference_icon.png)(convert:false)
 	[/PARAM]
 	[PARAM:title]
-		%Programming/Development/BuildGraph/BuildGraphScriptTypes:title%
+		%Programming/Development/BuildGraph/ScriptAnatomy/Types:title%
 	[/PARAM]
 	[PARAM:description]
-		%Programming/Development/BuildGraph/BuildGraphScriptTypes:description%
+		%Programming/Development/BuildGraph/ScriptAnatomy/Types:description%
 	[/PARAM]
 	[PARAM:path]
-		[RELATIVE:Programming/Development/BuildGraph/BuildGraphScriptTypes]
+		[RELATIVE:Programming/Development/BuildGraph/ScriptAnatomy/Types]
 	[/PARAM]
 [/OBJECT]
 [/VAR]
@@ -56,6 +56,6 @@ The following table contains valid data types that can be contained by BuildGrap
 | **File Spec**           | A set of file and tag names, separated by semi-colons. Wildcards such as `"..."`, `"*"`, and `"?"` are permitted when referencing files (i.e., `Engine/.../*.bat`). Unless otherwise specified, relative paths are resolved relative to the working root directory.                        |
 | **File Name**           | A path to a file. Unless otherwise specified, relative paths are resolved to the working root directory.                                                                                                                                                                                   |
 | **Directory Name**      | A path to a directory. Unless otherwise specified, relative paths are resolved relative to the working root directory.                                                                                                                                                                     |
-| **Condition**           | A [conditional expression](Programming/Development/BuildGraph/BuildGraphScriptConditions).                                                                                                                                                                                                      |
+| **Condition**           | A [conditional expression](Programming/Development/BuildGraph/ScriptAnatomy/Conditions).                                                                                                                                                                                                      |
 [/EXCERPT:TypesTable]
 -->
\ No newline at end of file
diff --git a/Engine/Documentation/Source/Programming/Development/BuildGraph/Usage/BuildGraphUsage.INT.udn b/Engine/Documentation/Source/Programming/Development/BuildGraph/Usage/BuildGraphUsage.INT.udn
--- a/Engine/Documentation/Source/Programming/Development/BuildGraph/Usage/BuildGraphUsage.INT.udn
+++ b/Engine/Documentation/Source/Programming/Development/BuildGraph/Usage/BuildGraphUsage.INT.udn
@@ -5,25 +5,25 @@ Description:This page shows you how to use BuildGraph Scripts.
 version: 4.13
 parent:Programming/Development/BuildGraph
 type:Overview
-related: Programming/Development/BuildGraph/BuildGraphScriptAnatomy
+related: Programming/Development/BuildGraph/ScriptAnatomy
 tags:BuildGraph
 
 [VAR:TopicCompact]
 [OBJECT:TopicCompact]
 	[PARAM:image]
-		![%Programming/Development/BuildGraph/BuildGraphUsage%](usage_topic.png)
+		![%Programming/Development/BuildGraph/Usage%](usage_topic.png)
 	[/PARAM]
 	[PARAM:icon]
 		![](%ROOT%/reference_icon.png)(convert:false)
 	[/PARAM]
 	[PARAM:title]
-		%Programming/Development/BuildGraph/BuildGraphUsage:title%
+		%Programming/Development/BuildGraph/Usage:title%
 	[/PARAM]
 	[PARAM:description]
-		%Programming/Development/BuildGraph/BuildGraphUsage:description%
+		%Programming/Development/BuildGraph/Usage:description%
 	[/PARAM]
 	[PARAM:path]
-		[RELATIVE:Programming/Development/BuildGraph/BuildGraphUsage]
+		[RELATIVE:Programming/Development/BuildGraph/Usage]
 	[/PARAM]
 [/OBJECT]
 [/VAR]
