Added related pages.

#UE4doc

diff --git a/Engine/Documentation/Source/Programming/Development/BuildGraph/ScriptAnatomy/BuildGraphScriptAnatomy.INT.udn b/Engine/Documentation/Source/Programming/Development/BuildGraph/ScriptAnatomy/BuildGraphScriptAnatomy.INT.udn
--- a/Engine/Documentation/Source/Programming/Development/BuildGraph/ScriptAnatomy/BuildGraphScriptAnatomy.INT.udn
+++ b/Engine/Documentation/Source/Programming/Development/BuildGraph/ScriptAnatomy/BuildGraphScriptAnatomy.INT.udn
@@ -6,6 +6,7 @@ version: 4.13
 parent:Programming/Development/BuildGraph
 type:Overview
 tags:BuildGraph
+related: Programming/Development/BuildGraph/Usage
 
 [VAR:TopicCompact]
 [OBJECT:TopicCompact]
diff --git a/Engine/Documentation/Source/Programming/Development/BuildGraph/Usage/BuildGraphUsage.INT.udn b/Engine/Documentation/Source/Programming/Development/BuildGraph/Usage/BuildGraphUsage.INT.udn
--- a/Engine/Documentation/Source/Programming/Development/BuildGraph/Usage/BuildGraphUsage.INT.udn
+++ b/Engine/Documentation/Source/Programming/Development/BuildGraph/Usage/BuildGraphUsage.INT.udn
@@ -6,6 +6,7 @@ version: 4.13
 parent:Programming/Development/BuildGraph
 type:Overview
 related: Programming/Development/BuildGraph/ScriptAnatomy
+related: Programming/Development/InstalledBuildReference
 tags:BuildGraph
 
 [VAR:TopicCompact]
