#UE4 Docs: Checking in the first pass of this new document.

diff --git a/Engine/Documentation/Source/Enterprise/Automotive/AutomotiveMaterialPack/AutomotiveMaterialPack.INT.udn b/Engine/Documentation/Source/Enterprise/Automotive/AutomotiveMaterialPack/AutomotiveMaterialPack.INT.udn
new file mode 100644
--- /dev/null
+++ b/Engine/Documentation/Source/Enterprise/Automotive/AutomotiveMaterialPack/AutomotiveMaterialPack.INT.udn
@@ -0,0 +1,297 @@
+Availability:Doc
+Title: Automotive Material Pack
+Description:A breakdown of what assets the Automotive Material pack contains and how to use them.
+Crumbs:%ROOT%,
+version:4.12
+Parent:Enterprise/Automotive
+tags:Enterprise
+type:overview
+order:1
+Related:Engine/Rendering/Materials/MaterialProperties/LightingModels
+Related:Engine/Rendering/Materials/HowTo/ClearCoatDualNormal
+Related:Engine/Content/Types/Textures/Cubemaps
+Related:Engine/UI/LevelEditor/Modes/MeshPaintMode/VertexColor
+
+[TOC(start:2 end:3)]
+
+![](AP_Intro.png)
+
+The Automotive Material Pack is a collection of automotive-themed PBR Materials and Textures that have been setup for use in Unreal Engine 4(UE4). The pack features 21 Master Materials and 75 Material Instances taken from Epic's 2016 GDC Mclaren 570S demo. In the following document, we will go over how to use the contents of the Automotive Material Pack in your UE4 projects as well as anything that you need to be aware of while using the content in your project.
+
+## What's in the Automotive Material Pack
+
+Inside of the Automotive Material Pack you will find the following Materials types:
+
+* Car Paint
+* Carbon Fiber
+* Caliper
+* Wheel
+* Brake Rotor
+* Exhaust
+* Reflector
+* Tire(Rubber)
+* Metal (Anodized/Dirt)
+* Leather (Perforated)
+* Alcantara
+* Fabric
+* Plastic
+* Glass
+
+[region:tip]
+If you are unsure what a certain parameter in the Material does, hovering your mouse over the parameter will display additional information.
+[/region]
+
+[region:note]
+The vehicle meshes used in the documentation below **are not** included with the asset pack.
+[/region]
+
+## Standard Material Parameters
+
+In the following section, we will take a look at some of the parameters that all of the Materials share, what those parameters do and how you can use them.
+
+* **Material Scale:** Material Scale gives you control over the Scaling / Tiling of a Texture or pattern that is used in the Material.
+
+* **Scale-Aware Tiling:** Scale- Aware Tiling scales the Materials UV's up or down based on the World Size of the Static Mesh that it is applied to. The result is a uniform Material scale on all of the objects regardless of their current size. Please note that this feature will only work with objects which have UV's channels set up.
+
+	[OBJECT:ComparisonSlider]
+		[PARAM:before]
+		![Scale Aware Tiling](AP_02_SMP_01_Scale-Aware.png)
+		[/PARAM]
+		[PARAM:after]
+		![Normal Tiling](AP_02_SMP_02_NormalUV.png)
+		[/PARAM]
+	[/OBJECT]
+
+* **World Aligned:** This feature should be used on objects that have incorrect or poorly setup UV's. Please note that this is best used on Static objects as moving objects will exhibit texture 'swimming'.
+
+	[OBJECT:ComparisonSlider]
+		[PARAM:before]
+		![World Aligned OFF and no UVs](AP_02_SMP_03_WA_OFF.png)
+		[/PARAM]
+		[PARAM:after]
+		![World Aligned ON and no UVs](AP_02_SMP_04_WA_ON.png)
+		[/PARAM]
+	[/OBJECT]
+
+## Orange Peel Material Parameter
+
+The clear coat lacquer applied on top of automotive paint often has a slightly bumpy surface - not unlike the surface of an orange and is the reason why this is referred to as Orange Peel. You can control the intensity of this effect by using the **Orange Peel Strength** and **Orange Peel Scale** parameters.To disable this effect, set the **Orange Peel Strength** to **0**.
+
+[OBJECT:ComparisonSlider]
+	[PARAM:before]
+	![Orange Peel Off](AP_03_Carpaint_03_OrangPeel_OFF.png)
+	[/PARAM]
+	[PARAM:after]
+	![Orange Peel On](AP_03_Carpaint_04_OrangPeel_ON.png)
+	[/PARAM]
+[/OBJECT]
+
+## Car Paint Material
+
+The Car Paint Material is a complex Material that comes with 29 adjustable parameters. When trying to modify this Material for your project, we recommend starting with the following parameters:
+
+* **Base Color/Flake**
+* **Scale/Flake**
+* **Normal Scale/Flake**
+* **Normal Strength/Parameters**
+ 
+By adjusting just these four different parameters you can achieve vastly different looks for your car paint Material. Here are two examples of what can be achieved by changing those four parameters:
+
+[OBJECT:ComparisonSlider]
+	[PARAM:before]
+	![Car Paint 01 ](AP_03_Carpaint_01.png)
+	[/PARAM]
+	[PARAM:after]
+	![Car Paint 02](AP_03_Carpaint_02.png)
+	[/PARAM]
+[/OBJECT]
+
+## Carbon Fiber Material
+
+The Carbon Fiber Material makes use of the new Dual Normals option that is in the Clearcoat Material. This new Dual Normal allows you to add a second Normal map that allows the bottom layer of the Clearcoat Material to interact with light. Without this feature, the glossy effect of the Carbon Fiber Material is lost, making the Carbon Fiber look not as realistic as it could be. 
+
+![](AP_04_CarbonFiber.png)
+
+## Caliper Material
+
+The caliper is part of a car's braking system and is usually placed on the brake disc in between the wheel and the brake pad. 
+
+![](AP_05_Caliper_01.png)
+
+
+The brake caliper Material has the following two different options. 
+
+* **Painted Caliper:** Plain color.
+* **Painted Caliper with Metallic Flakes:** Silver color with metallic flakes. 
+
+[OBJECT:ComparisonSlider]
+	[PARAM:before]
+	![Painted Caliper](AP_05_Caliper_02_Painted.png)
+	[/PARAM]
+	[PARAM:after]
+	![Painted Caliper with Metallic Flakes](AP_05_Caliper_03_Flakes.png)
+	[/PARAM]
+[/OBJECT]
+
+## Rim Material
+
+There are five different materials that can be used for the rims of the wheels including a Two Tone option. When using the Two Tone Rim keep in mind that the Two Tone Rim is made up of two different Materials with one Material being blue-black and the other material is silver. The **Red Channel** of the Vertex Color on the Rim Static Mesh is used to define which areas of the wheel should be silver. Any areas that are not using red Vertex color should use black as the Vertex Color so that they will display the 2nd Material.
+
+![](AP_06_Rim_Ball_01.png)
+
+[region:note]
+If you are unfamiliar with how to use Vertex Paint in UE4 please check out the [Mesh Paint](Engine/UI/LevelEditor/Modes/MeshPaintMode/VertexColor/) user guide.
+[/region]
+
+[REGION:fullwidth]
+[INCLUDE:Enterprise\Automotive\AutomotiveMaterialPack\#RM]
+[/REGION]
+
+## Brake Rotor Material
+
+Since the brake rotor, Material uses a radial pattern Texture it needs to be applied to a mesh using a planar UV projection method or it will not be displayed correctly. The image below shows an example of how the UVs should look like when a Texture is applied to the brake rotor.
+
+![](AP_07_BrakeRotor_01.png)
+
+The following images show examples of the types of brake rotor Materials that are available to use.
+[OBJECT:ComparisonSlider]
+	[PARAM:before]
+	![Carbon Ceramic](AP_07_BrakeRotor_02_CarbonCeramic.png)
+	[/PARAM]
+	[PARAM:after]
+	![Iron](AP_07_BrakeRotor_03_Iron.png)
+	[/PARAM]
+[/OBJECT]
+
+## Reflector Material
+
+Inside the Automotive Material Pack you will find the following three Reflector Patterns presets that are ready to be used;
+
+* **Hexagon**
+* **Honeycomb**
+* **Triangle**
+
+	![](AP_08_Reflector_02.png)
+
+To add or change the reflector patterns open up the **Reflector** Material Instance and change the **Reflector Normal** to the reflector pattern you want to use. You can also create your own Reflector Normal Maps to use by doing the following.
+
+1. Generate a Normal Map of the item you want to use a reflector.
+1. Save that image as a .TGA and then import it into UE4.
+1. Once imported, select the Texture in the Content Browser and then inside of the Reflector Material Instance under the Reflector (Bottom Layer) Normal section press the **white arrow** to load the selected Texture.
+
+	![](AP_08_Reflector_01.jpg)
+
+## Rubber / Tire Material
+
+The rubber Material that comes with the Automotive Material Pack can be used for any general rubber surface you might find in a car. For example, the following surfaces would make good use of the rubber Material.
+
+* **Door Seals**
+* **Floor Mats**
+* **Interior Pieces**
+* Tires**
+
+The rubber Material has a tire-specific option called **Use Sidewall**. Enabling this part of the Material will add a Normal Map Texture that will display the sidewall lettering that can be found on tires. When using the sidewall option for a tire make sure that the UV for the tire sidewall 3D Mesh lines up with the Normal Map.
+
+![](AP_Tire_UV_00.png)
+
+If you fail to properly line up the UV of the side wall to the Normal of the tire things like the lettering and logo on the sidewall will not show up correctly. In the following images 
+
+[OBJECT:ComparisonSlider]
+	[PARAM:before]
+	![Correct Way](AP_09_Rubber_04_Right.png)
+	[/PARAM]
+	[PARAM:after]
+	![Wrong Way](AP_09_Rubber_04_Wrong.png)
+	[/PARAM]
+[/OBJECT]
+
+## Leather / Perforated Material
+
+The leather Material that comes with the Automotive Material Pack has 12 different color leather preset to choose from as well as five different perforated leather versions. 
+
+[OBJECT:ComparisonSlider]
+	[PARAM:before]
+	![Perforated](AP_10_Leather_01_Perforated.png)
+	[/PARAM]
+	[PARAM:after]
+	![Leather](AP_10_Leather_02_Leather.png)
+	[/PARAM]
+[/OBJECT]
+
+## Alcantara Material
+
+The alcantara Material that comes with the Automotive Material Pack can be used to mainpulate of fabric. When using the Alcantara Material you can adjust the following parameters to mainpulate the look of the fabric.
+
+* **Fabric Scale / Strength**
+* **Cloth Folds Scale / Strength**
+* **Fabric Detail Scale / Strength**
+
+In the following set of images, you can see the effect that the Cloth Folds parameter has on the look of the Alcantara fabric Material.
+
+[OBJECT:ComparisonSlider]
+	[PARAM:before]
+	![ Cloth Folds Strength: 0](AP_11_Alcantara_01_00.png)
+	[/PARAM]
+	[PARAM:after]
+	![ Cloth Folds Strength: 10.0](AP_11_Alcantara_02_10.png)
+	[/PARAM]
+[/OBJECT]
+
+## Glass Material
+
+The glass Material that comes with the Automotive Material Pack has the following two reflections types that allow the glass to use two different reflections so the Material can reflect objects that are outside of the car as well as reflect objects the inside of the car.
+
+* **Outer Reflection:** Helps to reflect the scene around the exterior of the glass in real time. 
+
+	![](AP_12_Glass_06_Ex.png)
+
+* **Inner Reflection:** Helps to reflect the scene around the interor of the glass using a predefined Cubemap.
+
+	![](AP_12_Glass_07_Ex.png)
+
+
+Here is an example of what the glass Material looks like with the inner reflections enabled and without the inner reflection enabled. 
+
+[OBJECT:ComparisonSlider]
+	[PARAM:before]
+	![Without Inner Reflection](AP_12_Glass_04_WithoutInnerReflec.png)
+	[/PARAM]
+	[PARAM:after]
+	![With Inner Reflection](AP_12_Glass_05_WithInnerReflec.png)
+	[/PARAM]
+[/OBJECT]
+
+For the car windshield, it would be best to capture a Cubemap from the inside of the car and then input that Cubemap into the glass's inner reflection Cubemap like in the following image.
+
+![](AP_12_Glass_01.png)
+
+[region:note]
+If you are unfamiliar with how to capture and use Cubemap in UE4 please check out the [Cubemap user guide](Engine/Content/Types/Textures/Cubemaps) for more information on how to capture and use Cubemaps in UE4.
+[/region]
+
+
+<!-----
+[EXCERPT:RM]
+[OBJECT:SliderSequence]
+	[PARAMLITERAL:max]
+	5
+	[/PARAMLITERAL]
+	[PARAMLITERAL:ext]
+	jpg
+	[/PARAMLITERAL]
+	[PARAMLITERAL:id]
+	RM
+	[/PARAMLITERAL]
+	[PARAM:caption]
+	**1:Two Tone 2:Brushed 3:Grey 4:Blue Black 5:Silver**
+	[/PARAM]
+	[PARAM:image_list]
+	![](RM_1.png)
+	![](RM_2.png)
+	![](RM_3.png)
+	![](RM_4.png)
+	![](RM_5.png)
+	[/PARAM]
+[/OBJECT]
+[/EXCERPT:RM]
+-->
