Accepting RichardH's peer review changes.

#UE4doc

diff --git a/Engine/Documentation/Source/Programming/Development/BuildGraph/BuildGraph.INT.udn b/Engine/Documentation/Source/Programming/Development/BuildGraph/BuildGraph.INT.udn
--- a/Engine/Documentation/Source/Programming/Development/BuildGraph/BuildGraph.INT.udn
+++ b/Engine/Documentation/Source/Programming/Development/BuildGraph/BuildGraph.INT.udn
@@ -7,11 +7,11 @@ type:landing
 parent:Programming/Development
 tags:BuildGraph
 
-**BuildGraph** is a system that runs BuildGraph scripts for build automation (with building blocks common to Unreal Engine 4 (UE4) projects). 
-BuildGraph integrates with UnrealBuildTool, AutomationTool, and the editor. Which means that BuildGraph can be extended and customized for your projects.
+**BuildGraph** is a script-based build automation system that features graphs of building blocks common to Unreal Engine 4 (UE4) projects. 
+BuildGraph integrates with UnrealBuildTool, AutomationTool, and the editor, and can be extended and customized for your projects.
 
 BuildGraph scripts are written with [XML](https://www.w3.org/TR/NOTE-dcd]), specifying a graph of user-defined nodes with 
-dependencies between them. Each node consists of tasks executed in sequence to produce some sort of output (i.e., compile 
+dependencies between them. Each node consists of tasks executed in sequence to produce some sort of output (e.g., compile 
 this project, then cook, run this test, and so on). When asked to build a target (that is, a node or named output) BuildGraph 
 will execute all of the nodes in the graph required to make that happen.
 
