Addressed MatthewG's SME review comments.

#UE4doc

diff --git a/Engine/Documentation/Source/Programming/Development/InstalledBuildReference/InstalledBuildReference.INT.udn b/Engine/Documentation/Source/Programming/Development/InstalledBuildReference/InstalledBuildReference.INT.udn
--- a/Engine/Documentation/Source/Programming/Development/InstalledBuildReference/InstalledBuildReference.INT.udn
+++ b/Engine/Documentation/Source/Programming/Development/InstalledBuildReference/InstalledBuildReference.INT.udn
@@ -5,11 +5,12 @@ Description:This page describes the Installed Build process, including an overvi
 version: 4.13
 parent:Programming/Development
 type:Reference
-related:Programming/Development/BuildConfigurations
-related:Programming/UnrealBuildSystem
 related:Programming/Development/BuildGraph
+related:Programming/UnrealBuildSystem
+related:Programming/Development/BuildConfigurations
+tags: BuildGraph
 
-[TOC (start:1 end:4)]
+[TOC (start:1 end:5)]
 
 [VAR:TopicCompact]
 [OBJECT:TopicCompact]
@@ -38,7 +39,7 @@ about the Installed Build process, including how to run scripts that make Instal
 
 ## What's an Installed Build?
 
-An **Installed Build** of UE4 is a version that's distributed using the [Epic Games Launcher](GettingStarted/Installation).
+An **Installed Build** of UE4 is the same as a version that's distributed using the [Epic Games Launcher](GettingStarted/Installation).
 
 ![](InstalledBuilds.png)
 
@@ -46,11 +47,15 @@ This version of the engine is shipped with each release of UE4. Although the Ins
 Installed Build process that quickly gets users up and running with the engine. Additionally, the Unreal Editor is distributed as 
 precompiled binaries with a static library for each target platform in development and shipping configurations.
 
+[REGION:note]
+Contents found in the `Engine` folder should be treated as `read-only` because rebuilding game projects will not attempt to rebuild Engine code.
+[/REGION] 
+
 ## About the Installed Build Process
 
 The Installed Build process automates build packages and tools that are needed to make an Installed Build version of UE4. 
 In summary, making an Installed Build is the process of setting up default target platforms, compiling UE4 with its tools and 
-editor, generating embedded documentation, running tests, and deploying the Installed Build version of UE4 for targeted platforms.
+editor, running tests, and deploying the Installed Build version of UE4 for targeted platforms.
 
 ## Making an Installed Build
 
@@ -64,11 +69,6 @@ If you're interested in reading through example BuildGraph scripts, navigate to
 * `Building.xml`
 * `TagsAndFiles.xml`
 
-<!--TO DO: Add cross-link to pending BuildGraph Documentation
-If you're interested in learning more about our BuildGraph Scripting System, head over to our BuildGraph Scripting System reference to learn how you can 
-customize your own scripts.
--->
-
 ### Running the Installed Build Script
 
 The Installed Build script is located at `[UE4Root]/Engine/Build/InstalledEngineBuild.xml`. You can run the Installed Build script by invoking the 
@@ -163,12 +163,17 @@ can still be run from a single build machine.
 
 Depending on your targeted platform and available Software Development Kit (SDK), we're providing some useful commands that you should take note of.
 
-#### IOS and TVOS
+#### Making an Installed Build for your Current Platform
+
+If you want to make an Installed Build for your current platform, thereby avoiding the need to set all of the platform options to `false`, append 
+`-set:HostPlatformOnly=true` to the command line.
+
+##### IOS and TVOS
 
 If you want to include IOS and TVOS in your Installed Build, you'll have to have a Mac set up for remote building. Otherwise, append the following to your command line:
 `-set:WithIOS=false -setWithTVOS=false`.
 
-#### Playstation 4 and XboxOne
+##### Playstation 4 and XboxOne
 
 If you have the required SDKs, you can add Playstation 4 (PS4) and XboxOne as target platforms (although, this is untested). To add PS4 and XboxOne 
 as target platforms, append `-set:WithPS4=true -set:WithXboxOne=true` to the command line.
@@ -178,10 +183,9 @@ as target platforms, append `-set:WithPS4=true -set:WithXboxOne=true` to the com
 Building a stand-alone [Derived Data Cache (DDC)](Engine/Basics/DerivedDataCache) for the Engine and Template content can be one of the slowest aspects of the build. 
 If you don't need a stand-alone DDC, you can skip this step by appending `-set:WithDDC=false` to the command line.
 
-#### Not Signing Executables
+#### Signing Executables
 
-Typically, you don't have to set up your machine to sign executables, as this can cause the Install Build process to fail while attempting to sign the executables being built. 
-To avoid this problem, append `-NoSign` to the command line.
+If you have to set up your machine to sign executables as part of the build process, append `-set:SignExecutables=true` to the command line.
 
 ### Including Additional Files
 
