#UE4 Docs: table formatting change

diff --git a/Engine/Documentation/Source/Platforms/Linux/BeginnerLinuxDeveloper/SupportedGfxDrivers/SupportedGfxDrivers.INT.udn b/Engine/Documentation/Source/Platforms/Linux/BeginnerLinuxDeveloper/SupportedGfxDrivers/SupportedGfxDrivers.INT.udn
--- a/Engine/Documentation/Source/Platforms/Linux/BeginnerLinuxDeveloper/SupportedGfxDrivers/SupportedGfxDrivers.INT.udn
+++ b/Engine/Documentation/Source/Platforms/Linux/BeginnerLinuxDeveloper/SupportedGfxDrivers/SupportedGfxDrivers.INT.udn
@@ -36,7 +36,13 @@ Unreal Engine requires compatible hardware and software support in order to run
 machine. In an effort to help you determine your ideal development environment, we're providing 
 the following table of supported graphics drivers for developing games on Linux.
 
-[INCLUDE:#SupportedGfxDrivers]
+[EXCERPT:SupportedGfxDrivers]
+| **Manufacturer** | **Supported Graphics Driver Versions**    | **Suggested Graphics Driver Version**     |
+| ---------------- | ----------------------------------------- | ----------------------------------------- |
+| **NVIDIA**       | Between 347.09 and 361.43                 | 361.75                                    |
+| **AMD**          | 14.301.1001.0 or newer                    | 15.30.1025.1001                           |
+| **Intel**        | Latest available version for your system. | Latest available version for your system. |
+[/EXCERPT:SupportedGfxDrivers]
 
 Graphics drivers are typically provided by manufacturers for testing and integration with Unreal
 Engine 4 (UE4). If you want to get the most out of UE4 while maintaining hardware performance, please 
@@ -57,13 +63,3 @@ your system from crashing while running the engine.
 ## Intel Graphics Drivers
 
 If your system has Mesa drivers, we recommend that you update your driver to be newer than Mesa 10.5.x.
-
-<!--
-[EXCERPT:SupportedGfxDrivers]
-| **Manufacturer** | **Supported Graphics Driver Versions**    | **Suggested Graphics Driver Version**     |
-| ---------------- | ----------------------------------------- | ----------------------------------------- |
-| **NVIDIA**       | Between 347.09 and 361.43                 | 361.75                                    |
-| **AMD**          | 14.301.1001.0 or newer                    | 15.30.1025.1001                           |
-| **Intel**        | Latest available version for your system. | Latest available version for your system. |
-[/EXCERPT:SupportedGfxDrivers]
--->
\ No newline at end of file
