Adding VR Foliage Controls and note on main foliage page -UEDOC-3209

diff --git a/Engine/Documentation/Source/Engine/Editor/VR/Controls/VREditorControls.INT.udn b/Engine/Documentation/Source/Engine/Editor/VR/Controls/VREditorControls.INT.udn
--- a/Engine/Documentation/Source/Engine/Editor/VR/Controls/VREditorControls.INT.udn
+++ b/Engine/Documentation/Source/Engine/Editor/VR/Controls/VREditorControls.INT.udn
@@ -115,7 +115,7 @@ Each Controller emits a beam called the laser pointer which is used to determine
 | ![Aim Controller at Window Grip + Squeeze Trigger + Track Pad Towards or Away](scale_window.png)(convert:false) | Aim the controller at the Grip of a window and squeeze the Trigger to grab the window. Move your thumb up or down on the Track Pad to scale the window up or down. |
 [/REGION]
 
-## Mesh Paint Controls
+## Mesh Painting Controls
 
 [REGION:raw]
 | Control | Description |
@@ -124,6 +124,16 @@ Each Controller emits a beam called the laser pointer which is used to determine
 | ![Modifier + Aim Controller at Actor + Squeeze Trigger + Move Controller](ModifierMeshPaint.png)(convert:false) | To erase instead of paint, hold the modifier button. Aim the controller at a Static Mesh in the world and squeeze the Trigger to begin erasing. Move the controller to move the eraser along the Mesh.  |
 [/REGION]
 
+## Foliage Painting Controls
+
+For both painting and erasing foliage, pressure sensitivity on the Trigger is supported. 
+
+[REGION:raw]
+| Control | Description |
+| ------- | ----------- |
+| ![Aim Controller at Actor + Squeeze Trigger + Move Controller](MeshPaint_2_NoModifier.png)(convert:false) | Select a foliage type in the **Modes** panel. Aim the controller at a location in the world and squeeze the Trigger to begin painting foliage. Move the controller to move the brush around the world.|
+| ![Modifier + Aim Controller at Actor + Squeeze Trigger + Move Controller](ModifierMeshPaint.png)(convert:false) | To erase foliage, hold the modifier button. Aim the controller at a location in the world and squeeze the Trigger to begin erasing. Move the controller to move the eraser around the world.  |
+[/REGION]
 
 <!--
 [EXCERPT:MultiSelect]
diff --git a/Engine/Documentation/Source/Engine/Foliage/Foliage.INT.udn b/Engine/Documentation/Source/Engine/Foliage/Foliage.INT.udn
--- a/Engine/Documentation/Source/Engine/Foliage/Foliage.INT.udn
+++ b/Engine/Documentation/Source/Engine/Foliage/Foliage.INT.udn
@@ -186,6 +186,11 @@ When painting meshes, the engine performs a number of line traces inside the sph
 **To erase foliage from the brush area:**  
 * Hold down **Shift** and **Left Mouse** button and **click** where you want to remove foliage.
 
+
+[REGION:tip]
+Foliage painting also works in the VR Editor. For the controls for foliage painting in VR, see the [](Engine/Editor/VR/Controls).
+[/REGION]
+
 ### Paint Settings
 
 Each Static mesh in the Mesh List has a number of properties that control how instances are placed, scaled or rotated when the selected Static mesh is painted. 
