Improved based on SME review.

diff --git a/Engine/Documentation/Source/Engine/Replay/Replay.INT.udn b/Engine/Documentation/Source/Engine/Replay/Replay.INT.udn
--- a/Engine/Documentation/Source/Engine/Replay/Replay.INT.udn
+++ b/Engine/Documentation/Source/Engine/Replay/Replay.INT.udn
@@ -50,19 +50,13 @@ order:17
 [/REGION]
 
 ## Replay System Overview
-Unreal Engine 4 features a Replay system which can record gameplay for later viewing. At a high level, the Replay system works by reading data drawn from the built-in replication system, similar to how live gameplay operates in a networked environment. The data is then passed to a Replay Streamer, which handles the process of filtering and storing the data. There are three types of Streamers included with the engine:
+Unreal Engine 4 features a Replay system which can record gameplay for later viewing. At a high level, the Replay system works by reading data drawn from the built-in replication system, similar to how live gameplay operates in a networked environment. Any project that is set up to replicate data like a multiplayer game does, even if the project doesn't actually have a multiplayer mode, will be compatible with the Replay system without further modification. The way this works is by passing replicated data to a Replay Streamer, which handles the process of filtering and storing the data. When the Replay system attempts to show a saved replay, it will have all of the replicated information that was available during live play, as well as special data fields designated as relevant only to replays, and can reconstruct the events of the game from this data.
 
-### The NULL Streamer
-The NULL Streamer is the default way to record a replay. This streamer writes replay data to disk, e.g. a local hard drive. This is good for making local recordings, especially of single-player games. These recordings can be very helpful for a variety of uses, such as producing gameplay trailers or in-game cutscenes, or enabling your users to view and share speedrun or tutorial videos within your game.
+## Replay System Functionality
+The basic functions of the replay system are to start or stop recording a demo, or play back a previously-recorded demo. While in playback mode, the Replay system supports commands to pause, change the playback speed, or jump to a specific point in time. More advanced usage of the system can include adding text tags (usually meta tags or the names of players involved in the game) to a replay, or querying existing replays based on text tags or version information present. The Replay system can be used from C++ code, mainly via the `UGameInstance` and `UWorld` classes, or with console commands or command-line arguments. A unified C++/Blueprint API is under construction and will be released with a future engine version.
 
-### The Memory Streamer
-The Memory Streamer keeps a user-configurable length of replay data (e.g. 10 seconds, 3 minutes, etc.) in memory on the local machine. This type of stream is best suited to instant replays of recent, dramatic moments, such as scores or penalty calls in a sports game, deaths in a shooter, or the final moments of a boss battle in an action game.
+## Replay Data Format
+In terms of data, a replay contains three types of game-state information, as well as some additional text data. At the start is baseline data describing the starting state of the game world. Checkpoints that act as snapshots of the net changes to the world (from the baseline) appear at regular, user-defined intervals. The space between checkpoints is then filled with incremental changes to individual objects in the world. Any moment in the game can be reconstructed by the engine quickly and accurately by initializing the world to the starting state, making the changes described in the last checkpoint before the chosen time, and then applying each incremental change after that checkpoint leading up to the desired point in time. The text data contained in a replay consists of a display name, which can be used when making a player-facing list, and user-defined text tags (HTTP Streamer only), which can be used when searching or filtering through lists of games.
 
-### The HTTP Streamer
-The HTTP Streamer sends replay data to another server, which could be on a LAN or elsewhere on the internet. This is useful for live-streaming matches, or for keeping recordings of matches that can be viewed at any time. This streamer is especially useful for dedicated-server games, where only the server knows everything that's happening everywhere in the game at all times, and where offloading the work of processing replay data will increase the number of simultaneous games that can be hosted on a single server. It can also help to serve as a moderation or cheat detection tool, since the data can be captured from computers completely controlled by the party running the game.
-
-## Replay Data
-In terms of data, a replay contains three types of game-state information, as well as some additional text data. At the start is baseline data describing the starting state of the game world. Checkpoints that act as snapshots of the net changes to the world (from the baseline) appear at regular, user-defined intervals. The space between checkpoints is then filled with incremental changes to individual objects in the world. Any moment in the game can be reconstructed quickly and accurately by initializing the world to the starting state, making the changes described in the last checkpoint before the chosen time, and then applying each incremental change after that checkpoint leading up to the desired point in time. The text data contained in a replay consists of a display name, which can be used when making a player-facing list, and user-defined text tags (HTTP Streamer only), which can be used when searching or filtering through lists of games.
-
-## Replay System Usage
-The basic functions of the replay system are to start or stop recording a demo, or play back a previously-recorded demo. While in playback mode, the Replay system supports commands to pause, change the playback speed, or jump to a specific point in time. More advanced usage of the system can include adding text tags (usually meta tags or the names of players involved in the game) to a replay, or querying existing replays based on text tags or version information present. The Replay system can be used from C++ code, mainly via the `UGameInstance` and `UWorld` classes, or with console or command-line commands. Note that some of the Replay system's streamers are not intended to be used in PIE or on non-authoritative machines such as clients in networked games.
+## Streamers
+The Replay system's Streamers perform the actual extraction and recording of the information needed to create the replay. Depending on how the replay data is intended to be viewed, different Streamers can be used. There are three different Streamers, each suited to different use cases. The NULL Streamer, which is the default, records events from the host machine directly to disk, making it best suited for single-player games and games that keep replays locally, on the host player's own machine. The Memory Streamer runs on a client machine and stores data in memory, making it ideal for "instant replay" features in sports titles or "kill cams" in shooters. Finally, the HTTP Streamer is used to send replay data to a second machine over LAN or the internet, which makes it a great choice for dedicated-server games and for games that may need to be streamed live to a large number of viewers while still remaining responsive to players.
