Adding fixes from peer review feedback UEDOC-3163

diff --git a/Engine/Documentation/Source/Engine/Editor/VR/ActivateVRMode/ActivateVRMode.INT.udn b/Engine/Documentation/Source/Engine/Editor/VR/ActivateVRMode/ActivateVRMode.INT.udn
--- a/Engine/Documentation/Source/Engine/Editor/VR/ActivateVRMode/ActivateVRMode.INT.udn
+++ b/Engine/Documentation/Source/Engine/Editor/VR/ActivateVRMode/ActivateVRMode.INT.udn
@@ -7,7 +7,7 @@ type: howto
 order: 2
 tags: VR
 tags: Getting Started
-version: 4.11
+version: 4.13
 skilllevel: intermediate
 
 [VAR:Topic]
@@ -44,6 +44,8 @@ skilllevel: intermediate
 [/OBJECT]
 [/VAR]
 
+As of Unreal Engine 4.12, VR Editing Mode is included with the version of the engine available through the Epic Games Launcher. You can also get the most up-to-date build from [GitHub](Engine/Editor/VR/Setup).
+
 [EXCERPT:main]
 
 1. After the editor loads, click **File > New** and open up our new "VR-Basic" level.  This environment is a good playground for you to learn the controls.
@@ -69,6 +71,19 @@ At any time you can press the Home button to return to this menu and then select
 
 As of 4.13, you can also enter VR Editing Mode automatically.  As long as the editor is in the foreground and VR Editing mode is enabled, when you wear the HMD, you will automatically enter VR editing mode. When you remove the HMD, you will leave it. 
 
+[OBJECT:EmbeddedVideo]
+[PARAMLITERAL:width]
+640
+[/PARAMLITERAL]
+[PARAMLITERAL:height]
+360
+[/PARAMLITERAL]
+[PARAMLITERAL:videoid]
+cZJmh881FCo
+[/PARAMLITERAL]
+[/OBJECT]
+
+
 There is a setting under **VR** in the **Experimental** section of **Editor Preferences** that will allow you to turn off auto-entry if you prefer.
 
 ![](AutoEntryMenu_Highlight.png)
