Fixing table includes.

#UE4doc

diff --git a/Engine/Documentation/Source/Programming/Development/BuildGraph/ScriptAnatomy/Elements/BuildGraphScriptElements.INT.udn b/Engine/Documentation/Source/Programming/Development/BuildGraph/ScriptAnatomy/Elements/BuildGraphScriptElements.INT.udn
--- a/Engine/Documentation/Source/Programming/Development/BuildGraph/ScriptAnatomy/Elements/BuildGraphScriptElements.INT.udn
+++ b/Engine/Documentation/Source/Programming/Development/BuildGraph/ScriptAnatomy/Elements/BuildGraphScriptElements.INT.udn
@@ -57,27 +57,45 @@
 `<Node>` is the smallest unit of execution in BuildGraph, having a set of inputs and outputs. Each `<Node>` consists of a 
 sequence of tasks that are executed in order.
 
-[INCLUDE:#NodeTable]
+|                  |              |              |                                                                                                                                             |
+| ---------------- | ------------ | ------------ | ------------------------------------------------------------------------------------------------------------------------------------------- |
+| Name             | Name         | Required     | Name of the node.                                                                                                                           |
+| Requires         | Target List  | Optional     | List of nodes, aggregates, or tagged file sets produced by other nodes that this node requires to execute, separated by semicolons.         |
+| Produces         | Tag List     | Optional     | Tagged file sets that this node makes available to other nodes, separated by semicolons.                                                    |
+| After            | Target List  | Optional     | List of nodes that this node should run after, if they are part of the current target (ignored if they are not). Separated by semicolons.   |
+| NotifyOnWarnings | Boolean      | Optional     | If false, this node will not produce notifications on warnings. Only used by build systems. Defaults to `true`. `[META]`                  |
+| If               | Condition    | Optional     | Condition to be evaluated. The element is ignored unless the condition evaluates to `true`.                                                |
 
 ### Aggregate
 
 `<Aggregate>` declares a named aggregate that can be used as a synonym for a set of other nodes (or produced tag sets).
 
-[INCLUDE:#AggregateTable]
+|                  |              |              |                                                                                              |
+| ---------------- | ------------ | ------------ | -------------------------------------------------------------------------------------------- |
+| Name             | Name         | Required     | Name of the aggregate.                                                                       |
+| Requires         | Target List  | Required     | List of dependencies for this aggregate. May be nodes, tagged file sets, or agent groups.    |
+| If               | Condition    | Optional     | Condition to be evaluated. The element is ignored unless the condition evaluates to `true`. |
 
 ### Agent
 
 `<Agent>` defines requirements for an agent on nodes that will be executed in sequence (without cleaning intermediate directories). 
 `<Agent>` requirements are ignored when building locally, but must be specified.
 
-[INCLUDE:#AgentTable]
+|                  |                 |              |                                                                                                                                             |
+| ---------------- | --------------- | ------------ | ------------------------------------------------------------------------------------------------------------------------------------------  |
+| Name             | Name            | Required     | Name of the group.                                                                                                                          |
+| Type             | Identifier List | Optional     | Type of agent to run on. The meaning of this string is inferred by the host build system; it does not have any intrinsic meaning. `[META]` |
+| If               | Condition       | Optional     | Condition to be evaluated. The element is ignored unless the condition evaluates to `true`.                                                |
 
 ### Trigger
 
 `<Trigger>` serves as a container for portions of the graph that should only be executed with explicit user intervention. 
 To execute nodes behind a trigger, pass `-Trigger=<Name>` to the command line. 
 
-[INCLUDE:#TriggerTable]
+|                  |                 |              |                                                                                             |
+| ---------------- | --------------- | ------------ | ------------------------------------------------------------------------------------------- |
+| Name             | Name            | Required     | Name of the trigger.                                                                        |
+| If               | Condition       | Optional     | Condition to be evaluated. The element is ignored unless the condition evaluates to `true`. |
 
 ## Property Manipulation
 
@@ -92,20 +110,33 @@
 `<Property>` sets the value of a property. If a property in an outer scope has already been declared with the same name,
 `<Property>` overwrites it. Otherwise, a new property is declared at the current scope.
 
-[INCLUDE:#PropertyTable]
+|                  |            |          |                                                                                             |
+| ---------------- | ---------- | -------- | ------------------------------------------------------------------------------------------- |
+| Name             | Name       | Required | Name of the property to set.                                                                |
+| Value            | String     | Required | New value for the property.                                                                 |
+| If               | Condition  | Optional | Condition to be evaluated. The element is ignored unless the condition evaluates to `true`. |
 
 ### Option
 
 `<Option>` defines a user option that can be set from the command line. `<Option>` may only appear at a global scope.
 
-[INCLUDE:#OptionTable]
+|              |           |              |                                                                                             |
+| ------------ | --------- | ------------ | ------------------------------------------------------------------------------------------- |
+| Name         | Name      | Required     | Name of the option (and property) to initialize with the option's value.                    |
+| Description  | String    | Required     | Description for the option to display when running BuildGraph with the -ListOnly argument.  |
+| Restrict     | Regex     | Optional     | Regex that matches valid values for this option (eg. `[a-zA-Z]+`, `true | false`).          |
+| DefaultValue | String    | Required     | Default value for the option if the user does not set it explicitly.                        |
+| If           | Condition | Optional     | Condition to be evaluated. The element is ignored unless the condition evaluates to `true`. |
 
 ### EnvVar
 
 `<EnvVar>` declares a property to contain the contents of an environment variable (or an empty string if it's not set).
 `<EnvVar>` may only appear at a global scope.
 
-[INCLUDE:#EnvVarTable]
+|                  |             |          |                                                                                             |
+| ---------------- | ----------- | -------- | ------------------------------------------------------------------------------------------- |
+| Name             | Name        | Required | Name of the environment variable to introduce as a property.                                |
+| If               | Condition   | Optional | Condition to be evaluated. The element is ignored unless the condition evaluates to `true`. |
 
 ## Flow Control
 
@@ -121,13 +152,18 @@
 `<Include>` processes the contents of another BuildGraph script as if it appeared within this file.
 `<Include>` may only appear at a global scope.
 
-[INCLUDE:#IncludeTable]
+|                  |                 |              |                                                                                             |
+| ---------------- | --------------- | ------------ | ------------------------------------------------------------------------------------------- |
+| Script           | String          | Required     | Path to the script you wan to include (relative to the current script).                     |
+| If               | Condition       | Optional     | Condition to be evaluated. The element is ignored unless the condition evaluates to `true`. |
 
 ### Do
 
 `<Do>` groups a sequence of elements and only processes them if a condition evaluates as true.
 
-[INCLUDE:#DoTable]
+|    |           |          |                                                                                             |
+| -- | ----------| ---------| ------------------------------------------------------------------------------------------- |
+| If | Condition | Optional | Condition to be evaluated. The element is ignored unless the condition evaluates to `true`. |
 
 ### Switch
 
@@ -148,7 +184,9 @@
 		</Default>
 	</Switch>
 
-[INCLUDE:#SwitchTable]
+|    |           |          |                                                                                             |
+| -- | ----------| ---------| ------------------------------------------------------------------------------------------- |
+| If | Condition | Optional | Condition to be evaluated. The element is ignored unless the condition evaluates to `true`. |
 
 ### ForEach
 
@@ -161,7 +199,11 @@
 		<Log Message="Counter=$(Counter)"/>
 	</ForEach>
 
-[INCLUDE:#ForEachTable]
+|                  |                 |              |                                                                                             |
+| ---------------- | --------------- | ------------ | ------------------------------------------------------------------------------------------- |
+| Name             | Name            | Required     | Property to assign to each value in the list.                                               |
+| Values           | String List     | Required     | List of values that are separated by semicolons.                                            |
+| If               | Condition       | Optional     | Condition to be evaluated. The element is ignored unless the condition evaluates to `true`. |
 
 ## Diagnostics
 
@@ -176,147 +218,19 @@
 at global scope. `<Warning>` will only provide an output if it's still part of the graph after it has been trimmed down to 
 the target being executed.
 
-[INCLUDE:#WarningTable]
+|                  |                 |              |                                                                                             |
+| ---------------- | --------------- | ------------ | ------------------------------------------------------------------------------------------- |
+| Message          | String          | Required     | Text to be printed to the log.                                                              |
+| If               | Condition       | Optional     | Condition to be evaluated. The element is ignored unless the condition evaluates to `true`. |
 
 ### Error
 
 `<Error>`outputs an error message. `<Error>` may be included in nodes, agents, triggers, or at global scope. 
 `<Error>` will only provide an output if it's still part of the graph after it has been trimmed down to the target being executed.
 
-[INCLUDE:#ErrorTable]
+|                  |                 |              |                                                                                             |
+| ---------------- | --------------- | ------------ | ------------------------------------------------------------------------------------------- |
+| Message          | String          | Required     | Text to be printed to the log.                                                              |
+| If               | Condition       | Optional     | Condition to be evaluated. The element is ignored unless the condition evaluates to `true`. |
 
 [/EXCERPT:BuildGraphScriptElements]
-
-<!--Tables for the 'Defining a Graph' section.-->
-
-<!--
-[EXCERPT:NodeTable]
-|                  |              |              |                                                                                                                                             |
-| ---------------- | ------------ | ------------ | ------------------------------------------------------------------------------------------------------------------------------------------- |
-| Name             | Name         | Required     | Name of the node.                                                                                                                           |
-| Requires         | Target List  | Optional     | List of nodes, aggregates, or tagged file sets produced by other nodes that this node requires to execute, separated by semicolons.         |
-| Produces         | Tag List     | Optional     | Tagged file sets that this node makes available to other nodes, separated by semicolons.                                                    |
-| After            | Target List  | Optional     | List of nodes that this node should run after, if they are part of the current target (ignored if they are not). Separated by semicolons.   |
-| NotifyOnWarnings | Boolean      | Optional     | If false, this node will not produce notifications on warnings. Only used by build systems. Defaults to `true`. `[META]`                    |
-| If               | Condition    | Optional     | Condition to be evaluated. The element is ignored unless the condition evaluates to `true`.                                                 |
-[/EXCERPT:NodeTable]
--->
-
-<!--
-[EXCERPT:AggregateTable]
-|                  |              |              |                                                                                             |
-| ---------------- | ------------ | ------------ | ------------------------------------------------------------------------------------------- |
-| Name             | Name         | Required     | Name of the aggregate.                                                                      |
-| Requires         | Target List  | Required     | List of dependencies for this aggregate. May be nodes, tagged file sets, or agent groups.   |
-| If               | Condition    | Optional     | Condition to be evaluated. The element is ignored unless the condition evaluates to `true`. |
-[/EXCERPT:AggregateTable]
--->
-
-<!--
-[EXCERPT:AgentTable]
-|                  |                 |              |                                                                                                                                            |
-| ---------------- | --------------- | ------------ | ------------------------------------------------------------------------------------------------------------------------------------------ |
-| Name             | Name            | Required     | Name of the group.                                                                                                                         |
-| Type             | Identifier List | Optional     | Type of agent to run on. The meaning of this string is inferred by the host build system; it does not have any intrinsic meaning. `[META]` |
-| If               | Condition       | Optional     | Condition to be evaluated. The element is ignored unless the condition evaluates to `true`.                                                |
-[/EXCERPT:AgentTable]
--->
-
-<!--
-[EXCERPT:TriggerTable]
-|                  |                 |              |                                                                                             |
-| ---------------- | --------------- | ------------ | ------------------------------------------------------------------------------------------- |
-| Name             | Name            | Required     | Name of the trigger.                                                                        |
-| If               | Condition       | Optional     | Condition to be evaluated. The element is ignored unless the condition evaluates to `true`. |
-[/EXCERPT:TriggerTable]
--->
-
-<!--Tables for the 'Property Manipulation' section.-->
-
-<!--
-[EXCERPT:PropertyTable]
-|                  |            |          |                                                                                             |
-| ---------------- | ---------- | -------- | ------------------------------------------------------------------------------------------- |
-| Name             | Name       | Required | Name of the property to set.                                                                |
-| Value            | String     | Required | New value for the property.                                                                 |
-| If               | Condition  | Optional | Condition to be evaluated. The element is ignored unless the condition evaluates to `true`. |
-[/EXCERPT:PropertyTable]
--->
-
-<!--
-[EXCERPT:OptionTable]
-|              |           |              |                                                                                             |
-| ------------ | --------- | ------------ | ------------------------------------------------------------------------------------------- |
-| Name         | Name      | Required     | Name of the option (and property) to initialize with the option's value.                    |
-| Description  | String    | Required     | Description for the option to display when running BuildGraph with the -ListOnly argument.  |
-| Restrict     | Regex     | Optional     | Regex that matches valid values for this option (eg. `[a-zA-Z]+`, `true | false`).          |
-| DefaultValue | String    | Required     | Default value for the option if the user does not set it explicitly.                        |
-| If           | Condition | Optional     | Condition to be evaluated. The element is ignored unless the condition evaluates to `true`. |
-[/EXCERPT:OptionTable]
--->
-
-<!--
-[EXCERPT:EnvVarTable]
-|                  |             |          |                                                                                             |
-| ---------------- | ----------- | -------- | ------------------------------------------------------------------------------------------- |
-| Name             | Name        | Required | Name of the environment variable to introduce as a property.                                |
-| If               | Condition   | Optional | Condition to be evaluated. The element is ignored unless the condition evaluates to `true`. |
-[/EXCERPT:EnvVarTable]
--->
-
-<!--Tables for the 'Flow Control' section.-->
-
-<!--
-[EXCERPT:IncludeTable]
-|                  |                 |              |                                                                                             |
-| ---------------- | --------------- | ------------ | ------------------------------------------------------------------------------------------- |
-| Script           | String          | Required     | Path to the script you wan to include (relative to the current script).                     |
-| If               | Condition       | Optional     | Condition to be evaluated. The element is ignored unless the condition evaluates to `true`. |
-[/EXCERPT:IncludeTable]
--->
-
-<!--
-[EXCERPT:DoTable]
-|    |           |          |                                                                                             |
-| -- | ----------| ---------| ------------------------------------------------------------------------------------------- |
-| If | Condition | Optional | Condition to be evaluated. The element is ignored unless the condition evaluates to `true`. |
-[/EXCERPT:DoTable]
--->
-
-<!--
-[EXCERPT:SwitchTable]
-|    |           |          |                                                                                             |
-| -- | ----------| ---------| ------------------------------------------------------------------------------------------- |
-| If | Condition | Optional | Condition to be evaluated. The element is ignored unless the condition evaluates to `true`. |
-[/EXCERPT:SwitchTable]
--->
-
-<!--
-[EXCERPT:ForEachTable]
-|                  |                 |              |                                                                                             |
-| ---------------- | --------------- | ------------ | ------------------------------------------------------------------------------------------- |
-| Name             | Name            | Required     | Property to assign to each value in the list.                                               |
-| Values           | String List     | Required     | List of values that are separated by semicolons.                                            |
-| If               | Condition       | Optional     | Condition to be evaluated. The element is ignored unless the condition evaluates to `true`. |
-[/EXCERPT:ForEachTable]
--->
-
-<!--Tables for the 'Diagnostics' section.-->
-
-<!--
-[EXCERPT:WarningTable]
-|                  |                 |              |                                                                                             |
-| ---------------- | --------------- | ------------ | ------------------------------------------------------------------------------------------- |
-| Message          | String          | Required     | Text to be printed to the log.                                                              |
-| If               | Condition       | Optional     | Condition to be evaluated. The element is ignored unless the condition evaluates to `true`. |
-[/EXCERPT:WarningTable]
--->
-
-<!--
-[EXCERPT:ErrorTable]
-|                  |                 |              |                                                                                             |
-| ---------------- | --------------- | ------------ | ------------------------------------------------------------------------------------------- |
-| Message          | String          | Required     | Text to be printed to the log.                                                              |
-| If               | Condition       | Optional     | Condition to be evaluated. The element is ignored unless the condition evaluates to `true`. |
-[/EXCERPT:ErrorTable]
--->
