Fixing table includes.

#UE4doc

diff --git a/Engine/Documentation/Source/Programming/Development/BuildGraph/ScriptAnatomy/Conditions/BuildGraphScriptConditions.INT.udn b/Engine/Documentation/Source/Programming/Development/BuildGraph/ScriptAnatomy/Conditions/BuildGraphScriptConditions.INT.udn
--- a/Engine/Documentation/Source/Programming/Development/BuildGraph/ScriptAnatomy/Conditions/BuildGraphScriptConditions.INT.udn
+++ b/Engine/Documentation/Source/Programming/Development/BuildGraph/ScriptAnatomy/Conditions/BuildGraphScriptConditions.INT.udn
@@ -55,16 +55,6 @@ case-insensitive for comparisons, which means that the strings "True" and 'true'
 
 A list of operators is specified below: 
 
-[INCLUDE:#OperatorTable]
-
-[REGION:note]
-The `'<'` and `'>'` characters must be escaped as `"&lt;"` and `"&gt;"` in XML (referencing [XML Schema v1.0](https://www.w3.org/TR/NOTE-dcd])).
-[/REGION]
-
-[/EXCERPT:BuildGraphScriptConditions]
-
-<!--
-[EXCERPT:OperatorTable]
 | **Operator**            | **Description**                                                             | **Precedence**  |
 | ----------------------- | --------------------------------------------------------------------------- | --------------- |
 | **(x)**                 | Subexpression                                                               | 1               |
@@ -79,5 +69,9 @@ The `'<'` and `'>'` characters must be escaped as `"&lt;"` and `"&gt;"` in XML (
 | **x >= y**              | Compares whether the integer x is greater than or equal to the integer y.   | 2               |
 | **x and y**             | True if both arguments are `true`.                                          | 3               |
 | **x or y**              | True if either argument is `true`.                                          | 4               |
-[/EXCERPT:OperatorTable]
--->
\ No newline at end of file
+
+[REGION:note]
+The `'<'` and `'>'` characters must be escaped as `"&lt;"` and `"&gt;"` in XML (referencing [XML Schema v1.0](https://www.w3.org/TR/NOTE-dcd])).
+[/REGION]
+
+[/EXCERPT:BuildGraphScriptConditions]
\ No newline at end of file
diff --git a/Engine/Documentation/Source/Programming/Development/BuildGraph/ScriptAnatomy/Types/BuildGraphScriptTypes.INT.udn b/Engine/Documentation/Source/Programming/Development/BuildGraph/ScriptAnatomy/Types/BuildGraphScriptTypes.INT.udn
--- a/Engine/Documentation/Source/Programming/Development/BuildGraph/ScriptAnatomy/Types/BuildGraphScriptTypes.INT.udn
+++ b/Engine/Documentation/Source/Programming/Development/BuildGraph/ScriptAnatomy/Types/BuildGraphScriptTypes.INT.udn
@@ -34,12 +34,6 @@ tags:BuildGraph
 
 The following table contains valid data types that can be contained by BuildGraph attributes:
 
-[INCLUDE:#TypesTable]
-
-[/EXCERPT:BuildGraphScriptTypes]
-
-<!--
-[EXCERPT:TypesTable]
 | **Type**                | **Description**                                                                                                                                                                                                                                                                            |
 | ----------------------- | ------------------------------------------------------------------------------------------------------------------------------------------------------------------------------------------------------------------------------------------------------------------------------------------ |
 | **String**              | An arbitrary string.                                                                                                                                                                                                                                                                       |
@@ -57,5 +51,5 @@ The following table contains valid data types that can be contained by BuildGrap
 | **File Name**           | A path to a file. Unless otherwise specified, relative paths are resolved to the working root directory.                                                                                                                                                                                   |
 | **Directory Name**      | A path to a directory. Unless otherwise specified, relative paths are resolved relative to the working root directory.                                                                                                                                                                     |
 | **Condition**           | A [conditional expression](Programming/Development/BuildGraph/ScriptAnatomy/Conditions).                                                                                                                                                                                                      |
-[/EXCERPT:TypesTable]
--->
\ No newline at end of file
+
+[/EXCERPT:BuildGraphScriptTypes]
\ No newline at end of file
