Adding new editor windows that are accessible now to this doc, and updating screenshots of quick menu UI UEDOC-3210

diff --git a/Engine/Documentation/Source/Engine/Editor/VR/EditorUI/EditorUIVR.INT.udn b/Engine/Documentation/Source/Engine/Editor/VR/EditorUI/EditorUIVR.INT.udn
--- a/Engine/Documentation/Source/Engine/Editor/VR/EditorUI/EditorUIVR.INT.udn
+++ b/Engine/Documentation/Source/Engine/Editor/VR/EditorUI/EditorUIVR.INT.udn
@@ -6,7 +6,7 @@ parent: Engine/Editor/VR
 type: reference
 order: 5
 tags: VR
-version: 4.11
+version: 4.13
 skilllevel: intermediate
 
 [VAR:Topic]
@@ -48,9 +48,11 @@ skilllevel: intermediate
 The VR Editor enables you to summon and dock a number of different panels and windows including:
 
 [REGION:imagetable]
-| ![](window_content.png)(w:145) | ![](window_details.png)(w:145) | ![](window_modes.png)(w:145) | ![](window_outliner.png)(w:145) | ![](window_tutorials.png)(w:145) | ![](window_asset.png)(w:145) |
-| --- | --- | --- | --- | --- | --- |
-| Content Browser | Details Panel | Modes Panel | Outliner panel | Tutorials | Asset Panel |
+| ![](window_content.png)(w:185) | ![](window_details.png)(w:185) | ![](window_modes.png)(w:185) | ![](window_outliner.png)(w:185) |
+| --- | --- | --- | --- |
+| Content Browser | Details Panel | Modes Panel | Outliner panel | 
+|![](window_tutorials.png)(w:185) | ![](window_asset.png)(w:185) | ![](colorPicker.png)(w:185)| ![](worldSettings.png)(w:185)|
+|Tutorials | Asset Panel | Color Picker | World Settings |
 [/REGION]
 
 ## Opening a Window in VR Mode
@@ -59,11 +61,11 @@ To open a window:
 
 1. Aim one controller at the other to open the **Quick Select** menu.
 
-	![](window_quick_select.png)(w:480)
+	![](quickmenu_1.png)(w:480)
 
 1. Aim the free controller at one of the options in the Interface section and squeeze the **Trigger** to open that window.
 
-	![](window_quick_select_highlight.png)(w:480)
+	![](quickmenu_2.png)(w:480)
 
 
 ## Docking
@@ -117,13 +119,13 @@ There are a number of different ways you can interact with panels and windows in
 
 1. Squeeze the Trigger to close the window.
 
-**To show / hide all windows:** 
+**To show / hide key windows - Content Browser, Details Panel, Modes Panel, and Scene Outliner:** 
 
-1. With windows open, press down once on the **Track Pad** to hide all windows.
+1. With windows open, press down once on the **Track Pad** to hide key windows.
 
 	![](window_show.png)(w:480) ![](window_hide.png)(w:480)
 
-1. With windows hidden, press down once on the **Track Pad** to show all windows.
+1. With windows hidden, press down once on the **Track Pad** to show key windows.
 
 	![](window_hide.png)(w:480) ![](window_show.png)(w:480)
 
