Initial pass at Media Framework section:
- Landing Page updated.
- Overview Page added.
- Technical Page added.
- How-to Landing Page added.
- Controlling Playback How-to added.
- Playing Videos How-to added.
- Streaming Videos How-to added.
- Using Playlists How-to added.

diff --git a/Engine/Documentation/Source/Engine/MediaFramework/HowTo/ControllingPlayback/ControllingPlayback.INT.udn b/Engine/Documentation/Source/Engine/MediaFramework/HowTo/ControllingPlayback/ControllingPlayback.INT.udn
new file mode 100644
--- /dev/null
+++ b/Engine/Documentation/Source/Engine/MediaFramework/HowTo/ControllingPlayback/ControllingPlayback.INT.udn
@@ -0,0 +1,107 @@
+Availability:Docs
+Title: Rewind & Fast Forward Playback
+Crumbs:%ROOT%, Engine
+Description:Example of how you can control playback of your media assets.
+parent:Engine/MediaFramework/HowTo
+order:3
+checkpoint: editorqs
+Version: 4.13
+SkillLevel: Intermediate
+Related: Engine\MediaFramework
+
+In addition to playing a video with Media Framework, you can grant players the ability to control the Playback Rate to rewind, pause or fast forward videos. 
+In this how-to we allow the player to press a button to rewind, a different button to fast forward and another button to play a video playing in the level. 
+
+[REGION:note]
+For this how-to we are using the **Blueprint Third Person Template** project with **Starter Content** enabled.
+You will also need a few [Supported Video Files](Engine\MediaFramework\TechReference) saved on our computer that you want to add to the playlist for playback. You can use your own or **Right-click** and download this [Sample Video](Infiltrator Demo.mp4). 
+[/REGION]
+
+
+## Steps
+
+1.	In your project create a folder called **Movies**, then in explorer add your video to the **Movies** folder. 
+
+	![](ControllingPlayback_01.png)
+
+1.	Create a **File Media Source** (or any other source asset) and point the asset to your video. 
+
+	![](ControllingPlayback_02.png)
+
+1.	Create a **Media Player** asset called **VideoPlayer** and associated **Media Sound Wave** and **Media Texture** assets.
+
+	![](ControllingPlayback_03.png)
+
+1.	Inside the **Media Player** asset, uncheck the **Play on Open** option. 
+
+	![](ControllingPlayback_04.png)
+
+1.	In the **Content Browser** under **StarterContent/Shapes** drag a **Shape_Plane** into the level.  
+
+	![](ControllingPlayback_05.png)
+
+1.	Use the **Translate (W)**, **Rotation (E)** and **Scale (R)** tools to adjust the size and location of the Static Mesh as desired. 
+
+1.	Hold **Ctrl** then select both the **Sound Wave** and **Media Texture** assets and drop them on top of the **Static Mesh** in the level. 
+
+	![](ControllingPlayback_06.png)
+
+	This will automatically create a **Material** and apply it to the Static Mesh. 
+
+1.	From the Toolbar, click the **Blueprints** button and **Open Level Blueprint**. 
+
+	![](ControllingPlayback_07.png)
+
+1.	Add a **Variable** of the **Media Player Reference** type called **MediaPlayer** and set it to your **VideoPlayer** Media Player asset. 
+
+	![](ControllingPlayback_08.png)
+
+	You may need to **Compile** before you can assign the **Default Value**. 
+
+1.	Add an **Event Begin Play** node, your **Media Player** variable and use **Open Source** to open the **MyVideo** File Media Source asset.
+
+	![](ControllingPlayback_09.png)
+
+1.	Add a **2** Keyboard Event with **Pressed** set to **Play** the **MediaPlayer** asset. 
+
+	![](ControllingPlayback_10.png)
+
+1.	Add **1** and **3** Keyboard Events both using **Set Rate** off the **MediaPlayer** asset and adjust each node's **Rate** value. 
+
+	![](ControllingPlayback_11.png)
+
+	By setting Rate to **-1.0** this will play the video in reverse while a Rate of **2.0** doubles the forward playback speed, scrubbing forward. 
+
+	[REGION:warning]
+	Currently you may experience audio synch loss when issuing commands to the Media Player asset. You can insert short delays between commands to allow the player has time to process them. We are aware of these issues and are investigating possible fixes to accompany future engine updates. 
+	[/REGION]
+
+1.	Close the **Level Blueprint** and click the **Play** button to play in the editor. 
+
+
+## End Result
+
+[OBJECT:EmbeddedVideo]
+[PARAMLITERAL:width]
+640
+[/PARAMLITERAL]
+[PARAMLITERAL:height]
+360
+[/PARAMLITERAL]
+[PARAMLITERAL:videoid]
+s6rTrV5-Mdo
+[/PARAMLITERAL]
+[/OBJECT]
+
+When playing in the editor, the video playback will not initially start when the level is loaded. When the player presses the **2** key, the video will start playing. 
+Pressing the **1** key, the video will start to rewind while pressing the **3** key fast forwards the video. 
+
+The rate at which we rewind or fast forward is determined by the **Rate** value on the **Set Rate** node. Adjust those values to speed up the play back rate when going in reverse or forward. 
+
+While our example is bound to keyboard presses, you could execute the same functionality within a User Interface built with UMG where the player clicks on a button to execute the same commands. 
+
+
+
+
+
+
diff --git a/Engine/Documentation/Source/Engine/MediaFramework/HowTo/FileMediaSource/UsingFileMedia.INT.udn b/Engine/Documentation/Source/Engine/MediaFramework/HowTo/FileMediaSource/UsingFileMedia.INT.udn
new file mode 100644
--- /dev/null
+++ b/Engine/Documentation/Source/Engine/MediaFramework/HowTo/FileMediaSource/UsingFileMedia.INT.udn
@@ -0,0 +1,119 @@
+Availability:Docs
+Title: Playing Videos in a Level
+Crumbs:%ROOT%, Engine
+Description:Describes how to use the File Media asset to play videos inside a level. 
+parent:Engine/MediaFramework/HowTo
+order:3
+checkpoint: editorqs
+Version: 4.13
+SkillLevel: Beginner
+Related: Engine\MediaFramework
+
+If you are looking to have video playing inside your level on an object whether its a TV or a monitor or some other object, you can do so using the **Media Framework** assets. 
+In this how-to we place a video file inside the **Content** folder of our project, then use the **Media Player** and related assets (with a little **Blueprint Visual Script**) to play the video during gameplay. 
+
+[REGION:note]
+For this how-to we are using the **Blueprint Third Person Template** project with **Starter Content** enabled.
+You will also need a [Supported Video File](Engine\MediaFramework\TechReference) saved on our computer that you want to play. You can use your own or **Right-click** and download this [Sample Video](SampleVideo.mp4). 
+[/REGION]
+
+
+## Steps
+
+1.	In the **Content Browser** expand the **Sources Panel**, then under **Content** create a folder called **Movies**. 
+
+	![](FileMedia_01.png)
+
+1.	**Right-click** on the **Movies** folder and select **Show in Explorer**. 
+
+	![](FileMedia_01b.png)
+
+1.	Drag your video file into the **Content/Movies** folder of your project. 
+
+	![](FileMedia_02.png)
+
+	Placing your video files inside the **Content/Movies** folder of your project will ensure the videos are properly packaged. 
+
+1.	Inside Unreal Engine and your project, **Right-click** in the **Movies** folder and under **Media** select **File Media Source**. 
+
+	![](FileMedia_03.png)
+
+1.	Call the asset **SampleVideo**, then open it and point the **File Path** to the location of your video file in the **Content/Movies** folder. 
+
+	![](FileMedia_04.png)
+
+1.	**Right-click** in the **Content Browser** then under **Media** select the **Media Player** asset. 
+
+	![](FileMedia_05.png)
+
+1.	In the **Create Media Player** window, check **Audio output SoundWave asset** and **Video output Media Texture asset**. 
+
+	![](FileMedia_06.png)
+
+	This will automatically create a Sound Wave and Media Texture asset linked to this Media Player asset which is needed for playback. 
+
+	![](FileMedia_07.png)
+
+	We have called our Media Player asset **SampleMedia** which is also applied to the Sound Wave and Media Texture. 
+
+	If you open the Media Player asset, you can **Double-click** on your Media Source asset and it will playback. 
+
+	![](FileMedia_08.png)
+
+	In the lower-right in the **Details** panel the **Output** section will also have auto-assigned the Sound Wave and Media Texture. 
+
+1.	In the **Content Browser** under **StarterContent/Shapes** drag a **Shape_Plane** into the level. 
+
+	![](FileMedia_09.png)
+
+	We will use this Static Mesh to playback our video on. 
+
+1.	Use the **Translate (W)**, **Rotation (E)** and **Scale (R)** tools to adjust the mesh similar to below. 
+
+	![](FileMedia_10.png)
+
+1.	Hold **Ctrl** then select both the **Sound Wave** and **Media Texture** assets and drop them on top of the **Static Mesh** in the level. 
+
+	![](FileMedia_11.png)
+
+	This will automatically create a **Material** and apply it to the Static Mesh. 
+
+1.	From the Toolbar, click the **Blueprints** button and **Open Level Blueprint**. 
+
+	![](FileMedia_12.png)
+
+1.	Add a **Variable** of the **Media Player Reference** type called **Media Player** and set it to your **SampleMedia** Media Player asset. 
+
+	![](FileMedia_13.png)
+
+	You may need to **Compile** before you can assign the **Default Value**. 
+
+1.	Hold **Ctrl** and drag in the **MediaPlayer** variable and also **Right-click** and add an **Event Begin Play** node. 
+
+	![](FileMedia_14.png)
+
+1.	Drag off the **MediaPlayer** variable and use the **Open Source** node set to your **SampleVideo** Media Source then connect as shown. 
+
+	![](FileMedia_15.png)
+
+1.	Close the Level Blueprint, then **Play** in the editor. 
+
+	
+## End Result
+
+When you play in the editor, the video will start playing on the Static Mesh. 
+
+[OBJECT:EmbeddedVideo]
+[PARAMLITERAL:width]
+640
+[/PARAMLITERAL]
+[PARAMLITERAL:height]
+360
+[/PARAMLITERAL]
+[PARAMLITERAL:videoid]
+zxyY2s6ZfrU
+[/PARAMLITERAL]
+[/OBJECT]
+
+**Media Player** assets by default are set to **Play on Open** which is why the video automatically starts playing when calling the **Open Source** function.
+There are additional commands you can issue to the Media Player asset such as pausing, rewinding or stopping the video once it has begun playing which can be found under the **Media Player** section when dragging off the Media Player Reference.
diff --git a/Engine/Documentation/Source/Engine/MediaFramework/HowTo/MF_HowTo.INT.udn b/Engine/Documentation/Source/Engine/MediaFramework/HowTo/MF_HowTo.INT.udn
new file mode 100644
--- /dev/null
+++ b/Engine/Documentation/Source/Engine/MediaFramework/HowTo/MF_HowTo.INT.udn
@@ -0,0 +1,35 @@
+Availability:Docs
+Title: Media Framework How-To
+Crumbs:%ROOT%, Engine
+Description:Links to step-by-step guides for working with Media Framework in Unreal Engine 4. 
+Navigation:topic
+parent:Engine/MediaFramework
+order:2
+Version: 4.13
+Related: Engine\MediaFramework\Overview
+Related: Engine\MediaFramework\TechReference
+
+This page provides links to step-by-step guides to creating and using the **Media Framework** tools within Unreal Engine 4. 
+If you are looking to quickly get a video asset up and running inside your levels, check out the Beginner tutorials which illustrate how to use files located on disk as well as URL based media sources. 
+
+Refer to the table below for more information and click each link to view its contents. 
+
+## Beginner
+[REGION:wrapper transparent]
+[](FANCY:Engine\MediaFramework\HowTo\FileMediaSource "%Engine\MediaFramework\HowTo\FileMediaSource:description%")
+[](FANCY:Engine\MediaFramework\HowTo\StreamMediaSource "%Engine\MediaFramework\HowTo\StreamMediaSource:description%")
+[/REGION]
+
+
+## Intermediate
+[REGION:wrapper transparent]
+[](FANCY:Engine\MediaFramework\HowTo\MediaPlaylists "%Engine\MediaFramework\HowTo\MediaPlaylists:description%")
+[](FANCY:Engine\MediaFramework\HowTo\ControllingPlayback "%Engine\MediaFramework\HowTo\ControllingPlayback:description%")
+[/REGION]
+
+<!-- To Do
+## Advanced
+[REGION:wrapper transparent]
+[/REGION]
+-->
+
diff --git a/Engine/Documentation/Source/Engine/MediaFramework/HowTo/MediaPlaylists/UsingMediaPlaylists.INT.udn b/Engine/Documentation/Source/Engine/MediaFramework/HowTo/MediaPlaylists/UsingMediaPlaylists.INT.udn
new file mode 100644
--- /dev/null
+++ b/Engine/Documentation/Source/Engine/MediaFramework/HowTo/MediaPlaylists/UsingMediaPlaylists.INT.udn
@@ -0,0 +1,120 @@
+Availability:Docs
+Title: Working with Media Playlists
+Crumbs:%ROOT%, Engine
+Description:Example of how you can use a Media Playlist to playback multiple Media Source assets. 
+Version: 4.13
+SkillLevel: Intermediate
+parent:Engine/MediaFramework/HowTo
+order:3
+checkpoint: editorqs
+Related: Engine\MediaFramework
+
+Suppose you have a series of videos that you want to playback in a specific order or you want to allow players to access videos within a specific series of videos. 
+You can do so by using a **Media Playlist** asset which points to **Media Source** assets assigned in a predefined order. 
+
+In this how-to we create a Media Playlist and allow the player to playback each asset sequentially or access a specific asset inside the playlist through a key press. 
+
+[REGION:note]
+For this how-to we are using the **Blueprint Third Person Template** project with **Starter Content** enabled.
+You will also need a few [Supported Video Files](Engine\MediaFramework\TechReference) saved on our computer that you want to add to the playlist for playback. You can use your own or **Right-click** and download these two sample videos: [Video Example 1](Infiltrator Demo.mp4) and [Video Example 2](SampleVideo.mp4)
+[/REGION]
+
+
+## Steps
+
+1.	In the **Content Browser** expand the **Sources Panel**, then in the **Content** create a new folder called **Movies**. 
+
+	![](Playlist_00.png)
+
+1.	**Right-click** on the **Movies** folder and choose **Show in Explorer** then place your video files in the folder. 
+
+	![](Playlist_01.png)
+
+	Placing your video files inside the **Content/Movies** folder of your project will ensure the videos are properly packaged. 
+
+1.	Create two **File Media Source** assets (called **Video_01** and **Video_02**) point the **File Path** to each respective video. 
+
+	![](Playlist_02.png)
+
+1.	Create a **Media Playlist** asset inside the content browser called **MyPlaylist**. 
+
+	![](Playlist_03.png)
+
+1.	Inside **MyPlaylist**, **Double-click** on each **File Media Source** to add it to the playlist. 
+
+	![](Playlist_04.png)
+
+	You can also use the **+** button to add an empty element which you can then assign your source media. 
+	In this example we are using **File Media Source** assets, however you can also include **Stream Media Source** assets in a Playlist. 
+
+1.	Create a **Media Player** asset called **VideoPlayer**. 
+
+	![](Playlist_05.png)
+
+1.	In the **Create Media Player** window, check **Audio output SoundWave asset** and **Video output Media Texture asset**. 
+
+	![](Playlist_06.png)
+
+	This will automatically create a Sound Wave and Media Texture asset linked to this Media Player asset which is needed for playback. 
+
+	![](Playlist_07.png)
+
+1.	In the **Content Browser** under **StarterContent/Shapes** drag a **Shape_Plane** into the level.  
+
+	![](Playlist_09.png)
+
+1.	Use the **Translate (W)**, **Rotation (E)** and **Scale (R)** tools to adjust the size and location of the Static Mesh as desired. 
+
+1.	Hold **Ctrl** then select both the **Sound Wave** and **Media Texture** assets and drop them on top of the **Static Mesh** in the level. 
+
+	![](Playlist_10.png)
+
+	This will automatically create a **Material** and apply it to the Static Mesh. 
+
+1.	From the Toolbar, click the **Blueprints** button and **Open Level Blueprint**. 
+
+	![](Playlist_11.png)
+
+1.	Add a **Variable** of the **Media Player Reference** type called **Media Player** and set it to your **VideoPlayer** Media Player asset. 
+
+	![](Playlist_12.png)
+
+	You may need to **Compile** before you can assign the **Default Value**. 
+
+1.	**Right-click** in the **Event Graph** and add a **1 Keyboard Event**, **2 Keyboard Event** and **3 Keyboard Event**.  
+
+	![](Playlist_13.png)
+
+1.	**Ctrl** drag in the **Media Player** variable, then drag off it and use **Open Playlist Index** and set the **InPlaylist** to **MyPlaylist**. 
+
+	![](Playlist_14.png)
+
+1.	Use **Open Playlist** for the **2 Keyboard Event** with **Index** set to **1** and use **Open Playlist** for the **3 Keyboard Event**. 
+
+	![](Playlist_15.png)
+
+1.	Close the **Level Blueprint**, then click the **Play** button to play in the level. 
+
+	
+## End Result
+
+[OBJECT:EmbeddedVideo]
+[PARAMLITERAL:width]
+640
+[/PARAMLITERAL]
+[PARAMLITERAL:height]
+360
+[/PARAMLITERAL]
+[PARAMLITERAL:videoid]
+kjbgHVnQxBI
+[/PARAMLITERAL]
+[/OBJECT]
+
+When playing in the editor, the video playback will not initially start when the level is loaded. When the player presses the **1** key, the playlist will start playing from the first video in the playlist. 
+When pressing the **2** key, the video playback will jump to the second video in the playlist. When pressing **3**, we are opening and playing the video from the start of the playlist which is the same as opening the playlist at index 0. 
+
+Our **Media Player** asset is set to **Play on Open** which is why each video is played when we call the open playlist function. 
+If you want to open a Media Source but not start playback immediately, you can uncheck the Play on Open option inside the Media Player asset then use the **Play** function to play an open Media Source asset. 
+
+
+
diff --git a/Engine/Documentation/Source/Engine/MediaFramework/HowTo/StreamMediaSource/UsingStreamMedia.INT.udn b/Engine/Documentation/Source/Engine/MediaFramework/HowTo/StreamMediaSource/UsingStreamMedia.INT.udn
new file mode 100644
--- /dev/null
+++ b/Engine/Documentation/Source/Engine/MediaFramework/HowTo/StreamMediaSource/UsingStreamMedia.INT.udn
@@ -0,0 +1,103 @@
+Availability:Docs
+Title: Streaming Videos in a Level
+Crumbs:%ROOT%, Engine
+Description:Describes how to use the Stream Media asset to stream videos from the internet in Unreal Engine 4. 
+parent:Engine/MediaFramework/HowTo
+order:3
+checkpoint: editorqs
+Version: 4.13
+SkillLevel: Beginner
+Related: Engine\MediaFramework
+
+In this how-to we play a video on a Static Mesh inside the level when we start the game. Our source video will come from a provided URL stream and we use some **Blueprint Visual Script** to open the stream and start playback once we load into the level. 
+
+[REGION:note]
+For this how-to we are using the **Blueprint Third Person Template** project with **Starter Content** enabled.
+[/REGION]
+
+
+## Steps
+
+1.	In the **Content Browser** expand the **Sources Panel** then create a new folder called **Media** inside the **Content** folder.
+
+	![](StreamMedia_01.png)
+	
+1.	**Right-click** inside the **Content/Media** folder then under **Media** create a **Stream Media Source** and call it **MediaStream**. 
+
+	![](StreamMedia_02.png)
+
+1.	Open **MediaStream** then enter your desired **Stream URL**. 
+
+	![](StreamMedia_03.png)
+
+	Your Stream URL must link directly to a [Supported Format](Engine\MediaFramework\TechReference\) in order to play back the video. If you do not have a link, **Right-click** this [Sample URL](http://s3.amazonaws.com/akamai.netstorage/mp4/AirSTAR.mp4) and copy the link address then paste it in the Stream URL field. 
+
+1.	**Right-click** in the **Content Browser** then under **Media** select the **Media Player** asset. 
+
+	![](StreamMedia_04.png)
+
+1.	In the **Create Media Player** window, check **Audio output SoundWave asset** and **Video output Media Texture asset**. 
+
+	![](StreamMedia_05.png)
+
+	This will automatically create a Sound Wave and Media Texture asset linked to this Media Player asset which is needed for playback. 
+
+	![](StreamMedia_06.png)
+
+	We have called our Media Player asset **SampleStream** which is also applied to the Sound Wave and Media Texture. 
+
+1.	In the **Content Browser** under **StarterContent/Shapes** drag a **Shape_Plane** into the level. 
+
+	![](StreamMedia_07.png)
+
+	We will use this Static Mesh to playback our video on. 
+
+1.	Use the **Translate (W)**, **Rotation (E)** and **Scale (R)** tools to adjust the mesh to your desired size and location.
+
+1.	Hold **Ctrl** then select both the **Sound Wave** and **Media Texture** assets and drop them on top of the **Static Mesh** in the level. 
+
+	![](StreamMedia_08.png)
+
+	This will automatically create a **Material** and apply it to the Static Mesh. 
+
+1.	From the Toolbar, click the **Blueprints** button and **Open Level Blueprint**. 
+
+	![](StreamMedia_09.png)
+
+1.	Add a **Variable** of the **Media Player Reference** type called **Media Player** and set it to your **SampleStream** Media Player asset. 
+
+	![](StreamMedia_10.png)
+
+	You may need to **Compile** before you can assign the **Default Value**. 
+
+1.	Hold **Ctrl** and drag in the **MediaPlayer** variable and also **Right-click** and add an **Event Begin Play** node. 
+
+	![](StreamMedia_11.png)
+
+1.	Drag off the **MediaPlayer** variable and use the **Open Source** node set to your **MediaStream** Media Source then connect as shown. 
+
+	![](StreamMedia_12.png)
+
+1.	Close the Level Blueprint, then **Play** in the editor. 
+
+	
+## End Result
+
+When you play in the editor, the video will start playing on the Static Mesh. 
+
+[OBJECT:EmbeddedVideo]
+[PARAMLITERAL:width]
+640
+[/PARAMLITERAL]
+[PARAMLITERAL:height]
+360
+[/PARAMLITERAL]
+[PARAMLITERAL:videoid]
+0gxc69-M-Fk
+[/PARAMLITERAL]
+[/OBJECT]
+
+Similar to using a [File Media Source](Engine\MediaFramework\HowTo\FileMediaSource) which plays a video file from disk, Stream Media Sources will automatically play when calling **Open Source** if their associated **Media Player** asset are set to **Play on Open** which by default is enabled.
+There are additional commands you can issue to the Media Player asset such as pausing, rewinding or stopping the video once it has begun playing which can be found under the **Media Player** section when dragging off the Media Player Reference.
+
+
diff --git a/Engine/Documentation/Source/Engine/MediaFramework/MediaFramework.INT.udn b/Engine/Documentation/Source/Engine/MediaFramework/MediaFramework.INT.udn
--- a/Engine/Documentation/Source/Engine/MediaFramework/MediaFramework.INT.udn
+++ b/Engine/Documentation/Source/Engine/MediaFramework/MediaFramework.INT.udn
@@ -1,13 +1,11 @@
-Availability:Public
+Availability:Docs
 Title: Media Framework
 Crumbs:%ROOT%, Engine
-Description:Overview of Media Playback functionality in Unreal Engine 4. 
-Related: Engine/Blueprints
-Related: Engine/UMG
-Related: Engine/Rendering/Materials
+Description:Working with Media Playback functionality in Unreal Engine 4. 
+Navigation:topic
 parent:Engine
-order:12
-Version: 4.9
+order:1
+Version: 4.13
 
 [VAR:Topic]
 [OBJECT:Topic]
@@ -49,17 +47,14 @@ Version: 4.9
 [/VAR]
 
 [REGION:fullwidth raw]
-![](MediaBanner.png)(w:920)(convert:false)
+![](MediaFrameWork.png)(convert:false)
 [/REGION]
 
-[TOC(start:2 end:2)]
-
-
 While there is currently a **MoviePlayer** system in UE4, it can only be used for playing startup movies while the Engine loads up. It cannot be used to play movies in-game, 
 say for instance as part of a UI element or perhaps on a Static Mesh in your level like a movie playing on a TV. This is where the **Media Framework** comes in, as not only 
 will it allow you to do both of the examples presented above, it will also provide much more general media playback capabilities (a rundown of which are outlined below).  
 
-The Media Framework in Unreal Engine 4 is...
+The Media Framework in Unreal Engine 4 is:
 * Both Engine & Slate Agnostic
 * Able to support Localized Audio & Video Tracks
 * Accessible in the Content Browser, Material Editor, & Sound System
@@ -72,282 +67,14 @@ As stated above, the Media Framework itself is both Engine and Slate agnostic, w
 Editor. There are layers on top of the framework that provide media playback capabilities to other sub-systems as well, such as **Engine**, **Blueprints**, **Slate**, and 
 the **UMG UI Designer**.  This will cover most expected use cases, such as in-game textures & UI, in-Editor video tutorials, and Marketplace videos. 
 
+This page contains several links to additional documentation regarding Media Framework. It is recommended that you check out the Overview page for more of a breakdown of the features of Media Framework as well as the How-To and Quick Start pages for step-by-step instructions to working with Media Framework. 
+
 [REGION:warning]
 The Media Framework is not quite production ready yet and is still in the experimental stages. Some features may not work as expected, however we are currently working hard 
 to resolve those issues as well as add additional features and support. 
 [/REGION]
 
-## What is Media Framework
-
-The Media Framework is largely a collection of C++ interfaces, a couple of helper classes for common use cases and a media player factory that can be extended with so called 
-Media Player plug-ins. All the interesting work of playing movies and other media files happens in the plug-ins. 
-
-[REGION:note]
-Programmers: 
-* The code for the Media Framework is in the following location: _/Engine/Source/Runtime/Media/_
-* The existing player plug-ins can be found in the following location: _/Engine/Plugins/Media/_
-[/REGION:note]
-
-Currently a player plug-in for Windows is being used which uses the Windows Media Foundation API under the hood. There is also a plug-in for MacOS, which uses Apple's AV 
-Foundation, however it is currently quite limited but should work on iOS. Work has also begun on a player plug-in for Android. 
-
-[REGION:warning]
-While the Windows plug-in is the most complete, it has issues with certain movie formats, such as H.264 encoded .mp4 files. 
-[/REGION]
-
-
-## Media Asset Types & Importing Media
-
-With the implementation of the Media Framework comes three new asset types can be created from the Content Browser. 
-
-### MediaPlayer Asset
-
-The MediaPlayer asset represents a player for a media source such as a movie file on disk or a streaming media URL on the internet. Unlike the `UTextureMovie` in 
-Unreal Engine 3, the MediaPlayer asset does not actually contain any data, it only stores the path or URL to the media source. 
-
-There are a number of different ways you can create a MediaPlayer asset as illustrated below. 
-
-**Create MediaPlayer Asset**
-
-1. **Click** the **Create** button, then select **Miscellaneous** and **Media Player**. 
-
-	![](CreateMediaPlayerAsset.png)
-
-[REGION:note]
-Instead of clicking the **Create** button, you can also **Right-click** in the **Content Browser** to bring up the context menu.
-[/REGION]
-
-1. This will create the MediaPlayer Asset (**Double-click** on it to open it up). 
-
-	![](MediaPlayerAsset.png)(w:720)
-
-1. Inside the MediaPlayer asset, in the Details panel, click the **Browse** button. 
-
-	![](BrowseMediaPlayer.jpg)
-
-1. Browse to your media and select your file. 
-
-	![](BrowseToMovie.png)
-
-1. Once added, you can **click** the **Play** button inside the MediaPlayer to play your media. 
-
-	![](MediaPlay.png)
-
-1. Inside the **Details** panel, you may notice that a warning icon is present. 
-
-	![](EpicMovieMediaFrameWorkFile.jpg)
-
-1. In order to package correctly, local media files MUST be placed in the **Content/Movies/** directory of your project. 
-
-	![](EpicMovieMediaFrameWorkWMV.jpg)
-
-	Once you have placed your media in your projects Content/Movies/ folder, the warning icon will no longer be shown. 
-	
-**Importing through Drag-and-Drop**
-
-The easiest method of importing media is through drag-and-drop. 
-
-1. With your media inside your project's **Content/Movies/** folder, **drag-and-drop** it into the Content Browser. 
-
-	![](DragAndDropMedia.png)
-
-	This will automatically create the MediaPlayer asset and assign its location. 
-
-#### MediaPlayer Asset Details
-
-The MediaPlayer asset can output various content streams, including: audio, video, and caption texts. You can set these from the drop-down menus located along the top of 
-the media window below the playback controls. 
-
-![](OutputStreams.png)
-
-Inside the Details panel, you can inspect the properties of the media source. 
-
-![](EpicMovieMediaFrameWork.jpg)
-
-Under the **Playback** option, you have the opyion to enable looping of the media when playback reaches the end by checking the **Looping** checkbox. 
-
-For the **Source** section, you can specify the path to the media locally or through the Internet via URL, as well as set the Streaming Mode.  
-
-### MediaTexture Asset
-
-Once you have created a MediaPlayer asset, from it you can then create a **MediaTexture** asset which will allow you to extract a video stream. The MediaTexture allows 
-you to pick a video stream (if more than one is available) and you can use it just like any other texture in the Engine.
-
-To create a MediaTexture from a MediaPlayer, **Right-click** on a MediaPlayer asset then **select** the **Create Media Texture** option. 
-
-![](MediaTextureAsset.png)
-
-After which the new MediaTexture asset will be created. 
-
-![](MediaTextureCreated.png)
-
-From here, you can **Right-click** on the MediaTexture and create a **Material** from it that can be used like any other material in the Engine. 
-
-![](CreateMaterial.png)
-
-[REGION:note]
-If the video does not show on the texture but plays in the Media Asset Editor, then most likely the wrong video track is selected. You can change this by **Double-clicking** 
-the MediaTexture to open the texture editor and selecting a track via the Details panel under MediaPlayer. 
-
-If this does not solve the issue, inside your MediaPlayer asset, check the **Auto Play** option to see if the video is then streamed inside your texture. If that addresses 
-the issue, you can either continue to use **Auto Play** or through **Blueprints** set up some script to play the asset (see the 
-[Media Framework With Blueprints](#mediaframeworkwithblueprints) section on this page). 
-[/REGION]
-
-### MediaSoundWaves
-
-**MediaSoundWaves** are audio components that can be placed in your level as you would any other Sound Actor. Since it relies on the MediaPlayer Asset however, in order to
-play the sound, you would need to tell the MediaPlayer Asset to play (if it is not set to Auto Play). 
-
-You can create a MediaSoundWave by **Right-clicking** on a MediaPlayer asset and selecting **Create Media Sound Wave**. 
-
-![](MediaSoundWave.png)
-
-[REGION:note]
-MediaSoundWaves are still very early in development. Additional documentation on them will be updated at a later date. 
-[/REGION]
-
-## Media Framework with Textures & Materials
-
-This section will show you how to quickly get started with using the Media Framework with Textures and Materials. 
-
-### Textures
-
-Here is an example of a MediaTexture playing on a TV (click to play). 
-
-[OBJECT:EmbeddedVideo]
-[PARAMLITERAL:width]
-640
-[/PARAMLITERAL]
-[PARAMLITERAL:height]
-360
-[/PARAMLITERAL]
-[PARAMLITERAL:videoid]
-IdX1gEH2i6c
-[/PARAMLITERAL]
-[/OBJECT]
-
-_For this example, we **Migrated** the **SM_TV** asset from the **Realistic Rendering** project which can be found on the Learn Tab._
-
-1. Creating a MediaPlayer and MediaTexture asset with a desired video using the steps outlined on this page. 
-
-1. Open the MediaPlayer asset and **set** it to **Auto Play**. 
-
-1. Open up the 'M_TV' material located in the **Game/ArchVis/Materials** folder. 
-
-1. **Drag in** your MediaTexture into the 'M_TV' material and **place it** next to the **ScreenColor** node. 
-
-	![](Texture1.png)(w:800)
-
-1. **Delete** the ScreenColor node and **plug in** your texture to the Lerp node. 
-
-1. **Add** a **TextureCoordinate** node connected to the Texture Sample and adjust your 'UTiling' and 'VTiling' as needed. 
-
-	![](Texture2.png)
-
-
-### Materials 
-
-For Materials, all you need to do is create a Material from your MediaTexture then apply the material to an object. 
-
-![](MaterialsExample.png)
-
-Above, we have applied our material to a basic cube inside our level. Our MediaPlayer asset is also set to Auto Play. 
-
-
-## Media Framework with Blueprints
-
-You can also directly access and control the MediaPlayer asset through Blueprints. In this example, we have set a Trigger in front of our TV which has a MediaTexture applied
-to it. Our MediaPlayer asset is also not set to Auto Play. 
-
-![](Blueprint1.png)
-
-When the player enters the trigger, we will play the MediaTexture applied to the TV, when they exit the trigger we will pause it (click to play). 
-
-[OBJECT:EmbeddedVideo]
-[PARAMLITERAL:width]
-640
-[/PARAMLITERAL]
-[PARAMLITERAL:height]
-360
-[/PARAMLITERAL]
-[PARAMLITERAL:videoid]
-nmb2Fj_5PjA
-[/PARAMLITERAL]
-[/OBJECT]
-
-Inside our **Level Blueprint** we have created the following script. 
-
-[REGION:lightbox]
-[![](LevelBlueprintExample.png)(w:800)](levelBlueprintExample.png)
-[/REGION:lightbox]
-
-_Click for a full view._
-
-We have created a new variable of the MediaPlayer type and assigned our MediaPlayer asset as the Default Value. We have then stated that when the player character 
-(MyCharacter) overlaps the trigger box that we want to Play our MediaPlayer asset (which is assigned to the variable). 
-
-We have done the same thing for when the player character leaves the trigger volume, except we have used the Pause node. The **Event End Play** node is also used to 
-pause the MediaPlayer after exiting a Play in Editor session. These are just two of the nodes you can use with respect to MediaPlayers, if you drag off a MediaPlayer 
-variable, you will notice that there are several nodes that can be used with it. 
-
-![](BlueprintActions.png)
-
-
-## Media Framework with UMG
-
-If you want to have media playback as part of your UI or as an in-game movie player for example, you can do so with UMG by creating a **Widget Blueprint** and assigning 
-a material that is created from a MediaTexture (which in turn is created from a MediaPlayer asset). 
-
-Below we have created a setup so that when the player press the M key in-game, a movie player is shown on screen with controls to Play, Pause, or Rewind the video that 
-is being played (click to play). 
-
-[OBJECT:EmbeddedVideo]
-[PARAMLITERAL:width]
-640
-[/PARAMLITERAL]
-[PARAMLITERAL:height]
-360
-[/PARAMLITERAL]
-[PARAMLITERAL:videoid]
-HuF6btgkfxE
-[/PARAMLITERAL]
-[/OBJECT]
-
-We did this by first creating a **[Widget Blueprint](Engine\UMG\UserGuide\WidgetBlueprints)** as indicated below. 
-
-![](WidgetBlueprint.png)
-
-For the **Image** inside our Widget Blueprint, we have assigned our material that uses the MediaTexture. 
-
-![](BrushImage.png)
-
-Our Graph for the Widget Blueprint only has commands for our buttons we have placed to Play, Rewind, or Pause the video. 
-
-We have also added a MediaPlayer variable and assigned the MediaPlayer asset we want to use. 
-
-![](WidgetBlueprint2.png)(w:840)
-
-In order for the Material to be applied to our brush, we have to make one small change and set it to Emissive. 
-
-![](Material.png)
-
-In our Level Blueprint, we set up script for the player to enable/disable the in-game movie player we have created. 
-
-[REGION:lightbox]
-[![](LevelBlueprint.png)(w:840)](LevelBlueprint.png)
-[/REGION:lightbox]
-
-_Click for a full view._
-
-[REGION:note]
-Additional examples can be found inside the **Content Examples** and the **UMG** and **Media** maps. 
-[/REGION]
-
-
-
-
-
-
+## Essentials
 
+[DIR(output:"listbutton" parent:"Engine\MediaFramework" org:"hierarchy" end:"1")]
 
diff --git a/Engine/Documentation/Source/Engine/MediaFramework/Overview/MF_Overview.INT.udn b/Engine/Documentation/Source/Engine/MediaFramework/Overview/MF_Overview.INT.udn
new file mode 100644
--- /dev/null
+++ b/Engine/Documentation/Source/Engine/MediaFramework/Overview/MF_Overview.INT.udn
@@ -0,0 +1,181 @@
+Availability:Docs
+Title: Media Framework Overview
+Crumbs:%ROOT%, Engine
+Description:Overview of Media Framework functionality in Unreal Engine 4. 
+Navigation:topic
+parent:Engine/MediaFramework
+order:2
+Version: 4.13
+Related: Engine\MediaFramework\TechReference
+Related: Engine\MediaFramework\HowTo
+
+[TOC(start:2 end:2)]
+
+**Media Framework** uses several assets to allow the play back of video inside Unreal Engine 4. That video can be scrubbed, paused or rewound inside a **Media Player** asset as well as controlled through C++ or [](Engine/Blueprints). 
+Whether you want a montior in your level playing a video or you want to create a UI based interface with [](Engine/UMG) to allow players control over your video playback, you will need to specify the **Media Source** (where to find the media) to allow the Media Player the ability to play the media back. 
+
+## Media Source Types
+
+There are several different Media Source assets that you can use to define where the media is coming from. 
+
+![](MediaSource.png)
+
+Inside the **Content Browser** click the **Add New** Button, then under **Media** select the desired source asset: 
+
+[REGION:note]
+Please see the [](Engine\MediaFramework\TechReference\) page for information on supported file types and URLs. 
+[/REGION]
+
+### File Media Source
+
+![](MediaSource_File.png)
+
+This type of asset is used for media files stored on your device or in a shared local network file. The **File Path** section is where you specify where the media file resides. 
+While you can point to anywhere the file exists, however its typically a best practice to move your media files to the **Content/Movies** folder of your project and then point to the media file. 
+This will ensure that your media is packaged correctly with your project (you will also see the yellow exclamation point when pointing to a file outside of this location as a warning). 
+
+This type of source asset also allows you to load the entire media file into memory and play it from there by enabling the **Precache File** option. 
+Based on the size of your file, caching time may vary so take that into consideration when choosing this option. 
+
+Each of the source assets allow you to provide the name of the native media player with the **Default Player** field or leave it empty to find one automatically. 
+
+[REGION:note]
+Please see the [](Engine\MediaFramework\HowTo\FileMediaSource) how-to for working with File Media Source assets. 
+[/REGION]
+
+### Stream Media Source
+
+![](MediaSource_Stream.png)
+
+This asset allows you to specify a URL as the source of your media which is pulled from the internet. Currently linking to a YouTube or DailyMotion style URL is not supported and you will need to link directly to a hosted file, for example below is a link to a NASA test flight video: 
+
+* http://s3.amazonaws.com/akamai.netstorage/mp4/AirSTAR.mp4
+
+[REGION:note]
+Please see the [](Engine\MediaFramework\HowTo\StreamMediaSource) how-to for working with Stream Media Source assets. 
+[/REGION]
+
+### Platform Media Source
+
+![](MediaSource_Platform.png)
+
+The Platform Media Source asset allows overriding Media Sources on a per-platform basis. You can specify the Media Source to be used if no source was specified for a target platform as well as define the Media Sources per platform. 
+
+
+## Media Playlists
+
+After creating a few Media Source assets, you can add them to a **Media Playlist** which will play the assigned assets back sequentially. 
+
+![](MediaPlaylistButton.png)
+
+Select **Media Playlist** from the **Add New** button under the **Media** section. 
+
+![](CustomizedPlaylist.png)
+
+In the **Playlist** section you can define the source assets to include and playback order. 
+Any Media Source assets you have created will be displayed in the lower **Media Library** window which you can **Double-click** on to add to the end of the playlist. 
+
+[REGION:note]
+Please see the [](Engine\MediaFramework\HowTo\MediaPlaylists) how-to for working with Media Playlist assets. 
+[/REGION]
+
+## Media Sound Waves
+
+In order to hear audio associated with a video, you will need to create a **Media Sound Wave** asset.  
+
+![](MediaPlayerSound.png)
+
+From the **Add New** button under **Media**, select the **Media Sound Wave** asset. 
+
+![](MediaPlayerSoundWindow.png)
+
+The Media Sound Wave asset provides Debug, Concurrency and Attenuation settings that can be adjusted to define how the asset is processed. 
+When linked to a **Media Player** asset, audio attached to a video source will automatically play back along with the video. 
+Typically you can leave the default settings for a newly created Media Sound Wave asset, however if you need more control over how the sound is perceived you can adjust the Concurrency and Attenuation settings. 
+
+[REGION:note]
+Please see the [](Engine/Audio) section for more details on working with audio in Unreal Engine 4. 
+[/REGION]
+
+
+## Media Textures
+
+**Media Texture** assets are used for rendering video tracks from **Media Player** assets. 
+
+![](MediaTexture.png)
+
+While Media Sound Waves provide the audio, Media Textures provide the visuals from your media source assets. 
+Media Texture assets can be included within [](Engine/Rendering/Materials/) that can then be applied to meshes in your level such as a billboard, TV or monitor to create the appearance as if the video is playing on that object within the game world. 
+
+![](MediaTextureWindow.png)
+
+Above, a Material including a Media Texture asset has been created and applied to a static mesh in the level. 
+Inside the [Texture Editor](Engine/Content/Types/Textures/Properties/Interface) the Media Texture displays the same playback position as it appears in the Material in the level. 
+In addition to the standard [](Engine/Content/Types/Textures/Properties), Media Textures allow you to define the X and Y Axis **Addressing** to either Clamped, Mirror or World values. 
+
+
+
+## Media Player Assets
+
+Once you have a Media Source or Media Playlist, you can playback the asset using a **Media Player** asset. 
+
+![](MediaPlayer.png)
+
+From the **Add New** button under **Media**, select the **Media Player** asset. 
+
+The Media Player asset needs to have a **Media Sound Wave** and **Media Texture** assigned to it for it to properly play any video.
+If you have not created one, when you create any new Media Player asset a dialogue window will appear asking if would like to create and automatically link those assets with the newly created Media Player asset. 
+
+![](CreateMediaPlayer.png)
+
+When opening a Media Player asset, the Media Player Window will appear. 
+
+![](MediaPlayerWindow.png)
+
+Inside the Media Player Window you can play back and preview any media source assets you have created (created source media are displayed the **Media Library** panel). 
+You can play, pause, rewind or skip to the beginning or end of a video from the toolbar across the top. Any Media Playlists that are currently playing videos will be displayed in the upper right **Playlist** window. 
+
+In the lower-right **Details** panel you can define options for **Playback** or assign the **Sound Wave** and **Video Textures** (or an **Image Texture**) associated with the Media Player. 
+There is also an **Info** panel which will display info about the source such as the stream type, codecs, sample rate, dimensions and more. 
+
+
+## Media Assets & Blueprints
+
+After your Media Player asset is set up and linked with the appropriate assets, in order to play back the media during game play you will need to first **Open** the source asset through Blueprint or C++. 
+Before you can do that however, you will need a reference to your Media Player asset which you can create a [Variable](Engine/Blueprints/UserGuide/Variables/) for inside any Blueprint. 
+
+![](MediaPlayerDetailsSet.png)
+
+With a reference to your Media Player defined, you can then call an **Open** function based on the type of media source. 
+
+![](MediaPlayerOpen.png)
+
+| Option | Description |
+| -------- | ----------- |
+| **Open Playlist** | Opens the first media source in the specified play list. |
+| **Open Playlist Index** | Opens a specific media source within a specified play list. |
+| **Open Source** | Opens the specified media source (File Media, Stream Media or Platform Media). |
+| **Open URL** | Opens the specified media URL. |
+| **Reopen** | Reopens the currently opened media or play list. |
+
+Once you open the source you can then **Play** the media. Newly created Media Player assets have the **Play on Open** option enabled by default which means when you open a media source it will start playing the video automatically. 
+If you uncheck the Play on Open option (pictured below), you will then need to call the **Play** function on a media source after it has been opened. 
+
+![](UncheckPlayOnOpen.png)
+
+Below on Event Begin Play we open a Media Playlist, however we only start play back when the **Right Mouse Button** is pressed. 
+
+![](PlayWhenPressed.png)
+
+You can call other functions from a Media Player reference such as [Binding Events](Engine/Blueprints/UserGuide/EventDispatcher/BindingAndUnbinding/) to various states of playback (such if the video is paused or if the video ends). 
+You can also get or set options like if a video is playing or can play, if a video can loop, or play rate information. With a Media Player reference, you can drag off and see these options under the **Media Player** section of the Blueprint Context Menu. 
+
+![](MediaPlayerNodes.png)
+
+
+
+
+
+
+
+
diff --git a/Engine/Documentation/Source/Engine/MediaFramework/TechReference/MF_TechRef.INT.udn b/Engine/Documentation/Source/Engine/MediaFramework/TechReference/MF_TechRef.INT.udn
new file mode 100644
--- /dev/null
+++ b/Engine/Documentation/Source/Engine/MediaFramework/TechReference/MF_TechRef.INT.udn
@@ -0,0 +1,161 @@
+Availability:Docs
+Title: Media Framework Technical Reference
+Crumbs:%ROOT%, Engine
+Description:Technical information regarding Media Framework in Unreal Engine 4.
+Navigation:topic
+parent:Engine/MediaFramework
+order:2
+Version: 4.13
+Related: Engine\MediaFramework\Overview
+Related: Engine\MediaFramework\HowTo
+
+[TOC(start:2 end:3)]
+
+This page outlines the currently supported file extensions and URLs as well as outlines some troubleshooting tips on how to enable Media Framework with different platforms. 
+There are some known platform specific issues and limitations which we are currently investigating and trying to resolve, each of which have been identified in the Troubleshooting Platform sections. 
+
+
+## Supported Files & URLs
+
+The table below indicates the supported extensions based on each media player platform. 
+
+| | Android Media Player | AVF Media Player | PS4 Media Player | WMF Media Player | Xbox One Media Player |
+| --- | --- | --- |--- | --- | --- |
+| **.3g2** | No | No | No | Yes | No |
+| **.3gp** | No | No | No | Yes | No |
+| **.3gp2** | No | No | No | Yes | No |
+| **.3gpp** | Yes | No | No | Yes | No |
+| **.aac** | Yes | No | No | Yes | No |
+| **.adts** | No | No | No | Yes | No |
+| **.asf** | No | No | No | Yes | No |
+| **.avi** | No | Yes | No | Yes | No |
+| **.m4a** | No | No | No | Yes | No |
+| **.m4v** | No | Yes | No | Yes | No |
+| **.mov** | No | Yes | No | Yes | No |
+| **.mp3** | No | No | No | Yes | No |
+| **.mp4** | Yes | Yes | Yes | Yes | Yes |
+| **.sami** | No | No | No | Yes | No |
+| **.smi** | No | No | No | Yes | No |
+| **.wav** | No | No | No | Yes | No |
+| **.wma** | No | No | No | Yes | No |
+| **.wmv** | No | No | No | Yes | No |
+
+This table shows the supported URLs based on media player platform. 
+
+| | Android Media Player | AVF Media Player | PS4 Media Player | WMF Media Player | Xbox One Media Player |
+| --- | --- | --- |--- | --- | --- |
+| **file://** | Yes | Yes | Yes | Yes | Yes |
+| **http://** | Yes | No | Yes | Yes | No |
+| **httpd://** | Yes | No | No | Yes | No |
+| **https://** | Yes | No | Yes | Yes | No |
+| **mms://** | Yes | No | No | Yes | No |
+| **rtsp://** | Yes | No | No | Yes | No |
+| **rtspt://** | Yes | No | No | Yes | No |
+| **rtspu://** | Yes | Yes | No | Yes | No |
+
+
+## Troubleshooting
+
+Most player plug-ins will provide detailed logging of potential problems. To enable this, add the following to your projects **Engine.ini** file: 
+
+	[Core.Log]
+	LogMedia=Verbose
+	LogMediaAssets=Verbose
+	LogAndroidMedia=Verbose
+	LogAvfMedia=Verbose
+	LogMfMedia=Verbose
+	LogPS4Media=Verbose
+	LogWmfMedia=Verbose
+
+If a media source fails to open or play, check the **Output Log** under the [Developer Tools](https://docs.unrealengine.com/latest/INT/Engine/Tools/DevTools/) for additional details. 
+
+[REGION:raw]
+![](LogMessage.png)
+[/REGION]
+
+You can also check information about a Media Source inside the Media Player asset and **Info** tab. 
+
+![](InfoWindow.png)
+
+If **Verbose Logging** is enabled for the Media Assets log category, this information is also added to the Output Log. 
+
+### PlayStation 4
+
+The PS4Media player plug-in is disabled by default and must be enabled in the project’s [Plugin Settings](Programming/Plugins/).
+
+The video decoder on PS4 consumes a lot of Onion memory, especially with HLS. You may run out of memory when trying to play more than one video at the same time, because we allocate only 200 MB by default. 
+Currently, these default settings are not exposed in the UI but you can modify them by hand.
+
+To enable, add the following to the **PS4Engine.ini**: 
+
+	[SystemSettings]
+	r.GarlicHeapSizeMB=2600
+	r.OnionHeapSizeMB=200
+
+In your game's ** *.Target.cs**:
+
+	public override void SetupGlobalEnvironment(
+            	TargetInfo Target,
+            	ref LinkEnvironmentConfiguration OutLinkEnvironmentConfiguration,
+            	ref CPPEnvironmentConfiguration OutCPPEnvironmentConfiguration)
+	{
+            	// ...
+
+            	if (Target.Platform == UnrealTargetPlatform.PS4)
+            	{
+	OutCPPEnvironmentConfiguration.Definitions
+	.Add("GARLIC_HEAP_SIZE= (2900ULL * 1024 * 1024)");
+
+	OutCPPEnvironmentConfiguration.Definitions
+	.Add("ONION_HEAP_SIZE= (100ULL * 1024 * 1024)");
+            	}
+
+            	// ...
+	}
+
+### Xbox One
+
+The MfMedia player plug-in is disabled by default and must be enabled in the project’s Plugin settings.
+It is possible to compile and test the MfMedia plug-in on Windows 7 and up if the Engine is compiled against Windows 7 (by default we compile against Windows Vista). 
+
+This can be accomplished by changing the following code in the `SetUpCompileEnvironment()` function in **UEBuildWindows.cs**:
+
+	// Windows Vista or higher required	
+	InBuildTarget.GlobalCompileEnvironment.Config.Definitions.Add("_WIN32_WINNT=0x0601");	
+	InBuildTarget.GlobalCompileEnvironment.Config.Definitions.Add("WINVER=0x0601");
+
+### Windows
+
+The Windows Media Foundation (WMF) based player plug-in WmfMedia has a number of limitations for MP4 containers. 
+
+[REGION:note]
+See Microsoft's [MPEG-4 File Source](https://msdn.microsoft.com/en-us/library/windows/desktop/dd757766%28v=vs.85%29.aspx) page for more information.
+[/REGION]
+
+
+## Known Issues & Limitations
+
+**Media Framework** is still considered experimental and there are some known issues and limitations which are outlined below. 
+
+* **Android**
+	* Audio currently plays via the OS and cannot be piped through UE4's Sound System. 
+* **Audio**
+	* There is currently no synchronization of audio samples, so playback may get out of sync. 
+		*	Workaround: Use shorter clips where possible. 
+	* Audio is currently always mixed down to 16-bit stereo 44.1 kHz, regardless of the input format. 
+		*	Multi-channel audio support is being investigated and may arrive in future releases.
+* **Editor**
+	* The Playlist UI inside the Media Player only shows the currently opened Media Playlist asset. 
+		*	Workaround: Edit the Playlist contents by **Double-clicking** on the MediaPlaylist asset to open its own editor. 
+* **Playstation 4**
+	* Playlists currently do not work.
+* **Sequencer**
+	* Media Playback will not render out properly when recording cinematics from Sequencer. 
+		* We have this scheduled to release in the near future as we know this is important.
+* **Windows**
+	* The Media Player state machine is still in progress and may break if multiple control commands are issued simultaneously. 
+		*	Workaround: Insert short delays between commands, so that the player has time to process them.
+	* Some videos with multiple audio tracks fail to open.
+	* QuickTime movie (.mov) support on Windows 7 and later is currently unreliable. 
+		*	Workaround: This format is currently not recommended for use. 
+
