#ue4_docs
Apply SME feedback.

diff --git a/Engine/Documentation/Source/Engine/Sequencer/HowTo/ImportExportEDLs/ImportExportEDLs.INT.udn b/Engine/Documentation/Source/Engine/Sequencer/HowTo/ImportExportEDLs/ImportExportEDLs.INT.udn
--- a/Engine/Documentation/Source/Engine/Sequencer/HowTo/ImportExportEDLs/ImportExportEDLs.INT.udn
+++ b/Engine/Documentation/Source/Engine/Sequencer/HowTo/ImportExportEDLs/ImportExportEDLs.INT.udn
@@ -44,6 +44,11 @@ We are also using Adobe Premiere Pro CS6 as our external editing software; howev
 	Inside the settings window, we create a new folder to export the cinematic to and prefix our {shot} token with the name of our cinematic. 
 	By using the {shot} token each of our shots will be exported as individual video files along with the full cinematic. 
 
+	[REGION:note]
+	This will only work when the **Output Format** is set to **Video Sequence**. It does not work with individual frames. 
+	[/REGION]
+
+
 1.	Inside your software editing suite, **Import** the **EDL** file exported from the previous step. 
 
 	[OBJECT:EmbeddedVideo]
