#UE4 Docs: Adding a youtube video that shows off texture swiming.

diff --git a/Engine/Documentation/Source/Enterprise/Automotive/AutomotiveMaterialPack/AutomotiveMaterialPack.INT.udn b/Engine/Documentation/Source/Enterprise/Automotive/AutomotiveMaterialPack/AutomotiveMaterialPack.INT.udn
--- a/Engine/Documentation/Source/Enterprise/Automotive/AutomotiveMaterialPack/AutomotiveMaterialPack.INT.udn
+++ b/Engine/Documentation/Source/Enterprise/Automotive/AutomotiveMaterialPack/AutomotiveMaterialPack.INT.udn
@@ -1,4 +1,4 @@
-Availability:Doc
+Availability:Public
 Title: Automotive Material Pack
 Description:A breakdown of what assets the Automotive Material pack contains and how to use them.
 Crumbs:%ROOT%,
@@ -16,7 +16,7 @@ Related:Engine/UI/LevelEditor/Modes/MeshPaintMode/VertexColor
 
 ![](AP_Intro.png)
 
-The Automotive Material Pack is a collection of automotive-themed PBR Materials and Textures that have been setup for use in Unreal Engine 4(UE4). The pack features 21 Master Materials and 75 Material Instances taken from Epic's 2016 GDC Mclaren 570S demo. In the following document, we will go over how to use the contents of the Automotive Material Pack in your UE4 projects as well as anything that you need to be aware of while using the content in your project.
+The Automotive Material Pack is a collection of automotive-themed PBR Materials and Textures that have been setup for use in Unreal Engine 4(UE4). The pack features 21 Master Materials and 75 Material Instances taken from Epic's 2016 GDC Mclaren 570S demo. In the following document we will take a look at what Materials the Automotive Material Pack contains and any special information you might need to know in order to use these Materials in your UE4 projects.
 
 ## What's in the Automotive Material Pack
 
@@ -25,11 +25,11 @@ Inside of the Automotive Material Pack you will find the following Materials typ
 * Car Paint
 * Carbon Fiber
 * Caliper
-* Wheel
+* Rim
 * Brake Rotor
 * Exhaust
 * Reflector
-* Tire(Rubber)
+* Tire (Rubber)
 * Metal (Anodized/Dirt)
 * Leather (Perforated)
 * Alcantara
@@ -51,7 +51,7 @@ In the following section, we will take a look at some of the parameters that all
 
 * **Material Scale:** Material Scale gives you control over the Scaling / Tiling of a Texture or pattern that is used in the Material.
 
-* **Scale-Aware Tiling:** Scale- Aware Tiling scales the Materials UV's up or down based on the World Size of the Static Mesh that it is applied to. The result is a uniform Material scale on all of the objects regardless of their current size. Please note that this feature will only work with objects which have UV's channels set up.
+* **Scale-Aware Tiling:** Scale- Aware Tiling scales the Materials UV's up or down based on the World Size of the Static Mesh that it is applied to resulting in a uniform Material scale on all objects the Material is applied to regardless of their current size or scale settings. 
 
 	[OBJECT:ComparisonSlider]
 		[PARAM:before]
@@ -62,7 +62,7 @@ In the following section, we will take a look at some of the parameters that all
 		[/PARAM]
 	[/OBJECT]
 
-* **World Aligned:** This feature should be used on objects that have incorrect or poorly setup UV's. Please note that this is best used on Static objects as moving objects will exhibit texture 'swimming'.
+* **World Aligned:** World Aligned allows objects that have incorrect or poorly setup UV's to correctly display the applied Material.
 
 	[OBJECT:ComparisonSlider]
 		[PARAM:before]
@@ -73,29 +73,34 @@ In the following section, we will take a look at some of the parameters that all
 		[/PARAM]
 	[/OBJECT]
 
-## Orange Peel Material Parameter
+	[region:note]
+	Please note that World Aligned UV's work best when used on static objects. When World Aligned UV's are applied to objects that move the Texture will  exhibit texture 'swimming' like in the video below.
+	[OBJECT:EmbeddedVideo]
+		[PARAMLITERAL:width]
+		640
+		[/PARAMLITERAL]
+		[PARAMLITERAL:height]
+		360
+		[/PARAMLITERAL]
+		[PARAMLITERAL:videoid]
+		MMi8_J89Vqo
+		[/PARAMLITERAL]
+	[/OBJECT]
+	[/region]
+
 
-The clear coat lacquer applied on top of automotive paint often has a slightly bumpy surface - not unlike the surface of an orange and is the reason why this is referred to as Orange Peel. You can control the intensity of this effect by using the **Orange Peel Strength** and **Orange Peel Scale** parameters.To disable this effect, set the **Orange Peel Strength** to **0**.
 
-[OBJECT:ComparisonSlider]
-	[PARAM:before]
-	![Orange Peel Off](AP_03_Carpaint_03_OrangPeel_OFF.png)
-	[/PARAM]
-	[PARAM:after]
-	![Orange Peel On](AP_03_Carpaint_04_OrangPeel_ON.png)
-	[/PARAM]
-[/OBJECT]
 
 ## Car Paint Material
 
-The Car Paint Material is a complex Material that comes with 29 adjustable parameters. When trying to modify this Material for your project, we recommend starting with the following parameters:
+The Car Paint Material is a very robust Material that has 29 unique parameters that can be adjusted. When trying to modify this Material for your project, we recommend you starting adjusting and experimenting with the following parameters:
 
 * **Base Color/Flake**
 * **Scale/Flake**
 * **Normal Scale/Flake**
 * **Normal Strength/Parameters**
  
-By adjusting just these four different parameters you can achieve vastly different looks for your car paint Material. Here are two examples of what can be achieved by changing those four parameters:
+By adjusting just these four different parameters you can achieve vastly different looks for your car paint Material. Here are two examples of what can be achieved by changing just those four parameters:
 
 [OBJECT:ComparisonSlider]
 	[PARAM:before]
@@ -106,20 +111,43 @@ By adjusting just these four different parameters you can achieve vastly differe
 	[/PARAM]
 [/OBJECT]
 
+## Orange Peel Material Parameter
+
+The Clearcoat lacquer that is applied on top of automotive paint often has a slightly bumpy surface similar to the peel of an Orange. Because of this resemblance this layer of lacquer is referred to as Orange Peel. You can control the intensity of the Orange Peel using the **Orange Peel Strength** and **Orange Peel Scale** parameters. To disable the Orange Peel effect, set the **Orange Peel Strength** to **0**.
+
+[OBJECT:ComparisonSlider]
+	[PARAM:before]
+	![Orange Peel Off](AP_03_Carpaint_03_OrangPeel_OFF.png)
+	[/PARAM]
+	[PARAM:after]
+	![Orange Peel On](AP_03_Carpaint_04_OrangPeel_ON.png)
+	[/PARAM]
+[/OBJECT]
+
 ## Carbon Fiber Material
 
-The Carbon Fiber Material makes use of the new Dual Normals option that is in the Clearcoat Material. This new Dual Normal allows you to add a second Normal map that allows the bottom layer of the Clearcoat Material to interact with light. Without this feature, the glossy effect of the Carbon Fiber Material is lost, making the Carbon Fiber look not as realistic as it could be. 
+The Carbon Fiber Material makes use of the new Dual Normal option that is available when using a  Clearcoat Material. The Dual Normal option allows you to add a 2nd Normal map to the Clearcoat Material so that both the bottom and top layers of the Material will react to changes to lighting. Without this feature, the glossy effect of the Carbon Fiber Material is lost, making the Carbon Fiber look not as realistic as it could be. 
 
-![](AP_04_CarbonFiber.png)
+[OBJECT:ComparisonSlider]
+	[PARAM:before]
+	![Dual Normal On](AP_04_CarbonFiber.png)
+	[/PARAM]
+	[PARAM:after]
+	![Dual Normal Off](AP_04_CarbonFiber_No_2nd_Normal.png)
+	[/PARAM]
+[/OBJECT]
+
+[region:note]
+For more information on using Dual Normals with a Clearcoat Material make sure to check out the [Dual Normals](Engine/Rendering/Materials/HowTo/ClearCoatDualNormal) Material page.
+[/region]
 
 ## Caliper Material
 
-The caliper is part of a car's braking system and is usually placed on the brake disc in between the wheel and the brake pad. 
+The caliper is part of a car's braking system and is usually placed on the brake disc in between the wheel rim and the brake pad. 
 
 ![](AP_05_Caliper_01.png)
 
-
-The brake caliper Material has the following two different options. 
+When using the brake caliper Material you can adjust the following two parameters.
 
 * **Painted Caliper:** Plain color.
 * **Painted Caliper with Metallic Flakes:** Silver color with metallic flakes. 
@@ -135,7 +163,7 @@ The brake caliper Material has the following two different options.
 
 ## Rim Material
 
-There are five different materials that can be used for the rims of the wheels including a Two Tone option. When using the Two Tone Rim keep in mind that the Two Tone Rim is made up of two different Materials with one Material being blue-black and the other material is silver. The **Red Channel** of the Vertex Color on the Rim Static Mesh is used to define which areas of the wheel should be silver. Any areas that are not using red Vertex color should use black as the Vertex Color so that they will display the 2nd Material.
+There are five different Materials that can be used for the rims of the wheels including a Two Tone option. When using the Two Tone Rim keep in mind that the Two Tone Rim is made up of two different Materials with one Material being blue-black and the other material is silver. The **Red Channel** of the Vertex Color on the Rim Static Mesh is used to define which areas of the wheel should be silver. Any areas that are not using red Vertex color should use black as the Vertex Color so that they will display the 2nd Material.
 
 ![](AP_06_Rim_Ball_01.png)
 
@@ -149,11 +177,12 @@ If you are unfamiliar with how to use Vertex Paint in UE4 please check out the [
 
 ## Brake Rotor Material
 
-Since the brake rotor, Material uses a radial pattern Texture it needs to be applied to a mesh using a planar UV projection method or it will not be displayed correctly. The image below shows an example of how the UVs should look like when a Texture is applied to the brake rotor.
+Since the brake rotor Material uses a radial Texture pattern the UV's of the brake rotor need to use planar UV project mapping to display the brake rotor Material correctly. Failing to setup the meshes UVs in this manner will result in the brake rotor Material not being displayed correctly when applied to the brake rotor mesh.  The image below shows an example of how a Texture should look when placed on an object that has planar projection mapping.
 
 ![](AP_07_BrakeRotor_01.png)
 
-The following images show examples of the types of brake rotor Materials that are available to use.
+The following images show examples of the different types of brake rotor Materials that are available for use.
+
 [OBJECT:ComparisonSlider]
 	[PARAM:before]
 	![Carbon Ceramic](AP_07_BrakeRotor_02_CarbonCeramic.png)
@@ -165,7 +194,7 @@ The following images show examples of the types of brake rotor Materials that ar
 
 ## Reflector Material
 
-Inside the Automotive Material Pack you will find the following three Reflector Patterns presets that are ready to be used;
+The reflector Material is used for the plastic  reflectors you find next to turn singles and on lights.  When looking at the reflector Material you will find that it comes with the following three reflector pattern presets:
 
 * **Hexagon**
 * **Honeycomb**
@@ -173,13 +202,9 @@ Inside the Automotive Material Pack you will find the following three Reflector
 
 	![](AP_08_Reflector_02.png)
 
-To add or change the reflector patterns open up the **Reflector** Material Instance and change the **Reflector Normal** to the reflector pattern you want to use. You can also create your own Reflector Normal Maps to use by doing the following.
-
-1. Generate a Normal Map of the item you want to use a reflector.
-1. Save that image as a .TGA and then import it into UE4.
-1. Once imported, select the Texture in the Content Browser and then inside of the Reflector Material Instance under the Reflector (Bottom Layer) Normal section press the **white arrow** to load the selected Texture.
+To change which reflector pattern is used,  open up the **Reflector** Material Instance and change the **Reflector(Bottom Layer) Normal** to the reflector pattern you want to use.
 
-	![](AP_08_Reflector_01.jpg)
+![](AP_08_Reflector_01.png)
 
 ## Rubber / Tire Material
 
@@ -188,13 +213,17 @@ The rubber Material that comes with the Automotive Material Pack can be used for
 * **Door Seals**
 * **Floor Mats**
 * **Interior Pieces**
-* Tires**
+* **Tires**
 
-The rubber Material has a tire-specific option called **Use Sidewall**. Enabling this part of the Material will add a Normal Map Texture that will display the sidewall lettering that can be found on tires. When using the sidewall option for a tire make sure that the UV for the tire sidewall 3D Mesh lines up with the Normal Map.
+When looking at the rubber Material Instance you will notice that it has a tire-specific option called **Use Sidewall**. Enabling this option will add a new Normal Map Texture that is for display the lettering that can be found on the tire sidewalls. 
 
 ![](AP_Tire_UV_00.png)
 
-If you fail to properly line up the UV of the side wall to the Normal of the tire things like the lettering and logo on the sidewall will not show up correctly. In the following images 
+[region:note]
+Keep in mind that when using the sidewall option for a tire you will need to make sure that the UV for the tire sidewall mesh line up with the Normal Map.
+[/region]
+
+Failing to properly line up the UV of the tire sidewall to the Normal map of the tire will result in the tire lettering and logos not show up correctly on the tire sidewall like in the image below. 
 
 [OBJECT:ComparisonSlider]
 	[PARAM:before]
@@ -207,7 +236,7 @@ If you fail to properly line up the UV of the side wall to the Normal of the tir
 
 ## Leather / Perforated Material
 
-The leather Material that comes with the Automotive Material Pack has 12 different color leather preset to choose from as well as five different perforated leather versions. 
+The leather Material that comes with the Automotive Material Pack has 12 different preset leather colors to choose from and also contains five different perforated leather presets. 
 
 [OBJECT:ComparisonSlider]
 	[PARAM:before]
@@ -220,13 +249,13 @@ The leather Material that comes with the Automotive Material Pack has 12 differe
 
 ## Alcantara Material
 
-The alcantara Material that comes with the Automotive Material Pack can be used to mainpulate of fabric. When using the Alcantara Material you can adjust the following parameters to mainpulate the look of the fabric.
+The alcantara Material that comes with the Automotive Material Pack is used to display the fabrics you find in car's. When using the alcantara Material you can adjust the following parameters to manipulate the look of the fabric.
 
 * **Fabric Scale / Strength**
 * **Cloth Folds Scale / Strength**
 * **Fabric Detail Scale / Strength**
 
-In the following set of images, you can see the effect that the Cloth Folds parameter has on the look of the Alcantara fabric Material.
+In the following set of images, you can see the effect that the **Cloth Folds** parameter has on the look of the alcantara fabric Material.
 
 [OBJECT:ComparisonSlider]
 	[PARAM:before]
@@ -239,17 +268,16 @@ In the following set of images, you can see the effect that the Cloth Folds para
 
 ## Glass Material
 
-The glass Material that comes with the Automotive Material Pack has the following two reflections types that allow the glass to use two different reflections so the Material can reflect objects that are outside of the car as well as reflect objects the inside of the car.
+The glass Material that comes with the Automotive Material Pack has the ability to display two completely different reflections in one Material. With this Material you can select what the exterior glass should reflect as well as what the interior glass should reflect. 
 
-* **Outer Reflection:** Helps to reflect the scene around the exterior of the glass in real time. 
+* **Outer Reflection:** Helps to reflect the scene around the exterior of the glass. 
 
 	![](AP_12_Glass_06_Ex.png)
 
-* **Inner Reflection:** Helps to reflect the scene around the interor of the glass using a predefined Cubemap.
+* **Inner Reflection:** Helps to reflect the scene around the interor of the glass.
 
 	![](AP_12_Glass_07_Ex.png)
 
-
 Here is an example of what the glass Material looks like with the inner reflections enabled and without the inner reflection enabled. 
 
 [OBJECT:ComparisonSlider]
@@ -261,7 +289,7 @@ Here is an example of what the glass Material looks like with the inner reflecti
 	[/PARAM]
 [/OBJECT]
 
-For the car windshield, it would be best to capture a Cubemap from the inside of the car and then input that Cubemap into the glass's inner reflection Cubemap like in the following image.
+When using the inner reflections for a car windshield the best workflow is to capture a Cubemap from the inside of the car. Then in the glass Material Instance, enable the **Use Glass Inner Reflections Cubemap** option and then input the Cubemap you want to use into the **Cubemap** input.
 
 ![](AP_12_Glass_01.png)
 
