Addressing DmitryR's SME review comments

#UE4doc

diff --git a/Engine/Documentation/Source/Platforms/Linux/BeginnerLinuxDeveloper/SettingUpAnUnrealWorkflow/2/SettingUpAnUnrealWorkflow_2.INT.udn b/Engine/Documentation/Source/Platforms/Linux/BeginnerLinuxDeveloper/SettingUpAnUnrealWorkflow/2/SettingUpAnUnrealWorkflow_2.INT.udn
--- a/Engine/Documentation/Source/Platforms/Linux/BeginnerLinuxDeveloper/SettingUpAnUnrealWorkflow/2/SettingUpAnUnrealWorkflow_2.INT.udn
+++ b/Engine/Documentation/Source/Platforms/Linux/BeginnerLinuxDeveloper/SettingUpAnUnrealWorkflow/2/SettingUpAnUnrealWorkflow_2.INT.udn
@@ -27,13 +27,13 @@ SkillLevel:Beginner
 
 %Steps%
 
-Now that you've setup Git on your Linux machine, you're ready to download the Unreal Engine source code from our [GitHub site](http://github.com/EpicGames). There 
-are a couple of ways to download the UE4 source code, one way is to use the download feature on our GitHub page and the other way is to fork and clone the Unreal 
+Now that you've set up Git on your Linux machine, you're ready to download the Unreal Engine source code from our [GitHub site](http://github.com/EpicGames). There 
+are a couple of ways to download the UE4 source code, one way is to use the download feature on our GitHub page, and the other way is to clone the Unreal 
 Engine repository with your Git client. We'll cover both ways on this page.
 
 ## Downloading from GitHub
 
-If don't want to fork and clone the Unreal Engine repository using Git, you can use the **Download ZIP** button on our [GitHub site](http://github.com/EpicGames).
+If don't want to clone the Unreal Engine repository with Git, you can use the **Download ZIP** button on our [GitHub site](http://github.com/EpicGames).
 
 1.	First, select the branch you want to download by clicking the **Branch:** button on the left side of the page.
 
@@ -49,47 +49,39 @@ If don't want to fork and clone the Unreal Engine repository using Git, you can
 
 	![](DownloadSourceButton.png)
 
-1.	Finally, decompress 'UnrealEngine-(branch_name).zip' onto your hard disk.
+1.	Finally, decompress `UnrealEngine-(branch_name).zip` onto your hard disk.
 
-## Forking and Cloning with Git
+## Cloning with Git
 
-If you're new to forking and cloning with Git, refer to GitHub's [Forking Projects Guide](https://guides.github.com/activities/forking/) before executing the following steps.
+If you're new to cloning with Git, refer to [GitHub's Cloning Guide](https://help.github.com/articles/cloning-a-repository/#platform-linux) before executing the following steps.
 
 1.	Click on the **Unreal Engine** link to access our [repository](https://github.com/EpicGames/UnrealEngine).
 	
 	![](UnrealRepo.png)
 	
-1.	Click the **Fork** button to make a fork of our repository.
+1.	Click the **Clone or download** button.
 
-	![](Fork.png)
-
-1.	Navigate to the fork from your GitHub account.
-
-	![](NavigateToFork.png)
-
-1.	Now, click the **Clone or download** button.
-
-	![](ForkCloneOrDownloadButton.png)
+	![](CloneorDownloadButton.png)
 
 1.	Click the **Copy URL** button to copy the clone URL for the repository.
 
 	![](CopyHTTPS.png)
 
-1.	Open Git Bash, type `git clone`, and paste the URL you copied in the previous step, making sure to add your GitHub username rather than `YOUR-USERNAME`.
+1.	Open Git Bash, type `git clone`, and paste the URL you copied in the previous step.
 
 	![](GitBashClone.png)
 
 1.	Finally, press **Enter** to create your local clone.
 
 	[REGION:tip]
-	If you need a refresher on using Git from the command line, refer to GitHub's articles on [Git](https://help.github.com/articles/set-up-git/) and [Forking a Repository](https://help.github.com/articles/fork-a-repo/).
+	If you need a refresher using Git from the command line, refer to [GitHub's documentation](https://help.github.com/articles/set-up-git/#platform-linux).
 	[/REGION]
 
 ## End Result
 
 ![](end_result.png)
 
-Whether you downloaded the Unreal Engine source code by using the **Download ZIP** button on our [GitHub site](http://github.com/EpicGames), or by forking and cloning the Unreal Engine repository 
-using Git, you are now ready to build UE4 on Linux with the Unreal Engine source code that is now located on your hard disk.
+Whether you downloaded the Unreal Engine source code by using the **Download ZIP** button on our [GitHub site](http://github.com/EpicGames), or by cloning the Unreal Engine repository 
+using Git, you're ready to build UE4 on Linux with the Unreal Engine source code that is now located on your hard disk.
 
 %Steps%
\ No newline at end of file
diff --git a/Engine/Documentation/Source/Platforms/Linux/BeginnerLinuxDeveloper/SettingUpAnUnrealWorkflow/4/SettingUpAnUnrealWorkflow_4.INT.udn b/Engine/Documentation/Source/Platforms/Linux/BeginnerLinuxDeveloper/SettingUpAnUnrealWorkflow/4/SettingUpAnUnrealWorkflow_4.INT.udn
--- a/Engine/Documentation/Source/Platforms/Linux/BeginnerLinuxDeveloper/SettingUpAnUnrealWorkflow/4/SettingUpAnUnrealWorkflow_4.INT.udn
+++ b/Engine/Documentation/Source/Platforms/Linux/BeginnerLinuxDeveloper/SettingUpAnUnrealWorkflow/4/SettingUpAnUnrealWorkflow_4.INT.udn
@@ -33,6 +33,18 @@ games for the Linux platform.
 
 [INCLUDE:Programming/Development/BuildingUnrealEngine#linuxrun]
 
+## Launching UE4 on Linux
+
+Because you've compiled the engine from its source code on Linux, the engine will compile shaders for a few minutes after launching UE4.
+
+![](howTo_comilingshaders.png)
+
+## Creating your First Project
+
+If you haven't already set up an Integrated Development Environment (IDE) to work with UE4, we recommend that you start out by creating a **Blueprint** project. 
+
+![](howTo_creatingprrjwithoutIDE.png)
+
 ## End Result
 
 ![](howTo_endResult.png)
@@ -40,7 +52,7 @@ games for the Linux platform.
 If you're creating your first game on Linux, you just learned how to establish your primary UE4 workflow on Linux. As you worked your way through this tutorial, 
 you learned:
 
-&#x2713; How to setup Git to download the source code for Unreal Engine.  
+&#x2713; How to set up Git to download the source code for Unreal Engine.  
 &#x2713; How to fork and clone Unreal Engine's source code.  
 &#x2713; How to build Unreal Engine from source.  
 &#x2713; How to run Unreal Engine on your Linux machine.  
diff --git a/Engine/Documentation/Source/Platforms/Linux/BeginnerLinuxDeveloper/SettingUpAnUnrealWorkflow/5/SettingUpAnUnrealWorkflow_5.INT.udn b/Engine/Documentation/Source/Platforms/Linux/BeginnerLinuxDeveloper/SettingUpAnUnrealWorkflow/5/SettingUpAnUnrealWorkflow_5.INT.udn
--- a/Engine/Documentation/Source/Platforms/Linux/BeginnerLinuxDeveloper/SettingUpAnUnrealWorkflow/5/SettingUpAnUnrealWorkflow_5.INT.udn
+++ b/Engine/Documentation/Source/Platforms/Linux/BeginnerLinuxDeveloper/SettingUpAnUnrealWorkflow/5/SettingUpAnUnrealWorkflow_5.INT.udn
@@ -36,8 +36,8 @@ and head over to the [Unreal Engine Public Issues Tracker](https://issues.unreal
 
 ![](UE4_publicJira.png)
 
-Do you see a public issue that you can fix? If you know how to submit [pull requests](https://help.github.com/articles/using-pull-requests/) with Git, join our growing  
-community of developers and help shape the future of UE4 on Linux.
+Do you see a public issue that you can fix? If you know how [fork and clone](GettingStarted/DownloadingUnrealEngine/GitHub) our [repository](https://github.com/EpicGames/UnrealEngine) 
+to submit [pull requests](https://help.github.com/articles/using-pull-requests/) with Git, join our growing community of developers and help shape the future of UE4 on Linux.
 
 ![](gitHub_pullrequests.png)
 
diff --git a/Engine/Documentation/Source/Programming/Development/BuildingUnrealEngine/BuildingUnrealEngine.INT.udn b/Engine/Documentation/Source/Programming/Development/BuildingUnrealEngine/BuildingUnrealEngine.INT.udn
--- a/Engine/Documentation/Source/Programming/Development/BuildingUnrealEngine/BuildingUnrealEngine.INT.udn
+++ b/Engine/Documentation/Source/Programming/Development/BuildingUnrealEngine/BuildingUnrealEngine.INT.udn
@@ -149,9 +149,8 @@ Make sure that the latest version of XCode is installed prior to building Unreal
 
 [EXCERPT:linuxcompile]
 
-[REGION:note]
-Our development and support teams currently use the latest version of Ubuntu; as a result, we may not be able to provide support for other Linux distributions (including other versions of Ubuntu). Also, depending on your system's specifications, it may take anywhere from ten minutes to over an hour to compile the engine.
-[/REGION]
+Our development and support teams currently use the latest version of Ubuntu; as a result, we may not be able to provide support for other Linux distributions (including other versions of Ubuntu). Additionally, please make sure your system has at least ten (10) gigabytes of disk space before performing 
+the following steps.
 
 1.  Inside the root directory, run **Setup.sh** from the terminal to setup the files needed to generate the project files.
 
@@ -165,6 +164,9 @@ Our development and support teams currently use the latest version of Ubuntu; as
 
 	![](RunMakeFile_Linux.png)
 
+Depending on your system's specifications, it may take anywhere from ten minutes to over an hour to compile the engine. If you want to shorten the time it takes to compile the engine from source, we recommend compiling 
+the source code on a machine that has at least eight (8) gigabytes of RAM with a multi-core processor having at least eight (8) cores (including hyperthreading).
+
 [/EXCERPT:linuxcompile]
 -->
 
@@ -203,8 +205,7 @@ Our development and support teams currently use the latest version of Ubuntu; as
 [EXCERPT:linuxrun]
 
 [REGION:note]		
-If you're unable to create an OpenGL context as you're working through these steps, you may need to update your graphics drivers (refer to Ubuntu's documentation<!-- and our list of [](Platforms/Linux/BeginnerLinuxDeveloper/SupportedGfxDrivers) -->). 
-For an optimal experience, we recommend using NVIDIA binary drivers (versions equal to or greater than 343.xx).
+If you're unable to create an OpenGL context as you're working through these steps, you may need to update your graphics drivers (refer to our documentation on [](Platforms/Linux/BeginnerLinuxDeveloper/SupportedGfxDrivers)).
 [/REGION]
 
 1.	Navigate to the editor's binary path by entering `cd Engine/Binaries/Linux/` into the terminal.
