Rename/move file(s).

#UE4doc

diff --git a/Engine/Documentation/Source/Platforms/Linux/BeginnerLinuxDeveloper/SettingUpAnUnrealWorkflow/4/SettingUpAnUnrealWorkflow_4.INT.udn b/Engine/Documentation/Source/Platforms/Linux/BeginnerLinuxDeveloper/SettingUpAnUnrealWorkflow/4/SettingUpAnUnrealWorkflow_4.INT.udn
--- a/Engine/Documentation/Source/Platforms/Linux/BeginnerLinuxDeveloper/SettingUpAnUnrealWorkflow/4/SettingUpAnUnrealWorkflow_4.INT.udn
+++ b/Engine/Documentation/Source/Platforms/Linux/BeginnerLinuxDeveloper/SettingUpAnUnrealWorkflow/4/SettingUpAnUnrealWorkflow_4.INT.udn
@@ -37,7 +37,7 @@ games for the Linux platform.
 
 Because you've compiled the engine from its source code on Linux, the engine will compile shaders for a few minutes after launching UE4.
 
-![](howTo_comilingshaders.png)
+![](howTo_compilingshaders.png)
 
 ## Creating your First Project
 
