Addressing review feedback w/updated images, splitting scale plane into separate entry UEDOC-3162

diff --git a/Engine/Documentation/Source/Engine/Editor/VR/UniversalGizmo/UniversalGizmo.INT.udn b/Engine/Documentation/Source/Engine/Editor/VR/UniversalGizmo/UniversalGizmo.INT.udn
--- a/Engine/Documentation/Source/Engine/Editor/VR/UniversalGizmo/UniversalGizmo.INT.udn
+++ b/Engine/Documentation/Source/Engine/Editor/VR/UniversalGizmo/UniversalGizmo.INT.udn
@@ -48,10 +48,10 @@ The elements of the Universal Gizmo are colored red, green, and blue representin
 
 | Gizmo | World | Local |
 | ---- |-------| -----|
-| **Universal** | ![](WorldUniversalGizmo.png)(w:350)[INCLUDE:#worlduniversal]|![](LocalUniversalGizmo.png)(w:350) [INCLUDE:#localuniversal]  |
-| **Translate** | ![](WorldTranslateGizmo.png)(w:350)[INCLUDE:#worldtranslate]|![](LocalTranslateGizmo.png)(w:350) [INCLUDE:#localtranslate]|
-| **Rotate** | ![](WorldRotationGizmo.png)(w:350) [INCLUDE:#worldrotate]|![](LocalRotateGizmo.png)(w:350) [INCLUDE:#localrotate] |
-| **Scale** | ![](WorldScaleGizmo.png)(w:350) [INCLUDE:#worldscale]|![](LocalScaleGizmo.png)(w:350)[INCLUDE:#localscale] |
+| **Universal** | ![](GlobalUniversal.png)(w:350)[INCLUDE:#worlduniversal]|![](LocalUniversal.png)(w:350) [INCLUDE:#localuniversal]  |
+| **Translate** | ![](GlobalTranslate.png)(w:350)[INCLUDE:#worldtranslate]|![](LocalTranslate.png)(w:350) [INCLUDE:#localtranslate]|
+| **Rotate** | ![](GlobalRotate.png)(w:350) [INCLUDE:#worldrotate]|![](LocalRotate.png)(w:350) [INCLUDE:#localrotate] |
+| **Scale** | ![](GlobalScale.png)(w:350) [INCLUDE:#worldscale]|![](LocalScale.png)(w:350)[INCLUDE:#localscale] |
 
 
 <!--
@@ -102,6 +102,7 @@ The elements of the Universal Gizmo are colored red, green, and blue representin
 [EXCERPT:LocalScale]
 * [Uniform Scale](#uniformscale)
 * [Scale](#scale)
+* [Scale Plane](#scaleplane)
 * [Stretch](#stretch)
 [/EXCERPT]
 -->
@@ -112,12 +113,13 @@ Each gizmo gives you access to different transform actions. By aiming the contro
 
 | Action | Manipulator | Usage |
 | --- | --- | --- |
-|[REGION:none](#translate)[/REGION] **Translate** | ![](LocalTranslateGizmo_Translate.png)(w:300) | Aim the controller at one of the translation manipulators and squeeze the **Trigger** halfway to grab it. Move the controller to move the selected Actor(s) along the axis. |
-| [REGION:none](#translateplane)[/REGION]**Translate Plane** | ![](LocalTranslateGizmo_TranslatePlane.png)(w:300) | Aim the controller at one of the translation plane manipulators and squeeze the **Trigger** halfway to grab it. Move the controller to move the selected Actor(s) along the plane (in 2 axes). |
-|[REGION:none](#rotate)[/REGION] **Rotate** | ![](LocalRotateGizmo_Highlight.png)(w:300) | Aim the controller at one of the rotation manipulators and squeeze the **Trigger** halfway to grab it. Move the controller to rotate the selected Actor(s) around the axis. |
-| [REGION:none](#uniformscale)[/REGION]**Uniform Scale** | ![](WorldScaleGizmo_Highlight.png)(w:300) | Aim the controller at the white cube in the center of the gizmo and squeeze the **Trigger** halfway to grab it. Move the controller to scale the selected Actor(s) along all axes uniformly. |
-| [REGION:none](#scale)[/REGION]**Scale** | ![](LocalScaleGizmo_Highlight.png)(w:300) | Aim the controller at one of the scale manipulators and squeeze the **Trigger** halfway to grab it. Move the controller to scale the selected Actor(s) along the selected axis or plane. |
-| [REGION:none](#stretch)[/REGION]**Stretch** |![](LocalScaleGizmo_Stretch.png)(w:300) | Aim the controller at one of the stretch manipulators and squeeze the **Trigger** halfway to grab it. Move the controller to scale the selected Actor(s) non-uniformly along all axes while remaining anchored to the opposite corner.| 
+|[REGION:none](#translate)[/REGION] **Translate** | ![](TranslateBar.png)(w:300) | Aim the controller at one of the translation manipulators and squeeze the **Trigger** halfway to grab it. Move the controller to move the selected Actor(s) along the axis. |
+| [REGION:none](#translateplane)[/REGION]**Translate Plane** | ![](TranslatePlaneHighlight.png)(w:300) | Aim the controller at one of the translation plane manipulators and squeeze the **Trigger** halfway to grab it. Move the controller to move the selected Actor(s) along the plane (in 2 axes). |
+|[REGION:none](#rotate)[/REGION] **Rotate** | ![](RotateHighlight.png)(w:300) | Aim the controller at one of the rotation manipulators and squeeze the **Trigger** halfway to grab it. Move the controller to rotate the selected Actor(s) around the axis. |
+| [REGION:none](#uniformscale)[/REGION]**Uniform Scale** | ![](UniformScaleCube.png)(w:300) | Aim the controller at the white cube in the center of the gizmo and squeeze the **Trigger** halfway to grab it. Move the controller to scale the selected Actor(s) along all axes uniformly. |
+| [REGION:none](#scale)[/REGION]**Scale** | ![](ScaleCube.png)(w:300) | Aim the controller at one of the scale manipulator cubes and squeeze the **Trigger** halfway to grab it. Move the controller to scale the selected Actor(s) along the selected axis. |
+| [REGION:none](#scaleplane)[/REGION]**Scale Plane** | ![](ScalePlane.png)(w:300) | Aim the controller at one of the scale manipulator planes and squeeze the **Trigger** halfway to grab it. Move the controller to scale the selected Actor(s) along the selected plane. |
+| [REGION:none](#stretch)[/REGION]**Stretch** |![](StretchHandles.png)(w:300) | Aim the controller at one of the stretch manipulators and squeeze the **Trigger** halfway to grab it. Move the controller to scale the selected Actor(s) non-uniformly along all axes while remaining anchored to the opposite corner.| 
 
 
 
