#UE4 Docs: Fixing a few minor spelling mistakes.

diff --git a/Engine/Documentation/Source/Engine/MediaFramework/Overview/MF_Overview.INT.udn b/Engine/Documentation/Source/Engine/MediaFramework/Overview/MF_Overview.INT.udn
--- a/Engine/Documentation/Source/Engine/MediaFramework/Overview/MF_Overview.INT.udn
+++ b/Engine/Documentation/Source/Engine/MediaFramework/Overview/MF_Overview.INT.udn
@@ -11,8 +11,8 @@ Related: Engine\MediaFramework\HowTo
 
 [TOC(start:2 end:2)]
 
-**Media Framework** uses several assets to allow the play back of video inside Unreal Engine 4. That video can be scrubbed, paused or rewound inside a **Media Player** asset as well as controlled through C++ or [](Engine/Blueprints). 
-Whether you want a montior in your level playing a video or you want to create a UI based interface with [](Engine/UMG) to allow players control over your video playback, you will need to specify the **Media Source** (where to find the media) to allow the Media Player the ability to play the media back. 
+**Media Framework** uses several assets to allow the playback of video inside Unreal Engine 4. That video can be scrubbed, paused or rewound inside a **Media Player** asset as well as controlled through C++ or [](Engine/Blueprints). 
+Whether you want a monitor in your level playing a video or you want to create a UI based interface with [](Engine/UMG) to allow players control over your video playback, you will need to specify the **Media Source** (where to find the media) to allow the Media Player the ability to play the media back. 
 
 ## Media Source Types
 
@@ -31,13 +31,13 @@ Please see the [](Engine\MediaFramework\TechReference\) page for information on
 ![](MediaSource_File.png)
 
 This type of asset is used for media files stored on your device or in a shared local network file. The **File Path** section is where you specify where the media file resides. 
-While you can point to anywhere the file exists, however its typically a best practice to move your media files to the **Content/Movies** folder of your project and then point to the media file. 
+While you can point to anywhere the file exists, however, it's typically a best practice to move your media files to the **Content/Movies** folder of your project and then point to the media file. 
 This will ensure that your media is packaged correctly with your project (you will also see the yellow exclamation point when pointing to a file outside of this location as a warning). 
 
 This type of source asset also allows you to load the entire media file into memory and play it from there by enabling the **Precache File** option. 
 Based on the size of your file, caching time may vary so take that into consideration when choosing this option. 
 
-Each of the source assets allow you to provide the name of the native media player with the **Default Player** field or leave it empty to find one automatically. 
+Each of the source assets allows you to provide the name of the native media player with the **Default Player** field or leave it empty to find one automatically. 
 
 [REGION:note]
 Please see the [](Engine\MediaFramework\HowTo\FileMediaSource) how-to for working with File Media Source assets. 
@@ -90,8 +90,8 @@ From the **Add New** button under **Media**, select the **Media Sound Wave** ass
 ![](MediaPlayerSoundWindow.png)
 
 The Media Sound Wave asset provides Debug, Concurrency and Attenuation settings that can be adjusted to define how the asset is processed. 
-When linked to a **Media Player** asset, audio attached to a video source will automatically play back along with the video. 
-Typically you can leave the default settings for a newly created Media Sound Wave asset, however if you need more control over how the sound is perceived you can adjust the Concurrency and Attenuation settings. 
+When linked to a **Media Player** asset, audio attached to a video source will automatically playback along with the video. 
+Typically you can leave the default settings for a newly created Media Sound Wave asset, however, if you need more control over how the sound is perceived you can adjust the Concurrency and Attenuation settings. 
 
 [REGION:note]
 Please see the [](Engine/Audio) section for more details on working with audio in Unreal Engine 4. 
@@ -132,7 +132,7 @@ When opening a Media Player asset, the Media Player Window will appear.
 
 ![](MediaPlayerWindow.png)
 
-Inside the Media Player Window you can play back and preview any media source assets you have created (created source media are displayed the **Media Library** panel). 
+Inside the Media Player Window, you can playback and preview any media source assets you have created (created source media are displayed the **Media Library** panel). 
 You can play, pause, rewind or skip to the beginning or end of a video from the toolbar across the top. Any Media Playlists that are currently playing videos will be displayed in the upper right **Playlist** window. 
 
 In the lower-right **Details** panel you can define options for **Playback** or assign the **Sound Wave** and **Video Textures** (or an **Image Texture**) associated with the Media Player. 
@@ -141,8 +141,8 @@ There is also an **Info** panel which will display info about the source such as
 
 ## Media Assets & Blueprints
 
-After your Media Player asset is set up and linked with the appropriate assets, in order to play back the media during game play you will need to first **Open** the source asset through Blueprint or C++. 
-Before you can do that however, you will need a reference to your Media Player asset which you can create a [Variable](Engine/Blueprints/UserGuide/Variables/) for inside any Blueprint. 
+After your Media Player asset is set up and linked with the appropriate assets, in order to playback the media during game play you will need to first **Open** the source asset through Blueprint or C++. 
+Before you can do that, however, you will need a reference to your Media Player asset which you can create a [Variable](Engine/Blueprints/UserGuide/Variables/) for inside any Blueprint. 
 
 ![](MediaPlayerDetailsSet.png)
 
@@ -163,12 +163,12 @@ If you uncheck the Play on Open option (pictured below), you will then need to c
 
 ![](UncheckPlayOnOpen.png)
 
-Below on Event Begin Play we open a Media Playlist, however we only start play back when the **Right Mouse Button** is pressed. 
+Below on Event Begin Play we open a Media Playlist, however, we only start to playback when the **Right Mouse Button** is pressed. 
 
 ![](PlayWhenPressed.png)
 
 You can call other functions from a Media Player reference such as [Binding Events](Engine/Blueprints/UserGuide/EventDispatcher/BindingAndUnbinding/) to various states of playback (such if the video is paused or if the video ends). 
-You can also get or set options like if a video is playing or can play, if a video can loop, or play rate information. With a Media Player reference, you can drag off and see these options under the **Media Player** section of the Blueprint Context Menu. 
+You can also get or set options like if a video is playing or can play if a video can loop, or play rate information. With a Media Player reference, you can drag off and see these options under the **Media Player** section of the Blueprint Context Menu. 
 
 ![](MediaPlayerNodes.png)
 
