Peer review edits UEDOC-3206

diff --git a/Engine/Documentation/Source/Engine/MediaFramework/Overview/MF_Overview.INT.udn b/Engine/Documentation/Source/Engine/MediaFramework/Overview/MF_Overview.INT.udn
--- a/Engine/Documentation/Source/Engine/MediaFramework/Overview/MF_Overview.INT.udn
+++ b/Engine/Documentation/Source/Engine/MediaFramework/Overview/MF_Overview.INT.udn
@@ -11,7 +11,7 @@ Related: Engine\MediaFramework\HowTo
 
 [TOC(start:2 end:2)]
 
-**Media Framework** uses several assets to allow the playback of video inside Unreal Engine 4. That video can be scrubbed, paused or rewound inside a **Media Player** asset as well as controlled through C++ or [](Engine/Blueprints). 
+**Media Framework** uses several assets to allow the playback of videos inside Unreal Engine 4. Videos can be scrubbed, paused or rewound inside a **Media Player** asset as well as controlled through [C++](Programming) or [](Engine/Blueprints). 
 Whether you want a monitor in your level playing a video or you want to create a UI based interface with [](Engine/UMG) to allow players control over your video playback, you will need to specify the **Media Source** (where to find the media) to allow the Media Player the ability to play the media back. 
 
 ## Media Source Types
@@ -20,7 +20,7 @@ There are several different Media Source assets that you can use to define where
 
 ![](MediaSource.png)
 
-Inside the **Content Browser** click the **Add New** Button, then under **Media** select the desired source asset: 
+Inside the **Content Browser**, click the **Add New** Button, then under **Media** select the desired source asset: 
 
 [REGION:note]
 Please see the [](Engine\MediaFramework\TechReference\) page for information on supported file types and URLs. 
@@ -31,13 +31,13 @@ Please see the [](Engine\MediaFramework\TechReference\) page for information on
 ![](MediaSource_File.png)
 
 This type of asset is used for media files stored on your device or in a shared local network file. The **File Path** section is where you specify where the media file resides. 
-While you can point to anywhere the file exists, however, it's typically a best practice to move your media files to the **Content/Movies** folder of your project and then point to the media file. 
-This will ensure that your media is packaged correctly with your project (you will also see the yellow exclamation point when pointing to a file outside of this location as a warning). 
+While a media file can exist anywhere, it's typically a good practice to move your media files to the **Content/Movies** folder of your project and then point to the media file. 
+This will ensure that your media is packaged correctly with your project. You will also see a yellow exclamation point when pointing to a file outside of this location as a warning. 
 
 This type of source asset also allows you to load the entire media file into memory and play it from there by enabling the **Precache File** option. 
-Based on the size of your file, caching time may vary so take that into consideration when choosing this option. 
+Based on the size of your file, caching time may vary, so take that into consideration when choosing this option. 
 
-Each of the source assets allows you to provide the name of the native media player with the **Default Player** field or leave it empty to find one automatically. 
+Each of the source assets allows you to provide the name of the native media player with the **Default Player** field. If you leave this field empty, the media player will be detected automatically. 
 
 [REGION:note]
 Please see the [](Engine\MediaFramework\HowTo\FileMediaSource) how-to for working with File Media Source assets. 
@@ -47,7 +47,7 @@ Please see the [](Engine\MediaFramework\HowTo\FileMediaSource) how-to for workin
 
 ![](MediaSource_Stream.png)
 
-This asset allows you to specify a URL as the source of your media which is pulled from the internet. Currently linking to a YouTube or DailyMotion style URL is not supported and you will need to link directly to a hosted file, for example below is a link to a NASA test flight video: 
+This asset allows you to specify a URL as the source of your media which is pulled from the internet. Currently, linking to a YouTube or DailyMotion style URL is not supported and you will need to link directly to a hosted file. For example, below is a link to a NASA test flight video: 
 
 * http://s3.amazonaws.com/akamai.netstorage/mp4/AirSTAR.mp4
 
@@ -73,7 +73,7 @@ Select **Media Playlist** from the **Add New** button under the **Media** sectio
 ![](CustomizedPlaylist.png)
 
 In the **Playlist** section you can define the source assets to include and playback order. 
-Any Media Source assets you have created will be displayed in the lower **Media Library** window which you can **Double-click** on to add to the end of the playlist. 
+Any Media Source assets you have created will be displayed in the lower **Media Library** window which you can **double-click** on to add to the end of the playlist. 
 
 [REGION:note]
 Please see the [](Engine\MediaFramework\HowTo\MediaPlaylists) how-to for working with Media Playlist assets. 
@@ -90,8 +90,8 @@ From the **Add New** button under **Media**, select the **Media Sound Wave** ass
 ![](MediaPlayerSoundWindow.png)
 
 The Media Sound Wave asset provides Debug, Concurrency and Attenuation settings that can be adjusted to define how the asset is processed. 
-When linked to a **Media Player** asset, audio attached to a video source will automatically playback along with the video. 
-Typically you can leave the default settings for a newly created Media Sound Wave asset, however, if you need more control over how the sound is perceived you can adjust the Concurrency and Attenuation settings. 
+When linked to a **Media Player** asset, audio attached to a video source will automatically play back along with the video. 
+Typically you can leave the default settings for a newly created Media Sound Wave asset. However, if you need more control over how the sound is perceived you can adjust the Concurrency and Attenuation settings. 
 
 [REGION:note]
 Please see the [](Engine/Audio) section for more details on working with audio in Unreal Engine 4. 
@@ -105,19 +105,19 @@ Please see the [](Engine/Audio) section for more details on working with audio i
 ![](MediaTexture.png)
 
 While Media Sound Waves provide the audio, Media Textures provide the visuals from your media source assets. 
-Media Texture assets can be included within [](Engine/Rendering/Materials/) that can then be applied to meshes in your level such as a billboard, TV or monitor to create the appearance as if the video is playing on that object within the game world. 
+Media Texture assets can be included within [](Engine/Rendering/Materials/) that can then be applied to meshes in your level such as a billboard, TV or monitor to make it appear as if the video is playing on that object within the game world. 
 
 ![](MediaTextureWindow.png)
 
 Above, a Material including a Media Texture asset has been created and applied to a static mesh in the level. 
-Inside the [Texture Editor](Engine/Content/Types/Textures/Properties/Interface) the Media Texture displays the same playback position as it appears in the Material in the level. 
+Inside the [Texture Editor](Engine/Content/Types/Textures/Properties/Interface), the Media Texture displays the same playback position as it appears in the Material in the level. 
 In addition to the standard [](Engine/Content/Types/Textures/Properties), Media Textures allow you to define the X and Y Axis **Addressing** to either Clamped, Mirror or World values. 
 
 
 
 ## Media Player Assets
 
-Once you have a Media Source or Media Playlist, you can playback the asset using a **Media Player** asset. 
+Once you have a Media Source or Media Playlist, you can play back the asset using a **Media Player** asset. 
 
 ![](MediaPlayer.png)
 
@@ -139,7 +139,7 @@ In the lower-right **Details** panel you can define options for **Playback** or
 There is also an **Info** panel which will display info about the source such as the stream type, codecs, sample rate, dimensions and more. 
 
 
-## Media Assets & Blueprints
+## Media Assets & Scripting
 
 After your Media Player asset is set up and linked with the appropriate assets, in order to playback the media during game play you will need to first **Open** the source asset through Blueprint or C++. 
 Before you can do that, however, you will need a reference to your Media Player asset which you can create a [Variable](Engine/Blueprints/UserGuide/Variables/) for inside any Blueprint. 
@@ -163,7 +163,7 @@ If you uncheck the Play on Open option (pictured below), you will then need to c
 
 ![](UncheckPlayOnOpen.png)
 
-Below on Event Begin Play we open a Media Playlist, however, we only start to playback when the **Right Mouse Button** is pressed. 
+Below, when **Event BeginPlay** executes, we open a Media Playlist. However, we only start to play the **Media Player** when the **Right Mouse Button** is **Pressed**. 
 
 ![](PlayWhenPressed.png)
 
