Adding additional info about play in editor from Yannick's review. UEDOC-3161

diff --git a/Engine/Documentation/Source/Engine/Editor/VR/QuickSelect/QuickSelectMenu.INT.udn b/Engine/Documentation/Source/Engine/Editor/VR/QuickSelect/QuickSelectMenu.INT.udn
--- a/Engine/Documentation/Source/Engine/Editor/VR/QuickSelect/QuickSelectMenu.INT.udn
+++ b/Engine/Documentation/Source/Engine/Editor/VR/QuickSelect/QuickSelectMenu.INT.udn
@@ -75,5 +75,5 @@ The **Quick Select** menu enables you to quickly adjust grid, rotation and snap
 | Coordinate System | Sets the coordinate system [World Space/Local Space] to use when transforming Actors. |
 | Transform Selector | Cycles the Universal Gizmo through the available transform modes: Translate, Rotate, Scale. Only the manipulators for the selected mode are displayed. |
 | [REGION:tablesection] [/REGION] ||
-| Play | Starts your project from within the VR Editor so that you can test your gameplay quickly. To instantly go back to the VR Editor from your test session, press both triggers and a grip button on each motion controller.  |
+| Play | Starts your project from within the VR Editor so that you can test your gameplay quickly. You will begin gameplay from the player start location. To instantly go back to the VR Editor from your test session, press both triggers and a grip button on each motion controller.  |
 
