#UE4 Docs: Applying feedback from Allan.

diff --git a/Engine/Documentation/Source/Platforms/Mobile/Materials/MobileMaterials.INT.udn b/Engine/Documentation/Source/Platforms/Mobile/Materials/MobileMaterials.INT.udn
--- a/Engine/Documentation/Source/Platforms/Mobile/Materials/MobileMaterials.INT.udn
+++ b/Engine/Documentation/Source/Platforms/Mobile/Materials/MobileMaterials.INT.udn
@@ -74,7 +74,7 @@ Due to hardware differences between PCs and mobile devices you could run into is
 
 <!--
 [EXCERPT:UFP]
-When enabled, Use Full Precision will make sure that the Material does not automatically truncate data when the Material is used on a mobile device. Keep in mind that this will make the Material **more** expensive to render as you telling UE4 to disreguard any automatic optmizations that help reduce the cost of the Material.
+When enabled, Use Full Precision will make sure that the Material's shader will use highest precision math available when used on a mobile device. This can be used to work around rendering problems caused by low-precision math, such as tight specular reflections on glossy materials. Keep in mind that this will make the Material more expensive to render as you telling UE4 to disregard any automatic optmizations that help reduce the cost of the Material.
 
 [OBJECT:ComparisonSlider]
 	[PARAM:before]
