#DocFixes

UEDOC-3219

Made 2 gramatical errors in the document

diff --git a/Engine/Documentation/Source/Platforms/Mobile/MobilePackagingWizard/MobilePackagingWizard.INT.udn b/Engine/Documentation/Source/Platforms/Mobile/MobilePackagingWizard/MobilePackagingWizard.INT.udn
--- a/Engine/Documentation/Source/Platforms/Mobile/MobilePackagingWizard/MobilePackagingWizard.INT.udn
+++ b/Engine/Documentation/Source/Platforms/Mobile/MobilePackagingWizard/MobilePackagingWizard.INT.udn
@@ -23,7 +23,7 @@ The Unreal Engine 4 (UE4) **Mobile Packaging Wizard** will help to reduce the si
 
 To access the Mobile Packaging Wizard you will need to do the folowing.
 
-1. From the **Main Toolbar** locate the **Window** option and click on it. 
+1. From the **Main Toolbar**, locate the **Window** option and click on it. 
 
 	![](MPW_Location.png)
 
@@ -66,7 +66,7 @@ Before the Mobile Packaging Wizard can be used, a new profile that is specific t
 
 ## Building the Application 
 
-Now that all of the options for the Mobile Packing Wizard have been set up and saved into a Custom Launch Profile, it is now time to run that profile and package your project. In the following section, we will take a look at what needs to be done in order for the Mobile Packing Wizard to successfully package your UE4 project.
+Now that all of the options for the Mobile Packing Wizard have been set up and saved into a Custom Launch Profile, it is time to run that profile and package your project. In the following section, we will take a look at what needs to be done in order for the Mobile Packing Wizard to successfully package your UE4 project.
 
 1. If the Project Launcher is not already open, open it up and then look at the bottom of the Project Launcher in the **Custom Launch Profile** section. In this section, you should have two profiles that start with your project name and are followed by either **Android APK** or **Android DLC**. 
 
