#UE4 Docs: Checking in new images and document updates.

diff --git a/Engine/Documentation/Source/Enterprise/Automotive/AutomotiveMaterialPack/AutomotiveMaterialPack.INT.udn b/Engine/Documentation/Source/Enterprise/Automotive/AutomotiveMaterialPack/AutomotiveMaterialPack.INT.udn
--- a/Engine/Documentation/Source/Enterprise/Automotive/AutomotiveMaterialPack/AutomotiveMaterialPack.INT.udn
+++ b/Engine/Documentation/Source/Enterprise/Automotive/AutomotiveMaterialPack/AutomotiveMaterialPack.INT.udn
@@ -88,9 +88,6 @@ In the following section, we will take a look at some of the parameters that all
 	[/OBJECT]
 	[/region]
 
-
-
-
 ## Car Paint Material
 
 The Car Paint Material is a very robust Material that has 29 unique parameters that can be adjusted. When trying to modify this Material for your project, we recommend you starting adjusting and experimenting with the following parameters:
@@ -111,6 +108,12 @@ By adjusting just these four different parameters you can achieve vastly differe
 	[/PARAM]
 [/OBJECT]
 
+The following image sequence shows off the variety of paint colors and surface looks that can be achieved by using just  the parameters listed above.
+
+[REGION:fullwidth]
+[INCLUDE:Enterprise\Automotive\AutomotiveMaterialPack\#CP]
+[/REGION]
+
 ## Orange Peel Material Parameter
 
 The Clearcoat lacquer that is applied on top of automotive paint often has a slightly bumpy surface similar to the peel of an Orange. Because of this resemblance this layer of lacquer is referred to as Orange Peel. You can control the intensity of the Orange Peel using the **Orange Peel Strength** and **Orange Peel Scale** parameters. To disable the Orange Peel effect, set the **Orange Peel Strength** to **0**.
@@ -215,25 +218,38 @@ The rubber Material that comes with the Automotive Material Pack can be used for
 * **Interior Pieces**
 * **Tires**
 
-When looking at the rubber Material Instance you will notice that it has a tire-specific option called **Use Sidewall**. Enabling this option will add a new Normal Map Texture that is for display the lettering that can be found on the tire sidewalls. 
+When looking at the rubber Material Instance you will notice that it has a tire-specific option called **Use Sidewall**. Enabling this option will add a new Normal Map Texture that is for display the lettering that can be found on the tire sidewalls. Failing to properly line up the UV of the tire sidewall to the Normal map of the tire will result in the tire lettering and logos not show up correctly on the tire sidewall like in the image below. 
 
-![](AP_Tire_UV_00.png)
+[OBJECT:ComparisonSlider]
+	[PARAM:before]
+	![Correct Way](AP_09_Rubber_04_Right.png)
+	[/PARAM]
+	[PARAM:after]
+	![Wrong Way](AP_09_Rubber_04_Wrong.png)
+	[/PARAM]
+[/OBJECT]
 
 [region:note]
-Keep in mind that when using the sidewall option for a tire you will need to make sure that the UV for the tire sidewall mesh line up with the Normal Map.
+When using the sidewall option for a tire you will need to make sure that the UV for the tire sidewall mesh line up with the Normal Map like in the image marked Correct UV Layout.
 [/region]
 
-Failing to properly line up the UV of the tire sidewall to the Normal map of the tire will result in the tire lettering and logos not show up correctly on the tire sidewall like in the image below. 
-
 [OBJECT:ComparisonSlider]
 	[PARAM:before]
-	![Correct Way](AP_09_Rubber_04_Right.png)
+	![Correct UV Layout](AP_09_Rubber_05_UV_LineUp.png)
 	[/PARAM]
 	[PARAM:after]
-	![Wrong Way](AP_09_Rubber_04_Wrong.png)
+	![Wrong UV Layout](AP_09_Rubber_06_UV_Wrong.png)
 	[/PARAM]
 [/OBJECT]
 
+[REGION:lightbox]
+[![](AP_Tire_Sidewall.png)(w:800)](AP_Tire_Sidewall.png)
+[/REGION]
+
+[REGION:caption]
+Click for full image.
+[/REGION]
+
 ## Leather / Perforated Material
 
 The leather Material that comes with the Automotive Material Pack has 12 different preset leather colors to choose from and also contains five different perforated leather presets. 
@@ -323,3 +339,38 @@ If you are unfamiliar with how to capture and use Cubemap in UE4 please check ou
 [/OBJECT]
 [/EXCERPT:RM]
 -->
+
+<!-----
+[EXCERPT:CP]
+[OBJECT:SliderSequence]
+	[PARAMLITERAL:max]
+	14
+	[/PARAMLITERAL]
+	[PARAMLITERAL:ext]
+	jpg
+	[/PARAMLITERAL]
+	[PARAMLITERAL:id]
+	CP
+	[/PARAMLITERAL]
+	[PARAM:caption]
+	**Drag the slider to see the different colors of paint.**
+	[/PARAM]
+	[PARAM:image_list]
+	![](CP_1.png)
+	![](CP_2.png)
+	![](CP_3.png)
+	![](CP_4.png)
+	![](CP_5.png)
+	![](CP_6.png)
+	![](CP_7.png)
+	![](CP_8.png)
+	![](CP_9.png)
+	![](CP_10.png)
+	![](CP_11.png)
+	![](CP_12.png)
+	![](CP_13.png)
+	![](CP_14.png)
+	[/PARAM]
+[/OBJECT]
+[/EXCERPT:CP]
+-->
\ No newline at end of file
