#UE4 docs Pose Assets docs initial submission

diff --git a/Engine/Documentation/Source/Engine/Animation/AnimHowTo/AnimHowToOverview.INT.udn b/Engine/Documentation/Source/Engine/Animation/AnimHowTo/AnimHowToOverview.INT.udn
--- a/Engine/Documentation/Source/Engine/Animation/AnimHowTo/AnimHowToOverview.INT.udn
+++ b/Engine/Documentation/Source/Engine/Animation/AnimHowTo/AnimHowToOverview.INT.udn
@@ -2,7 +2,7 @@ Availability: Public
 Title:Animation How To's
 Crumbs: %ROOT%, Engine, Engine\Animation, Engine\Animation\AnimHowTo
 Description:Examples of different aspects of the Animation system in Unreal Engine 4 are demonstrated through several step-by-step How To guides.
-version: 4.9
+version: 4.12
 
 The Animation How To's provide several step-by-step examples of working with and using different aspects of [Persona](Engine/Animation), which is the animation tool inside Unreal Engine 4. 
 Whether you are new to animation in Unreal Engine 4 or want to learn how to work with certain aspects of Persona, the How To pages below cover a broad range of topics and aim to give you a place to start or expand upon your current knowledge. 
@@ -21,6 +21,7 @@ Refer to the tables below for more information and click each link to view its c
 | [](Engine\Animation\AnimHowTo\PropertyBasedBlending) | %Engine\Animation\AnimHowTo\PropertyBasedBlending:description% |
 | [](Engine\Animation\AnimHowTo\AimOffset) | %Engine\Animation\AnimHowTo\AimOffset:description% |
 | [](Engine\Animation\AnimHowTo\Retargeting) | %Engine\Animation\AnimHowTo\Retargeting:description% |
+| [](Engine\Animation\AnimHowTo\CreatePoseAsset) | %Engine\Animation\AnimHowTo\CreatePoseAsset:description% |
 
 
 ## Advanced
diff --git a/Engine/Documentation/Source/Engine/Animation/AnimHowTo/CreatePoseAsset/HowToCreatePoseAsset.INT.udn b/Engine/Documentation/Source/Engine/Animation/AnimHowTo/CreatePoseAsset/HowToCreatePoseAsset.INT.udn
new file mode 100644
--- /dev/null
+++ b/Engine/Documentation/Source/Engine/Animation/AnimHowTo/CreatePoseAsset/HowToCreatePoseAsset.INT.udn
@@ -0,0 +1,41 @@
+Availability:Public
+Title:Creating a Pose Asset
+Crumbs: %ROOT%, Engine, Engine\Animation, Engine\Animation\AnimHowTo
+Description:How to use Persona to generate a single frame Pose Asset.
+version: 4.13
+skilllevel: Intermediate
+
+While you've been able to extract single frames of animation from an Animation Sequence for some time, you've been left to setup blends on your own. In 4.13, this changes with **Pose Assets**. These where implemented to support facial animation where either FACS (Facial Action Coding System) or viseme curves can drive poses. However, you could use this system to drive to create new animation by blending multiple poses. This supports bone transform as well as blendshapes. 
+
+This How-To will walk you through the creation of a **Pose Asset**.
+
+## Steps
+
+1. From the Content Browser, you can right-click an Animatino Sequence and create an **Pose Asset** from it:
+
+	[REGION:raw]
+	![](PoseAsset1.png)
+	[/REGION]
+	
+
+1. Or, in Persona, you can use the **Create Asset** -> **Create PoseAsset** -> **From Current Pose** or **Animation** functions:
+
+	![](PoseAsset2.png)
+
+1. Once you create pose asset, the names will be auto generated by default. You can rename each pose or you can paste from clipboard. 
+
+	![](PoseAsset3.png)
+
+1. You can then see the effect of each pose by changing their weights in the **FacePose** Panel.
+
+	![](PoseAsset4.png)
+
+## Result
+
+You now have a new **Pose Asset** to use in your animations. To preview it make sure **Preview Post Asset** is set to your **Pose Asset**, and then add a **Variable Curve** to you animation that needs the **Pose Asset** data.
+
+![](PoseAsset5.png)
+
+For these to work at run time, you'll also need **PoseAsset** node in your Anim Blueprint.
+
+![](PoseAsset6.png)
