#UE4 docs availability and spelling corrections

diff --git a/Engine/Documentation/Source/Engine/Animation/AnimHowTo/CreatePoseAsset/HowToCreatePoseAsset.INT.udn b/Engine/Documentation/Source/Engine/Animation/AnimHowTo/CreatePoseAsset/HowToCreatePoseAsset.INT.udn
--- a/Engine/Documentation/Source/Engine/Animation/AnimHowTo/CreatePoseAsset/HowToCreatePoseAsset.INT.udn
+++ b/Engine/Documentation/Source/Engine/Animation/AnimHowTo/CreatePoseAsset/HowToCreatePoseAsset.INT.udn
@@ -1,17 +1,17 @@
-Availability:Public
+Availability:Docs
 Title:Creating a Pose Asset
 Crumbs: %ROOT%, Engine, Engine\Animation, Engine\Animation\AnimHowTo
 Description:How to use Persona to generate a single frame Pose Asset.
 version: 4.13
 skilllevel: Intermediate
 
-While you've been able to extract single frames of animation from an Animation Sequence for some time, you've been left to setup blends on your own. In 4.13, this changes with **Pose Assets**. These where implemented to support facial animation where either FACS (Facial Action Coding System) or viseme curves can drive poses. However, you could use this system to drive to create new animation by blending multiple poses. This supports bone transform as well as blendshapes. 
+While you've been able to extract single frames of animation from an Animation Sequence for some time, you've been left to setup blends on your own. In 4.13, this changes with **Pose Assets**. These were implemented to support facial animation where either FACS (Facial Action Coding System) or viseme curves can drive poses. However, you could use this system to drive to create new animation by blending multiple poses. This supports bone transform as well as blendshapes. 
 
 This How-To will walk you through the creation of a **Pose Asset**.
 
 ## Steps
 
-1. From the Content Browser, you can right-click an Animatino Sequence and create an **Pose Asset** from it:
+1. From the Content Browser, you can right-click an **Animation Sequence** and create an **Pose Asset** from it:
 
 	[REGION:raw]
 	![](PoseAsset1.png)
