#UE4 Docs: Changed you to your so the sentance is correct.

diff --git a/Engine/Documentation/Source/Engine/Animation/AnimHowTo/CreatePoseAsset/HowToCreatePoseAsset.INT.udn b/Engine/Documentation/Source/Engine/Animation/AnimHowTo/CreatePoseAsset/HowToCreatePoseAsset.INT.udn
--- a/Engine/Documentation/Source/Engine/Animation/AnimHowTo/CreatePoseAsset/HowToCreatePoseAsset.INT.udn
+++ b/Engine/Documentation/Source/Engine/Animation/AnimHowTo/CreatePoseAsset/HowToCreatePoseAsset.INT.udn
@@ -32,7 +32,7 @@ This How-To will walk you through the creation of a **Pose Asset**.
 
 ## Result
 
-You now have a new **Pose Asset** to use in your animations. To preview it make sure **Preview Post Asset** is set to your **Pose Asset**, and then add a **Variable Curve** to you animation that needs the **Pose Asset** data.
+You now have a new **Pose Asset** to use in your animations. To preview it make sure **Preview Post Asset** is set to your **Pose Asset**, and then add a **Variable Curve** to your animation that needs the **Pose Asset** data.
 
 ![](PoseAsset5.png)
 
