Applying SME feedback.

diff --git a/Engine/Documentation/Source/Engine/UMG/HowTo/DragAndDrop/DragAndDrop.INT.udn b/Engine/Documentation/Source/Engine/UMG/HowTo/DragAndDrop/DragAndDrop.INT.udn
--- a/Engine/Documentation/Source/Engine/UMG/HowTo/DragAndDrop/DragAndDrop.INT.udn
+++ b/Engine/Documentation/Source/Engine/UMG/HowTo/DragAndDrop/DragAndDrop.INT.udn
@@ -19,7 +19,7 @@ In this example, we take a look at how a player can drag and reposition a health
 360
 [/PARAMLITERAL]
 [PARAMLITERAL:videoid]
-qUDIkWcQChI
+kYkhxuJsVs4
 [/PARAMLITERAL]
 [/OBJECT]
 
@@ -39,6 +39,7 @@ _At the end of this guide, you will have a UI widget that you can drag around in
 	* [](Engine\UMG\HowTo\DragAndDrop\WidgetSetup)
 	* [](Engine\UMG\HowTo\DragAndDrop\OnDragDetected)
 	* [](Engine\UMG\HowTo\DragAndDrop\OnDrop)
+	* [](Engine\UMG\HowTo\DragAndDrop\FinishingUp)
 	[/PARAM]
 [/OBJECT]
 
diff --git a/Engine/Documentation/Source/Engine/UMG/HowTo/DragAndDrop/FinishingUp/FinishingUp.INT.udn b/Engine/Documentation/Source/Engine/UMG/HowTo/DragAndDrop/FinishingUp/FinishingUp.INT.udn
new file mode 100644
--- /dev/null
+++ b/Engine/Documentation/Source/Engine/UMG/HowTo/DragAndDrop/FinishingUp/FinishingUp.INT.udn
@@ -0,0 +1,89 @@
+Availability: Docs
+Title:5. Finishing Up
+Crumbs: %ROOT%, Engine, Engine/UMG, Engine/UMG/HowTo
+Description:We finish by configuring our drag widget, adding our Health Bar to our Main HUD for display and testing everything out.
+Related: Engine/UMG
+Related: Engine/UMG/UserGuide
+Version: 4.12
+SkillLevel: Intermediate
+
+[VAR:Steps]
+[OBJECT:Navigation]
+	[PARAM:previous]
+		[Previous Step](Engine\UMG\HowTo\DragAndDrop\OnDrop "%Engine\UMG\HowTo\DragAndDrop\OnDrop:title%")
+	[/PARAM]
+	[PARAM:current]
+	[/PARAM]
+	[PARAM:home]
+		[Drag and Drop Widgets](Engine\UMG\HowTo\DragAndDrop)
+	[/PARAM]
+	[PARAM:next]
+	[/PARAM]
+[/OBJECT]
+[/VAR]
+
+%Steps%
+
+In this final step we configure our **DragWidget** (which is the visual widget we are dragging). 
+We then tell our Character Blueprint to add the **HUD** Widget Blueprint to the viewport as well as enable the Mouse Cursor so we can see where we are dragging. 
+
+## Steps
+
+1.	Open the **DragWidget** Widget Blueprint and replace the **Canvas Panel** with a **Size Box** with a child **Border**. 
+
+1.	Rename the **SizeBox** to **WidgetSize** and set **isVariable** to true, then check both **Width Override** and **Height Override**. 
+
+	![](UpdateDragDrop3.png)
+
+	We are going to set the Width and Height of the Size Box to reflect the size of our Health Bar widget. 
+
+1.	For the **Border**, set the **Brush Color** to **Black** with an **Alpha** of **0.5** for some transparency. 
+
+	![](UpdateDragDrop4.png)
+
+	In our example we will have a silhouette of our Health Bar as the drag visual while keeping the actual Health Bar in the same location. 
+
+1.	On the **Graph**, drag in **Widget Reference** and cast to the **HealthBar** to get the **Desired Size** then **Break Vector**. 
+
+	![](UpdateDragDrop5.png)
+
+	This will give us the size of our Health Bar which we can then set our Size Box to. 
+	We could have manually entered the overrides, however if we change the size of our Health Bar we would then have to also change it here. 
+
+1.	Drag in the **Widget Size** and off is **Set Height Override** and **Set Width Override** and connect as shown below. 
+
+	[REGION:lightbox]
+	[![](UpdateDragDrop6.png)](UpdateDragDrop6.png)
+	[/REGION:lightbox]
+
+	_Click for full view._
+	
+1.	In the **Content Browser** under **Content/ThirdPersonBP/Blueprints**, open the **ThirdPersonCharacter** Blueprint.
+
+1.	Off an **Event Begin Play**, **Create Widget** (set to **HUD**) then **Add to Viewport** and **Get Player Controller** to **Show Mouse Cursor**. 
+	![](DragDrop19.png)
+
+1.	**Compile**, **Save** then click the **Play** button to play in the Editor. 
+
+
+## End Result
+
+When you play in the Editor, you can **Left-click** and drag the health bar on screen then drop it in a new location. 
+
+[OBJECT:EmbeddedVideo]
+[PARAMLITERAL:width]
+640
+[/PARAMLITERAL]
+[PARAMLITERAL:height]
+360
+[/PARAMLITERAL]
+[PARAMLITERAL:videoid]
+kYkhxuJsVs4
+[/PARAMLITERAL]
+[/OBJECT]
+
+This is just an example of the elements you will need to start the Drag and Drop process. 
+Additional checking may be needed to ensure that players do not drag the widget outside of Safe Zones or potentially on top of other widgets. 
+
+%Steps%
+
diff --git a/Engine/Documentation/Source/Engine/UMG/HowTo/DragAndDrop/OnDragDetected/OnDragDetected.INT.udn b/Engine/Documentation/Source/Engine/UMG/HowTo/DragAndDrop/OnDragDetected/OnDragDetected.INT.udn
--- a/Engine/Documentation/Source/Engine/UMG/HowTo/DragAndDrop/OnDragDetected/OnDragDetected.INT.udn
+++ b/Engine/Documentation/Source/Engine/UMG/HowTo/DragAndDrop/OnDragDetected/OnDragDetected.INT.udn
@@ -29,34 +29,46 @@ In this step we determine what happens when the player is in fact dragging the H
 
 ## Steps
 
-1.	Open the **OnDragDetected** function, then **Right-click** and add a **Create Drag Drop Operation** node set to use **WidgetDrag**. 
+1.	Open the **OnDragDetected** function, **Right-click** and add a **Create Widget** node set to **Drag Widget**, then **Browse** to it. 
 
-	![](DragDrop9.png)
+	![](UpdatedDragDrop1.png)
 
-	UMG supplies a default class for drag and drop operations however we have created a custom class that has additional variables in it that we will pass through as part of the process.
+	When drag is detected, we will create our DragWidget (which will hold our drag widget). 
+	Before we continue with our OnDragDetected function, we need to pass through the Health Bar widget as reference when the DragWidget is created. 
 
-1.	**Right-click** in the graph and **Get Reference to Self** and connect as shown and set the **Pivot** drop-down to **Mouse Down**. 
-	
-	![](DragDrop10.png)
+1.	Open **DragWidget** and add a **User Widget** variable called **WidgetReference** set to **Expose on Spawn** and **Editable**. 
 
-	Pivot determines where the Drag Widget Visual will appear when being dragged relative to the pointer performing the drag operation. 
+	![](UpdatedDragDrop2.png)
 
-1.	From the **MyBlueprint** panel, drag in the **Drag Offset** variable and connect it and the pins shown below. 
+1.	**Compile** and close **DragWidget**, then open the **HealthBar** Widget Blueprint and the **OnDragDetected** function again. 
 
-	![](DragDrop11.png)
+1.	Assign **Self** as the **Widget Reference** and **Promote** the **Return Value** to a variable called **Dragged Widget** and connect as shown. 
 
-	Here we are passing through that vector 2D information of where the player clicked on the widget which we will use as part of the OnDrop operation. 
+	![](UpdatedDragDrop3.png)
 
-1.	**Right-click** in the graph and add a **Remove from Parent** node and connect as shown. 
+	Now when our drag widget is created, it will have a reference to the existing Health Bar widget. 
+	We also promote the dummy drag widget to a variable so that we have access to it later and can remove it from display when we choose to drop the widget. 
 
-	![](DragDrop12.png)
+1.	**Right-click** and create a **DragAndDropOperation** node set to **Widget Drag** with **Pivot** set to **Mouse Down**. 
+	
+	![](UpdatedDragDrop4.png)
+	
+	Pivot determines where the Drag Widget Visual will appear when being dragged relative to the pointer performing the drag operation. 
 
-	Before we complete the On Drag Detected process, we remove the widget we are dragging from display.
-	If we don't remove the original HealthBar widget from display, it would remain in the default position while we are dragging a copy around onscreen. 
-	The widget itself is being passed through as part of the process so we are not creating a "new" version of the widget. 
+1.	Add **Self** as the **Widget Reference**, **DraggedWidget** as the **Default Drag Visual** and **Drag Offset** as the **Drag Offset** on the node. 
 
-The last few things we need to do is determine what happens when the player releases the Left Mouse Button and performs the OnDrop function. 
+	[REGION:lightbox]
+	[![](UpdatedDragDrop5.png)](UpdatedDragDrop5.png)
+	[/REGION:lightbox]
+
+	_Click for full view._
+
+	Above we indicated what Widget Blueprint is our reference, what the drag visual should be and provided an offset in which to start dragging. 
+	There is an offset pin on the node by default, however we are using our own in our own calculation to determine the position in which to start dragging when Mouse Button Down is pressed. 
+
+The last few things we need to do is to determine what happens when the player releases the Left Mouse Button and performs the OnDrop function. 
 We also need to add the HealthBar Widget to our Main HUD Widget Blueprint and tell that to be displayed before testing everything out.  
 
 %Steps%
 
+
diff --git a/Engine/Documentation/Source/Engine/UMG/HowTo/DragAndDrop/OnDrop/OnDrop.INT.udn b/Engine/Documentation/Source/Engine/UMG/HowTo/DragAndDrop/OnDrop/OnDrop.INT.udn
--- a/Engine/Documentation/Source/Engine/UMG/HowTo/DragAndDrop/OnDrop/OnDrop.INT.udn
+++ b/Engine/Documentation/Source/Engine/UMG/HowTo/DragAndDrop/OnDrop/OnDrop.INT.udn
@@ -1,7 +1,7 @@
 Availability: Docs
 Title:4. On Drop
 Crumbs: %ROOT%, Engine, Engine/UMG, Engine/UMG/HowTo
-Description:We finish by creating the script for OnDrop functionality before we test. 
+Description:In this step we configure what happens when the player drops the dragged widget.
 Related: Engine/UMG
 Related: Engine/UMG/UserGuide
 Version: 4.12
@@ -18,13 +18,14 @@ SkillLevel: Intermediate
 		[Drag and Drop Widgets](Engine\UMG\HowTo\DragAndDrop)
 	[/PARAM]
 	[PARAM:next]
+		[Next Step](Engine\UMG\HowTo\DragAndDrop\FinishingUp "%Engine\UMG\HowTo\DragAndDrop\FinishingUp:title%")
 	[/PARAM]
 [/OBJECT]
 [/VAR]
 
 %Steps%
 
-To finish up, we script what happens when we "drop" the widget in our HUD as well as test everything out. 
+Here we set up our main **HUD** Widget Blueprint and override what happens when we perform the **OnDrop** function. 
 
 ## Steps
 
@@ -58,46 +59,24 @@ To finish up, we script what happens when we "drop" the widget in our HUD as wel
 	We will use this minus the **DragOffset** to determine where we should place the widget when it is dropped. 
 
 1.	Subtract the **DragOffset** from the **Return Value** of **Absolute to Local**. 
+	
+	![](UpdateDragDrop1.png)
 
-1.	For the **Widget Reference**, use the **Add to Viewport** and **Set Position in Viewport** and connect as shown. 
+1.	Off **Widget Reference**, add a **Remove from Parent** followed by **Add to Viewport** and **Set Position in Viewport** nodes. 
 
-	[REGION:lightbox]
-	[![](DragDrop18.png)](DragDrop18.png)
-	[/REGION]
+1.	Connect the nodes as shown below and uncheck **Remove DPIScale** and check the **Return Value** on the **Return Node**. 
 
-	[REGION:caption]
-	Click for full view. 
+	[REGION:lightbox]
+	[![](UpdateDragDrop2.png)](UpdateDragDrop2.png)
 	[/REGION]
 
-	Make sure to uncheck **Remove DPIScale** and check the **Return Value** checkboxes as we do not need to remove the [](Engine/UMG/UserGuide/DPIScaling) and we have handled this function so we return as true. 
-	With this script we add the widget being dragged to the viewport again and set its position based on where we clicked and started dragging. The last thing we need to do is create and display our HUD Widget Blueprint. 
-
-1.	Inside the **Content/ThirdPersonBP/Blueprints** folder, open the **ThirdPersonCharacter** Blueprint. 
-
-1.	Create the **HUD** Widget Blueprint, **Add to Viewport** and **Get Player Controller** to **Show Mouse Cursor**. 
-
-	![](DragDrop19.png)
-
-1.	**Compile**, **Save** then click the **Play** button to play in the Editor. 
+	_Click for full view._
+	
+	We uncheck **Remove DPIScale** and check the **Return Value** checkboxes as we do not need to remove the [](Engine/UMG/UserGuide/DPIScaling) and we have handled this function so we return as true. 
+	With this script we first remove the existing Health Bar widget before re-adding it to the screen in the new location relative to our Drag Offset. 
 
-
-## End Result
-
-When you play in the Editor, you can **Left-click** and drag the health bar on screen then drop it in a new location. 
-
-[OBJECT:EmbeddedVideo]
-[PARAMLITERAL:width]
-640
-[/PARAMLITERAL]
-[PARAMLITERAL:height]
-360
-[/PARAMLITERAL]
-[PARAMLITERAL:videoid]
-qUDIkWcQChI
-[/PARAMLITERAL]
-[/OBJECT]
-
-This is just an example of the elements you will need to start the Drag and Drop process. 
-Additional checking may be needed to ensure that players do not drag the widget outside of Safe Zones or potentially on top of other widgets. 
+Our HUD is setup to handle dropping our dragged widget and it displays our Health Bar. 
+In the next and final step in this tutorial we set up our Dragged Widget to mimic the shape of our Health Bar as well as instruct our Character Blueprint to add the HUD to the viewport. 
 
 %Steps%
+
diff --git a/Engine/Documentation/Source/Engine/UMG/HowTo/DragAndDrop/ProjectSetup/ProjectSetup.INT.udn b/Engine/Documentation/Source/Engine/UMG/HowTo/DragAndDrop/ProjectSetup/ProjectSetup.INT.udn
--- a/Engine/Documentation/Source/Engine/UMG/HowTo/DragAndDrop/ProjectSetup/ProjectSetup.INT.udn
+++ b/Engine/Documentation/Source/Engine/UMG/HowTo/DragAndDrop/ProjectSetup/ProjectSetup.INT.udn
@@ -30,11 +30,12 @@ For this how-to, we are using the **Blueprint Third Person Template**. However,
 
 ## Steps
 
-1.	In the **Content Browser** of your project, create a folder with two **Widget Blueprints**, one called **HealthBar** and the other **HUD**. 
+1.	In the **Content Browser**, create three **Widget Blueprints** called: **DragWidget**, **HealthBar** and **HUD**. 
 
 	![](DragDrop1.png)
 
 	In this example, the element that we want to be drag and droppable is our **HealthBar** widget that will be placed in our **HUD**. 
+	We are going to create and drag a widget called **DragWidget** in the same shape as our Health Bar rather than dragging the actual Health Bar around. 
 
 1.	**Right-click** in the **Content Browser** and create a new **Blueprint Class** of the **DragDropOpertion** class and call it **WidgetDrag**. 
 
