#UE4 Docs: minor text updates

diff --git a/Engine/Documentation/Source/Platforms/iOS/Windows/Windows.INT.udn b/Engine/Documentation/Source/Platforms/iOS/Windows/Windows.INT.udn
--- a/Engine/Documentation/Source/Platforms/iOS/Windows/Windows.INT.udn
+++ b/Engine/Documentation/Source/Platforms/iOS/Windows/Windows.INT.udn
@@ -20,7 +20,7 @@ While building for iOS requires a Mac, using remote compilation means that a who
 1. Install your iOS Developer Certificates to the System keychain. If your certificates are already installed to your login keychain, just drag them to the System keychain and they will be installed there as well.
 
 	[REGION:note]
-	For more on creating iOS Developer Certificates, see [](Platforms/iOS/QuickStart/3).
+	For more on creating iOS Developer Certificates, see [Generating SIgning Certificates](Platforms/iOS/QuickStart/3).
 	[/REGION]
 
 1. Lastly, make sure the private keys for your certificates are set up to be accessible from any program.  This can be accessed by double-clicking on the key under the Developer Certificate. For more information, see [this Apple Support article](https://support.apple.com/kb/PH18682).
@@ -63,11 +63,9 @@ At this point, you should successfully have generated your SSH key and are ready
 
 If you wish to pass around this data so that it can be used for all projects in the future, then you will want to specify the **Remote Server Name** and **Remote User Name** in one or more of the `*Engine.ini` files on your computer.  
 
-1. Select the multi-dot button next to the property.  
+1. Select the multi-dot button next to the property to open the **Configuration Editor**.
 
 	![](ConfigFiles.png)
-
-	This will open the **Configuration Editor** 
 	
 1. Set the properties for each individual `.ini` file you wish.  
 
