Incorporating editorial review feedback by making auto-entry enabling/disabling steps more procedural. UEDOC-3163

diff --git a/Engine/Documentation/Source/Engine/Editor/VR/ActivateVRMode/ActivateVRMode.INT.udn b/Engine/Documentation/Source/Engine/Editor/VR/ActivateVRMode/ActivateVRMode.INT.udn
--- a/Engine/Documentation/Source/Engine/Editor/VR/ActivateVRMode/ActivateVRMode.INT.udn
+++ b/Engine/Documentation/Source/Engine/Editor/VR/ActivateVRMode/ActivateVRMode.INT.udn
@@ -83,8 +83,14 @@ cZJmh881FCo
 [/PARAMLITERAL]
 [/OBJECT]
 
+**To enable/disable VR Editing Mode Auto-Entry:**
+1. In the **Edit** menu, click on **Editor Preferences** to view the Unreal Editor options.
+	
+	![](EditorPref.png)
+2. Under the **General** heading, select the **Experimental** section and locate the **VR** category which contains the VR Editor options.
 
-There is a setting under **VR** in the **Experimental** section of **Editor Preferences** that will allow you to turn off auto-entry if you prefer.
+	![](ExperimentalVR.png)
+3. Toggle the **Enable VR Mode Auto-Entry** option to disable or enable the auto-entry functionality.
 
-![](AutoEntryMenu_Highlight.png)
+	![](AutoEntryMenu_Highlight.png)
 
