Adding clarifying text to a step after editorial review. UEDOC-1394

diff --git a/Engine/Documentation/Source/Platforms/iOS/Windows/Windows.INT.udn b/Engine/Documentation/Source/Platforms/iOS/Windows/Windows.INT.udn
--- a/Engine/Documentation/Source/Platforms/iOS/Windows/Windows.INT.udn
+++ b/Engine/Documentation/Source/Platforms/iOS/Windows/Windows.INT.udn
@@ -67,7 +67,7 @@ If you wish to pass around this data so that it can be used for all projects in
 
 	![](ConfigFiles.png)
 	
-1. Set the properties for each individual `.ini` file you wish.  
+1. Set the properties for each individual `.ini` file you want to pass around the SSH data in.  
 
 	![](inifileexample.png)
 
