#UE4 docs integrated SME feedback

diff --git a/Engine/Documentation/Source/Engine/Rendering/LightingAndShadows/LightingChannels/LightingChannels.INT.udn b/Engine/Documentation/Source/Engine/Rendering/LightingAndShadows/LightingChannels/LightingChannels.INT.udn
--- a/Engine/Documentation/Source/Engine/Rendering/LightingAndShadows/LightingChannels/LightingChannels.INT.udn
+++ b/Engine/Documentation/Source/Engine/Rendering/LightingAndShadows/LightingChannels/LightingChannels.INT.udn
@@ -92,7 +92,14 @@ The performance impact of using **Lighting Channels** is minimal, but also non-z
 
 ## On Mobile
 
-As of Unreal Engine 4.12, **Lighting Channels** works on mobile renderers.
+As of Unreal Engine 4.13, Lighting Channels works on mobile renderers.
+
+Mobile renderer supports:
+
+* Multiple directional lights are supported in different channels.
+* Each primitives can only be affected by one directional, and it will use the directional light from first lighting channel it has set.
+* CSM shadows from stationary or movable directional lights cast only on primitives with matching lighting channels.
+* Dynamic point lights fully support lighting channels.
 
 
 
