#UEDOC-3206
#UEDOC-3207

Applying peer feedback.

diff --git a/Engine/Documentation/Source/Engine/MediaFramework/HowTo/StreamMediaSource/UsingStreamMedia.INT.udn b/Engine/Documentation/Source/Engine/MediaFramework/HowTo/StreamMediaSource/UsingStreamMedia.INT.udn
--- a/Engine/Documentation/Source/Engine/MediaFramework/HowTo/StreamMediaSource/UsingStreamMedia.INT.udn
+++ b/Engine/Documentation/Source/Engine/MediaFramework/HowTo/StreamMediaSource/UsingStreamMedia.INT.udn
@@ -30,7 +30,8 @@ For this how-to we are using the **Blueprint Third Person Template** project wit
 
 	![](StreamMedia_03.png)
 
-	Your Stream URL must link directly to a [Supported Format](Engine\MediaFramework\TechReference\) in order to play back the video. If you do not have a link, **Right-click** this [Sample URL](http://s3.amazonaws.com/akamai.netstorage/mp4/AirSTAR.mp4) and copy the link address then paste it in the Stream URL field. 
+	Above we are using a fake address, however your Stream URL must link directly to a [Supported Format](Engine\MediaFramework\TechReference\) in order to play back the video. 
+	If you do not have your own file to link to, you can **Right-click** this [Sample Video](Infiltrator Demo.mp4) and copy the link address then paste it in the Stream URL field. 
 
 1.	**Right-click** in the **Content Browser** then under **Media** select the **Media Player** asset. 
 
@@ -93,7 +94,7 @@ When you play in the editor, the video will start playing on the Static Mesh.
 360
 [/PARAMLITERAL]
 [PARAMLITERAL:videoid]
-0gxc69-M-Fk
+ncHIQ0QteC8
 [/PARAMLITERAL]
 [/OBJECT]
 
diff --git a/Engine/Documentation/Source/Engine/MediaFramework/Overview/MF_Overview.INT.udn b/Engine/Documentation/Source/Engine/MediaFramework/Overview/MF_Overview.INT.udn
--- a/Engine/Documentation/Source/Engine/MediaFramework/Overview/MF_Overview.INT.udn
+++ b/Engine/Documentation/Source/Engine/MediaFramework/Overview/MF_Overview.INT.udn
@@ -17,11 +17,11 @@ Whether you want a monitor in your level playing a video or you want to create a
 ## Media Source Types
 
 There are several different Media Source assets that you can use to define where the media is coming from. 
+Whether you have a file that is included as part of your project, streaming media from a website or using platform specific media you will need to define a source before you can play the video.
+You can add a Media Source inside the **Content Browser** by clicking the **Add New** Button then specifying which source type under the **Media** section.  
 
 ![](MediaSource.png)
 
-Inside the **Content Browser**, click the **Add New** Button, then under **Media** select the desired source asset: 
-
 [REGION:note]
 Please see the [](Engine\MediaFramework\TechReference\) page for information on supported file types and URLs. 
 [/REGION]
@@ -30,7 +30,7 @@ Please see the [](Engine\MediaFramework\TechReference\) page for information on
 
 ![](MediaSource_File.png)
 
-This type of asset is used for media files stored on your device or in a shared local network file. The **File Path** section is where you specify where the media file resides. 
+The **File Media Source** asset is used for media files stored on your device or in a shared local network file. The **File Path** section is where you specify where the media file resides. 
 While a media file can exist anywhere, it's typically a good practice to move your media files to the **Content/Movies** folder of your project and then point to the media file. 
 This will ensure that your media is packaged correctly with your project. You will also see a yellow exclamation point when pointing to a file outside of this location as a warning. 
 
@@ -47,9 +47,8 @@ Please see the [](Engine\MediaFramework\HowTo\FileMediaSource) how-to for workin
 
 ![](MediaSource_Stream.png)
 
-This asset allows you to specify a URL as the source of your media which is pulled from the internet. Currently, linking to a YouTube or DailyMotion style URL is not supported and you will need to link directly to a hosted file. For example, below is a link to a NASA test flight video: 
-
-* http://s3.amazonaws.com/akamai.netstorage/mp4/AirSTAR.mp4
+This **Stream Media Source** asset allows you to specify a URL as the source of your media which is pulled from the internet. Currently, linking to a YouTube or DailyMotion style URL is not supported and you will need to link directly to a hosted file. 
+Above we are using a fake address with a fake hosted MP4 file. 
 
 [REGION:note]
 Please see the [](Engine\MediaFramework\HowTo\StreamMediaSource) how-to for working with Stream Media Source assets. 
@@ -59,21 +58,24 @@ Please see the [](Engine\MediaFramework\HowTo\StreamMediaSource) how-to for work
 
 ![](MediaSource_Platform.png)
 
-The Platform Media Source asset allows overriding Media Sources on a per-platform basis. You can specify the Media Source to be used if no source was specified for a target platform as well as define the Media Sources per platform. 
+The **Platform Media Source** asset allows overriding Media Sources on a per-platform basis. You can specify the Media Source to be used if no source was specified for a target platform as well as define the Media Sources per platform. 
+
+For example, supppose you want a particluar video to play only on Android Devices or only on PS4. Inside the **Media Sources** section you can designate which videos play on which platforms where leaving the Media Source blank will cause the **Default Source** to be used. 
 
 
 ## Media Playlists
 
 After creating a few Media Source assets, you can add them to a **Media Playlist** which will play the assigned assets back sequentially. 
+You can create Media Playlists inside the **Content Browser** from the **Add New** Button under the **Media** section. 
 
 ![](MediaPlaylistButton.png)
 
-Select **Media Playlist** from the **Add New** button under the **Media** section. 
+In the Media Playstlist window (pictured below) under the **Playlist** section you can define the source assets to include and playback order. 
+Any Media Source assets you have created will be displayed in the lower **Media Library** window which you can **Double-click** on to add to the end of the playlist. 
 
 ![](CustomizedPlaylist.png)
 
-In the **Playlist** section you can define the source assets to include and playback order. 
-Any Media Source assets you have created will be displayed in the lower **Media Library** window which you can **double-click** on to add to the end of the playlist. 
+Media Playlists can use a mixture of each Media Source types and playback in order unless specified through Blueprint Script. 
 
 [REGION:note]
 Please see the [](Engine\MediaFramework\HowTo\MediaPlaylists) how-to for working with Media Playlist assets. 
@@ -89,8 +91,8 @@ From the **Add New** button under **Media**, select the **Media Sound Wave** ass
 
 ![](MediaPlayerSoundWindow.png)
 
-The Media Sound Wave asset provides Debug, Concurrency and Attenuation settings that can be adjusted to define how the asset is processed. 
-When linked to a **Media Player** asset, audio attached to a video source will automatically play back along with the video. 
+The Media Sound Wave asset provides **Debug**, **Concurrency** and **Attenuation** settings that can be adjusted to define how the asset is processed. 
+When linked to a Media Player asset, audio attached to a video source will automatically play back along with the video. 
 Typically you can leave the default settings for a newly created Media Sound Wave asset. However, if you need more control over how the sound is perceived you can adjust the Concurrency and Attenuation settings. 
 
 [REGION:note]
@@ -100,7 +102,7 @@ Please see the [](Engine/Audio) section for more details on working with audio i
 
 ## Media Textures
 
-**Media Texture** assets are used for rendering video tracks from **Media Player** assets. 
+**Media Texture** assets are used for rendering video tracks from Media Player assets. 
 
 ![](MediaTexture.png)
 
@@ -111,28 +113,30 @@ Media Texture assets can be included within [](Engine/Rendering/Materials/) that
 
 Above, a Material including a Media Texture asset has been created and applied to a static mesh in the level. 
 Inside the [Texture Editor](Engine/Content/Types/Textures/Properties/Interface), the Media Texture displays the same playback position as it appears in the Material in the level. 
-In addition to the standard [](Engine/Content/Types/Textures/Properties), Media Textures allow you to define the X and Y Axis **Addressing** to either Clamped, Mirror or World values. 
-
+In addition to the standard [](Engine/Content/Types/Textures/Properties), Media Textures allow you to define the X and Y Axis **Addressing** to either **Clamped**, **Mirror** or **World** values. 
 
+[REGION:note]
+Please see the [](Engine\MediaFramework\HowTo) section for guides on how to setup and use Media Textures to play videos. 
+[/REGION]
 
 ## Media Player Assets
 
 Once you have a Media Source or Media Playlist, you can play back the asset using a **Media Player** asset. 
+The Media Player asset requires the use of a Media Texture to produce the video playback and a Media Sound Wave to produce the audio associated with the video. 
+You can create Media Player assets inside the **Content Browser** from the **Add New** Button under the **Media** section. 
 
 ![](MediaPlayer.png)
 
-From the **Add New** button under **Media**, select the **Media Player** asset. 
-
-The Media Player asset needs to have a **Media Sound Wave** and **Media Texture** assigned to it for it to properly play any video.
-If you have not created one, when you create any new Media Player asset a dialogue window will appear asking if would like to create and automatically link those assets with the newly created Media Player asset. 
+When creating Media Player assets, the **Create Media Player** window (pictured below) will appear asking if you would also like to create and link assets to the Media Player. 
+This allows you to automatically create and assign a Media Texture and Media Sound Wave associated with the Media Player asset you are creating. 
 
 ![](CreateMediaPlayer.png)
 
-When opening a Media Player asset, the Media Player Window will appear. 
+**Double-clicking** on a Media Player asset will open the **Media Player Window**: 
 
 ![](MediaPlayerWindow.png)
 
-Inside the Media Player Window, you can playback and preview any media source assets you have created (created source media are displayed the **Media Library** panel). 
+Inside the Media Player Window, you can playback and preview any Media Source assets you have created (created source media are displayed the **Media Library** panel). 
 You can play, pause, rewind or skip to the beginning or end of a video from the toolbar across the top. Any Media Playlists that are currently playing videos will be displayed in the upper right **Playlist** window. 
 
 In the lower-right **Details** panel you can define options for **Playback** or assign the **Sound Wave** and **Video Textures** (or an **Image Texture**) associated with the Media Player. 
@@ -141,8 +145,8 @@ There is also an **Info** panel which will display info about the source such as
 
 ## Media Assets & Scripting
 
-After your Media Player asset is set up and linked with the appropriate assets, in order to playback the media during game play you will need to first **Open** the source asset through Blueprint or C++. 
-Before you can do that, however, you will need a reference to your Media Player asset which you can create a [Variable](Engine/Blueprints/UserGuide/Variables/) for inside any Blueprint. 
+After your Media Player asset is set up and linked with the appropriate assets, in order to playback the media during game play you will need to first open the Media Source asset through Blueprint or C++. 
+Before you can open your Media Source, you will first need to get a reference to your Media Player asset which you can create a [Variable](Engine/Blueprints/UserGuide/Variables/) for inside any Blueprint. 
 
 ![](MediaPlayerDetailsSet.png)
 
@@ -158,12 +162,12 @@ With a reference to your Media Player defined, you can then call an **Open** fun
 | **Open URL** | Opens the specified media URL. |
 | **Reopen** | Reopens the currently opened media or play list. |
 
-Once you open the source you can then **Play** the media. Newly created Media Player assets have the **Play on Open** option enabled by default which means when you open a media source it will start playing the video automatically. 
-If you uncheck the Play on Open option (pictured below), you will then need to call the **Play** function on a media source after it has been opened. 
+Newly created Media Player assets have the **Play on Open** option enabled by default which means when you open a Media Source it will start playing the video automatically. 
+If you uncheck the Play on Open option (pictured below), you will then need to call the **Play** function on a Media Source after it has been opened. 
 
 ![](UncheckPlayOnOpen.png)
 
-Below, when **Event BeginPlay** executes, we open a Media Playlist. However, we only start to play the **Media Player** when the **Right Mouse Button** is **Pressed**. 
+Below on **Event BeginPlay** we open a Media Playlist and play the **Media Player** when the **Right Mouse Button** is **Pressed**. 
 
 ![](PlayWhenPressed.png)
 
