#UE4 Docs: updated metadata

diff --git a/Engine/Documentation/Source/Engine/Replay/Replay.INT.udn b/Engine/Documentation/Source/Engine/Replay/Replay.INT.udn
--- a/Engine/Documentation/Source/Engine/Replay/Replay.INT.udn
+++ b/Engine/Documentation/Source/Engine/Replay/Replay.INT.udn
@@ -1,52 +1,15 @@
 Availability:Docs
 Title: Replay Overview
 Crumbs: %ROOT%, Engine
-Description:Overview of the Replay system
+Description:Overview of the Replay system for recording and playback of gameplay
 Version: 4.12
 parent:Engine
 order:17
+topic-image:replay_topic.png
 
-[VAR:Topic]
-[OBJECT:Topic]
-	[PARAM:image]
-		![%Engine/Replay:title%](Engine/Replay/replay_topic.png)
-	[/PARAM]
-	[PARAM:icon]
-		![](%ROOT%/light_icon.png)(convert:false)
-	[/PARAM]
-	[PARAM:title]
-		%Engine/Replay:title%
-	[/PARAM]
-	[PARAM:description]
-		%Engine/Replay:description%
-	[/PARAM]
-	[PARAM:path]
-		[RELATIVE:Engine/Replay]
-	[/PARAM]
-[/OBJECT]
-[/VAR]
-[VAR:TopicCompact]
-[OBJECT:TopicCompact]
-	[PARAM:image]
-		![%Engine/Replay:title%](Engine/Replay/replay_topic.png)
-	[/PARAM]
-	[PARAM:icon]
-		![](%ROOT%/light_icon.png)(convert:false)
-	[/PARAM]
-	[PARAM:title]
-		%Engine/Replay:title%
-	[/PARAM]
-	[PARAM:description]
-		%Engine/Replay:description%
-	[/PARAM]
-	[PARAM:path]
-		[RELATIVE:Engine/Replay]
-	[/PARAM]
-[/OBJECT]
-[/VAR]
 
 [REGION:banner]
-![](Engine/Replay/replay_banner.png)(convert:false)
+![](replay_banner.png)(convert:false)
 [/REGION]
 
 ## Replay System Overview
