#DocFix

UEDOC-3218: Minor grammatical fix

diff --git a/Engine/Documentation/Source/Platforms/Android/OpenGLES31MobileRenderer/OpenGLES31MobileRenderer.INT.udn b/Engine/Documentation/Source/Platforms/Android/OpenGLES31MobileRenderer/OpenGLES31MobileRenderer.INT.udn
--- a/Engine/Documentation/Source/Platforms/Android/OpenGLES31MobileRenderer/OpenGLES31MobileRenderer.INT.udn
+++ b/Engine/Documentation/Source/Platforms/Android/OpenGLES31MobileRenderer/OpenGLES31MobileRenderer.INT.udn
@@ -121,7 +121,7 @@ Please note that OpenGL ES 3.1 Support will only work on code based projects tha
 
 ## Building and Deploying a UE4 Project with OpenGL ES 3.1 Support
 
-Now that the project has been setup to work with OpenGL ES 3.1, it is now time to build the project and deploy it to a Android based Smartphone. To build a project that has OpenGL ES 3.1 support you will need to do the following:
+Now that the project has been setup to work with OpenGL ES 3.1, it is time to build the project and deploy it to a Android based Smartphone. To build a project that has OpenGL ES 3.1 support you will need to do the following:
 
 [region:note]
 Before beginning the next part, make sure that you have your Android Smartphone plugged into your development PC via a USB cable. Failing to do this will result in your project not being able to be deployed.
