#UE4 Docs: minor formatting updates

diff --git a/Engine/Documentation/Source/Engine/Editor/VR/EditorUI/EditorUIVR.INT.udn b/Engine/Documentation/Source/Engine/Editor/VR/EditorUI/EditorUIVR.INT.udn
--- a/Engine/Documentation/Source/Engine/Editor/VR/EditorUI/EditorUIVR.INT.udn
+++ b/Engine/Documentation/Source/Engine/Editor/VR/EditorUI/EditorUIVR.INT.udn
@@ -48,10 +48,10 @@ skilllevel: intermediate
 The VR Editor enables you to summon and dock a number of different panels and windows including:
 
 [REGION:imagetable]
-| ![](window_content.png)(w:185) | ![](window_details.png)(w:185) | ![](window_modes.png)(w:185) | ![](window_outliner.png)(w:185) |
+| ![](window_content.png)(w:225) | ![](window_details.png)(w:225) | ![](window_modes.png)(w:225) | ![](window_outliner.png)(w:225) |
 | --- | --- | --- | --- |
-| Content Browser | Details Panel | Modes Panel | Outliner panel | 
-|![](window_tutorials.png)(w:185) | ![](window_asset.png)(w:185) | ![](colorPicker.png)(w:185)| ![](worldSettings.png)(w:185)|
+| Content Browser | Details Panel | Modes Panel | Outliner Panel | 
+|![](window_tutorials.png)(w:225) | ![](window_asset.png)(w:225) | ![](colorPicker.png)(w:225)| ![](worldSettings.png)(w:225)|
 |Tutorials | Asset Panel | Color Picker | World Settings |
 [/REGION]
 
