#UE4 Docs: minor text/formatting updates

diff --git a/Engine/Documentation/Source/Engine/Rendering/LightingAndShadows/LightingChannels/LightingChannels.INT.udn b/Engine/Documentation/Source/Engine/Rendering/LightingAndShadows/LightingChannels/LightingChannels.INT.udn
--- a/Engine/Documentation/Source/Engine/Rendering/LightingAndShadows/LightingChannels/LightingChannels.INT.udn
+++ b/Engine/Documentation/Source/Engine/Rendering/LightingAndShadows/LightingChannels/LightingChannels.INT.udn
@@ -84,17 +84,14 @@ The performance impact of using **Lighting Channels** is minimal, but also non-z
 
 <!--
 [EXCERPT:per1]
-0.08ms   CopyStencilToLightingChannels 1 draws 1 prims 3 verts
-
+0.08ms   CopyStencilToLightingChannels 1 draws 1 prims 3 verts  
 0.45ms   StandardDeferredLighting 1 draws 1 prims 3 verts
 [/EXCERPT]
 -->
 
 ## On Mobile
 
-As of Unreal Engine 4.13, Lighting Channels works on mobile renderers.
-
-Mobile renderer supports:
+As of Unreal Engine 4.13, Lighting Channels works on mobile renderers with the following features supported:
 
 * Multiple directional lights are supported in different channels.
 * Each primitives can only be affected by one directional, and it will use the directional light from first lighting channel it has set.
