New doc - upgrading existing media framework content to 4.13, tips and notes. UEDOC-3213

diff --git a/Engine/Documentation/Source/Engine/MediaFramework/TechReference/MF_TechRef.INT.udn b/Engine/Documentation/Source/Engine/MediaFramework/TechReference/MF_TechRef.INT.udn
--- a/Engine/Documentation/Source/Engine/MediaFramework/TechReference/MF_TechRef.INT.udn
+++ b/Engine/Documentation/Source/Engine/MediaFramework/TechReference/MF_TechRef.INT.udn
@@ -14,6 +14,9 @@ Related: Engine\MediaFramework\HowTo
 This page outlines the currently supported file extensions and URLs as well as outlines some troubleshooting tips on how to enable Media Framework with different platforms. 
 There are some known platform specific issues and limitations which we are currently investigating and trying to resolve, each of which have been identified in the Troubleshooting Platform sections. 
 
+[REGION:tip]
+If you are upgrading Media Framework content from a pre-4.13 project, see the tips in [](Engine/MediaFramework/TechReference/Upgrading).
+[/REGION]
 
 ## Supported Files & URLs
 
diff --git a/Engine/Documentation/Source/Engine/MediaFramework/TechReference/Upgrading/Upgrading.INT.udn b/Engine/Documentation/Source/Engine/MediaFramework/TechReference/Upgrading/Upgrading.INT.udn
new file mode 100644
--- /dev/null
+++ b/Engine/Documentation/Source/Engine/MediaFramework/TechReference/Upgrading/Upgrading.INT.udn
@@ -0,0 +1,44 @@
+Availability:Docs
+Title: Upgrading Media Framework Content to 4.13
+Crumbs:%ROOT%, Engine
+Description:Updating projects containing Media Framework assets and playback logic to Unreal Engine 4.13
+Navigation:topic
+parent:Engine/MediaFramework
+order:3
+Version: 4.13
+Related: Engine\MediaFramework\Overview
+Related: Engine\MediaFramework\HowTo
+Related: Engine\MediaFramework\TechReference
+
+In Unreal Engine 4.13, the Media Framework system received an update to the asset types and playback methods. This means that if you have existing Media Framework content in your project, you have to update your content. 
+Your project will still load, but playback will no longer work with existing 4.12 content until you have performed the updates, as the workflow is slightly different now. 
+
+As you transition your project, keep the following items in mind:
+
+* **Media Player** assets no longer store a URL for the media. Instead there are now new asset types called 'media sources' that describe the media. For more information on these new asset types, see [](Engine/MediaFramework/Overview).
+
+	![](MediaTypes.png)
+
+* **File Media Source** assets are used for media files stored on your device or in local network file shares. For an example of how to use a **File Media Source** asset, see [](Engine/MediaFramework/HowTo/FileMediaSource).
+
+	![](FileMediaSource.png)
+* **Stream Media Source** assets are used for streaming media over the internet. For an example of how to use a **Stream Media Source** asset, see [](Engine/MediaFramework/HowTo/StreamMediaSource).
+
+	![](StreamMediaSource.png)
+* **Media Playlist** assets are used to combine multiple media sources and play them one after another. Playlists will automatically cycle back to the beginning. For an example of working with playlists, see [](Engine/MediaFramework/HowTo/MediaPlaylists).
+
+	![](MediaPlaylist.png)
+
+* **Media Texture** and **Media Sound Wave** assets no longer hold a reference to the **Media Player**. Instead, the Media Player now references the texture and the sound assets. Additionally, there is no longer a many-to-one relationship between **Media Texture** / **Media Sound Wave** and 
+**Media Player** assets. Every **Media Player** now references at most one of each.
+
+	![](Upgrading_MediaLinks.png)
+
+* There is no longer an option for automatically starting media playback in Media Player assets. Instead you must open a media source via Blueprint. For convenience, Media Player assets have a **Play On Open** option that will begin playback after a media source opened successfully. 
+If you do not set this option, you must initiate playback via Blueprints or C++.
+
+	![](PlayOnOpen.png)
+
+* The **Loop** option will loop the currently opened media source. Please note that, while looping is enabled, playlists do not advance.
+
+	![](Loop.png)
\ No newline at end of file
