Applying SME Feedback.

diff --git a/Engine/Documentation/Source/Engine/MediaFramework/Overview/MF_Overview.INT.udn b/Engine/Documentation/Source/Engine/MediaFramework/Overview/MF_Overview.INT.udn
--- a/Engine/Documentation/Source/Engine/MediaFramework/Overview/MF_Overview.INT.udn
+++ b/Engine/Documentation/Source/Engine/MediaFramework/Overview/MF_Overview.INT.udn
@@ -62,10 +62,15 @@ The **Platform Media Source** asset allows overriding Media Sources on a per-pla
 
 For example, supppose you want a particluar video to play only on Android Devices or only on PS4. Inside the **Media Sources** section you can designate which videos play on which platforms where leaving the Media Source blank will cause the **Default Source** to be used. 
 
+[REGION:warning]
+The Platform Media Source asset is extermely experimental does not work for 4.13 and are planned for a 4.14 release.
+[/REGION]
+
 
 ## Media Playlists
 
-After creating a few Media Source assets, you can add them to a **Media Playlist** which will play the assigned assets back sequentially. 
+While you can play individual Media Source assets individually, one asset you can use to play back multiple Media Sources is the **Media Playlist** asset. 
+After creating a few Media Source assets, you can add them to a Media Playlist which will automatically cycle through and play each back in the order you define.
 You can create Media Playlists inside the **Content Browser** from the **Add New** Button under the **Media** section. 
 
 ![](MediaPlaylistButton.png)
@@ -142,6 +147,9 @@ You can play, pause, rewind or skip to the beginning or end of a video from the
 In the lower-right **Details** panel you can define options for **Playback** or assign the **Sound Wave** and **Video Textures** (or an **Image Texture**) associated with the Media Player. 
 There is also an **Info** panel which will display info about the source such as the stream type, codecs, sample rate, dimensions and more. 
 
+[REGION:warning]
+Still Image Texture output is currently not implemented by any of the Media Player plug-in and this feature is considered experimental.
+[/REGION]
 
 ## Media Assets & Scripting
 
@@ -163,6 +171,7 @@ With a reference to your Media Player defined, you can then call an **Open** fun
 | **Reopen** | Reopens the currently opened media or play list. |
 
 Newly created Media Player assets have the **Play on Open** option enabled by default which means when you open a Media Source it will start playing the video automatically. 
+The Play on Open checkbox is a convenience, but not strictly required. If you do not want to automatically play the media once it has been opened successfully, you can also hook into the **OnMediaOpened** Blueprint event. 
 If you uncheck the Play on Open option (pictured below), you will then need to call the **Play** function on a Media Source after it has been opened. 
 
 ![](UncheckPlayOnOpen.png)
@@ -171,6 +180,11 @@ Below on **Event BeginPlay** we open a Media Playlist and play the **Media Playe
 
 ![](PlayWhenPressed.png)
 
+[REGION:note]
+If you are not automatically playing a Media Source on open and are using the Play function to start playback, it is advised that you do not chain the Play call directly after an Open Source or Open Playlist call. 
+This is because it may take some time for the Media Source to open and the Play command will just return false and the movie will not play as expected.
+[/REGION]
+
 You can call other functions from a Media Player reference such as [Binding Events](Engine/Blueprints/UserGuide/EventDispatcher/BindingAndUnbinding/) to various states of playback (such if the video is paused or if the video ends). 
 You can also get or set options like if a video is playing or can play if a video can loop, or play rate information. With a Media Player reference, you can drag off and see these options under the **Media Player** section of the Blueprint Context Menu. 
 
