#UE4 Docs: Applying SME feedback.

diff --git a/Engine/Documentation/Source/Platforms/Mobile/Materials/MobileMaterials.INT.udn b/Engine/Documentation/Source/Platforms/Mobile/Materials/MobileMaterials.INT.udn
--- a/Engine/Documentation/Source/Platforms/Mobile/Materials/MobileMaterials.INT.udn
+++ b/Engine/Documentation/Source/Platforms/Mobile/Materials/MobileMaterials.INT.udn
@@ -74,7 +74,7 @@ Due to hardware differences between PCs and mobile devices you could run into is
 
 <!--
 [EXCERPT:UFP]
-When enabled, Use Full Precision will make sure that the Material's shader will use highest precision math available when used on a mobile device. This can be used to work around rendering problems caused by low-precision math, such as tight specular reflections on glossy materials. Keep in mind that this will make the Material more expensive to render as you telling UE4 to disregard any automatic optmizations that help reduce the cost of the Material.
+When enabled, Use Full Precision will make sure that the Material's shader will use highest precision math available when used on a mobile device. This can be used to work around rendering problems caused by low-precision math, such as usage of world coordinates or tight specular reflections on glossy materials. Keep in mind that this will make the Material more expensive.
 
 [OBJECT:ComparisonSlider]
 	[PARAM:before]
diff --git a/Engine/Documentation/Source/Platforms/Mobile/Performance/MobilePerformance.INT.udn b/Engine/Documentation/Source/Platforms/Mobile/Performance/MobilePerformance.INT.udn
--- a/Engine/Documentation/Source/Platforms/Mobile/Performance/MobilePerformance.INT.udn
+++ b/Engine/Documentation/Source/Platforms/Mobile/Performance/MobilePerformance.INT.udn
@@ -143,7 +143,7 @@ Here the tree is causing pixels to be **extremely** expensive.  If the player co
 
 ## Mobile Content Scale Factor 
 **Mobile Content Scale Factor** is a way to scale the resolution of your project to best fit the screen resolution of the mobile device being used to view your project. 
-You can create and store multiple device profiles by creating a new .INI file called **DefaultDeviceProfiles.INI** in your project's config folder. 
+You can create and store multiple device profiles by creating a new .INI file called **DefaultDeviceProfiles.ini** in your project's config folder. 
 Inside of this file, you can specify what will happen to the resolution of your project by entering the following command. 
 
 	 r.MobileContentScaleFactor = 
@@ -164,8 +164,7 @@ The number that is used after the command is what determines if the projects res
 
 
 ### Mobile Content Scale Factor for iOS
-When using Mobile Content Scale Factor for iOS devices inputting the following number after the r.MobileContentScaleFactor will produce the following results.
-Please note that the actual resolution for any scale factor other than 0.0 will be corrected to match the aspect ratio of the screen and clamped to native resolution.
+When using Mobile Content Scale Factor for iOS devices inputting the following number after the r.MobileContentScaleFactor will produce the following results. Please note that the actual resolution for any scale factor other than 0.0 will be corrected to match the aspect ratio of the screen and clamped to native resolution.
 
 [region:note]
 The iOS scale factor is directly related to Apple's scale factor system.
@@ -179,8 +178,7 @@ The iOS scale factor is directly related to Apple's scale factor system.
 |r.MobileContentScaleFactor|3.0|ull native resolution for iPhone 6+.|
 
 ### Mobile Content Scale Factor for Android
-When using the Mobile Content Scale Factor for Android devices inputting the following number after the r.MobileContentScaleFactor will produce the following results.
-Please note that the actual resolution for any scale factor other than 0.0 will be corrected to match the aspect ratio of the screen and clamped to native resolution.
+When using the Mobile Content Scale Factor for Android devices inputting the following number after the r.MobileContentScaleFactor will produce the following results. Please note that inputting value other than **0.0** will use this value as the scale factor for a standard 1280x720 or 720x1280 resolution depending on orientation.
 
 |Command| Value|Result|
 |-------|------|------|
@@ -222,6 +220,10 @@ You can set the Material Quality Level in the following ways inside UE4:
 |1|High|
 |2|Medium|
 
+[region:note]
+Please note that 1 for high and 2 for medium **IS** correct and not a typo!
+[/region]
+
 **From Device Profiles**: You can also adjust the Material Quality by doing the following:
 
 1. From the Main Toolbar go to **Windows** > **Developer Tools** and then click on **Device Profiles**.  
@@ -306,6 +308,35 @@ You can override the available rendering option a platform will use by adjusting
 Click for full image.
 [/REGION]
 
+## Mobile Rendering Options
+Inside of your **Project's Settings** menu under the **Mobile** section you will find a number of properties that control  how lighting will be handled inside of your mobile project. In the following section we will take a look at these settings and what effect they will have on your project.
+
+### Accessing your Project Settings
+To access your project settings you will first need to go to the **Main Toolbar** and then select **Edit** > **Project Settings**. 
+
+![](MRO_Project_Settings.png)
+
+Once the Project Settings menu is open, locate the **Engine** section and then click on the **Rendering** category and look for the **Mobile** heading. 
+
+[REGION:lightbox]
+[![](MRO_Rendering_Settings.png)(w:430)](MRO_Rendering_Settings.png)
+[/REGION]
+
+[REGION:caption]
+Click for full image.
+[/REGION]
+
+### Mobile Rendering Options
+Under the Mobile section you will find a number of different options that can be enabled or disabled. In the following section we will take a look each one of the properties and the effect they will have on your mobile project.
+
+![](MRO_Mobile_Settings.png)
+
+|Property Name| Description|
+|--------------------|-------|
+|Mobile HDR|If true, mobile renders in full HDR. Disable this setting for games that do not require lighting features for better performance on slow devices.|
+|Max Dynamic Point Lights|The number of dynamic point lights to support on mobile devices. Setting this to 0 for games which do not require dynamic point lights will reduce the number of shader generated. Changing this setting requires restarting the editor.|
+|Use Shared Dynamic Point LIght Shaders|If this setting is enabled, the same shader will be sued for any number of dynamic point lights(up to the maximum specified above)hitting a surface. This is slightly slower but reduced the number of shaders generated. Changing this setting requires restarting the editor. |
+|Enabled Combined Static and CSM Shadowing|Allo primitives to receive both static and CSM shadows from a stationary light. Disabling will free a mobile texture sampler.|
 
 
 <!-----
