Applying SME feedback.

diff --git a/Engine/Documentation/Source/Engine/MediaFramework/HowTo/ControllingPlayback/ControllingPlayback.INT.udn b/Engine/Documentation/Source/Engine/MediaFramework/HowTo/ControllingPlayback/ControllingPlayback.INT.udn
--- a/Engine/Documentation/Source/Engine/MediaFramework/HowTo/ControllingPlayback/ControllingPlayback.INT.udn
+++ b/Engine/Documentation/Source/Engine/MediaFramework/HowTo/ControllingPlayback/ControllingPlayback.INT.udn
@@ -100,6 +100,29 @@ The rate at which we rewind or fast forward is determined by the **Rate** value
 
 While our example is bound to keyboard presses, you could execute the same functionality within a User Interface built with UMG where the player clicks on a button to execute the same commands. 
 
+## Other Considerations
+
+The availability of some advanced playback features depends on the Media Source, as well as the used player plug-in. 
+For example, not all Media Sources allow for seeking or changing the playback rate, and not all player plug-ins implement support for these features yet. 
+
+Some combinations of Media Sources and players allow for rate changes, but only for certain rate ranges, or the playback may be thinned (i.e. frames will skipped = choppy). 
+Some playback features can be queried via Blueprint, for example CanPause, GetForwardRates, GetReverseRates. 
+All playback functions (Play, Rewind, Pause, SetRate, etc.) provide a Boolean return value indicating success or failure of the desired operation.
+
+There is also currently a limitation in the WmfMedia player plug-in for Windows that prevents chained calls of playback functions. For example if you string together a chain of commands like Play -> SetRate -> Pause -> SetRate -> Play like so: 
+
+[REGION:lightbox]
+[![](ExampleChain.png)](ExampleChain.png)
+[/REGION]
+
+_Click for full view._
+
+Your state machine for the Media Player will likely break and the only way to recover would be to re-open the Media Source or Playlist. 
+
+If you wish to chain calls, one possible workaround is to insert Delay nodes between each call in order to give the player enough time to process the current command. 
+
+![](ExampleChainDelays.png)
+
 
 
 
