Applying SME Feedback.
- Mentioned that Startup Movie Player framework is slated to be replaced by [REDACTED_PERSON] in a future engine release.

diff --git a/Engine/Documentation/Source/Engine/MediaFramework/MediaFramework.INT.udn b/Engine/Documentation/Source/Engine/MediaFramework/MediaFramework.INT.udn
--- a/Engine/Documentation/Source/Engine/MediaFramework/MediaFramework.INT.udn
+++ b/Engine/Documentation/Source/Engine/MediaFramework/MediaFramework.INT.udn
@@ -50,9 +50,9 @@ Version: 4.13
 ![](MediaFrameWork.png)(convert:false)
 [/REGION]
 
-While there is currently a **MoviePlayer** system in UE4, it can only be used for playing startup movies while the Engine loads up. It cannot be used to play movies in-game, 
+While there is currently a **Startup Movie Player** system in UE4, it can only be used for playing startup movies while the Engine loads up. It cannot be used to play movies in-game, 
 say for instance as part of a UI element or perhaps on a Static Mesh in your level like a movie playing on a TV. This is where the **Media Framework** comes in, as not only 
-will it allow you to do both of the examples presented above, it will also provide much more general media playback capabilities (a rundown of which are outlined below).  
+will it allow you to do both of the examples presented above, it will also provide much more general media playback capabilities (outlined below) and is also scheduled to replace the deprecated Startup Movie Player framework in a future engine release. 
 
 The Media Framework in Unreal Engine 4 is:
 * Both Engine & Slate Agnostic
