Applying SME Feedback
- Clarified supported format section within a note.

diff --git a/Engine/Documentation/Source/Engine/MediaFramework/HowTo/StreamMediaSource/UsingStreamMedia.INT.udn b/Engine/Documentation/Source/Engine/MediaFramework/HowTo/StreamMediaSource/UsingStreamMedia.INT.udn
--- a/Engine/Documentation/Source/Engine/MediaFramework/HowTo/StreamMediaSource/UsingStreamMedia.INT.udn
+++ b/Engine/Documentation/Source/Engine/MediaFramework/HowTo/StreamMediaSource/UsingStreamMedia.INT.udn
@@ -30,8 +30,13 @@ For this how-to we are using the **Blueprint Third Person Template** project wit
 
 	![](StreamMedia_03.png)
 
-	Above we are using a fake address, however your Stream URL must link directly to a [Supported Format](Engine\MediaFramework\TechReference\) in order to play back the video. 
-	If you do not have your own file to link to, you can **Right-click** this [Sample Video](Infiltrator Demo.mp4) and copy the link address then paste it in the Stream URL field. 
+	If you do not have a file to link to, **Right-click** this [Sample Video](Infiltrator Demo.mp4) and copy the link address and paste it in the Stream URL field. 
+	
+	[REGION:note]
+	Above we are using a fake address, however your Stream URL must link directly to a supported format in order to play back the video. 
+	For example, PS4Media (PS4) only supports MP4 through HLS right now, while WmfMedia (Windows) can support a number of other stream sources. 
+	For futher information on supported formats per platform/player plug-in, see the [](Engine\MediaFramework\TechReference\) page. 
+	[/REGION]
 
 1.	**Right-click** in the **Content Browser** then under **Media** select the **Media Player** asset. 
 
