Applying SME Feedback.

See JIRA UEDOC-3482 for SME's feedback.

diff --git a/Engine/Documentation/Source/Engine/MediaFramework/TechReference/MF_TechRef.INT.udn b/Engine/Documentation/Source/Engine/MediaFramework/TechReference/MF_TechRef.INT.udn
--- a/Engine/Documentation/Source/Engine/MediaFramework/TechReference/MF_TechRef.INT.udn
+++ b/Engine/Documentation/Source/Engine/MediaFramework/TechReference/MF_TechRef.INT.udn
@@ -36,13 +36,17 @@ The table below indicates the supported extensions based on each media player pl
 | **.m4v** | No | Yes | No | Yes | No |
 | **.mov** | No | Yes | No | Yes | No |
 | **.mp3** | No | No | No | Yes | No |
-| **.mp4** | Yes | Yes | Yes | Yes | Yes |
+| **.mp4** | Yes | Yes | Yes | Yes | No |
 | **.sami** | No | No | No | Yes | No |
 | **.smi** | No | No | No | Yes | No |
 | **.wav** | No | No | No | Yes | No |
 | **.wma** | No | No | No | Yes | No |
 | **.wmv** | No | No | No | Yes | No |
 
+[REGION:warning]
+The MfMedia player plug-in for XboxOne is currently in development and does not work in 4.13. The AvfMedia player plug-in for macOS and iOS is in an experimental state and not all features may work as expected. We are working to resolve both issues and hope to have them addressed in a future engine release.  
+[/REGION]
+
 This table shows the supported URLs based on media player platform. 
 
 | | Android Media Player | AVF Media Player | PS4 Media Player | WMF Media Player | Xbox One Media Player |
@@ -116,6 +120,16 @@ In your game's ** *.Target.cs**:
             	// ...
 	}
 
+To improve playback performance on PS4, your video file should be encoded as interleaved, so that audio and video can be read in a single sequential operation instead of seeking to different areas of the media file. 
+This can be done without re-encoding with [**mp4box**](https://gpac.wp.mines-telecom.fr/mp4box/) by adding the following: 
+
+	mp4box.exe -add non-interleaved-source.mp4 --inter 100 -delay 2=0 -new output.mp4
+
+One should also double check that the encoding is optimum for the PS4 as this will significantly reduce the CPU and GPU overhead of the player. 
+This can be done using [**FFMPEG**](https://ffmpeg.org/). For example, for a 1080p mp4 file we suggest the following: 
+
+	ffmpeg.exe -i "highqualitysource.mp4" -c:v libx264 -crf 20 -preset veryslow -maxrate 15000k -bufsize 30000k -slices 24 -x264opts b-pyramid=0:no-cabac:weightb=0:weightp=0 -refs 1 -pix_fmt yuv420p "non-interleaved-output.mp4"
+
 ### Xbox One
 
 The MfMedia player plug-in is disabled by default and must be enabled in the project’s Plugin settings.
