Minor fix to remove support for XboxOne for 4.13

diff --git a/Engine/Documentation/Source/Engine/MediaFramework/TechReference/MF_TechRef.INT.udn b/Engine/Documentation/Source/Engine/MediaFramework/TechReference/MF_TechRef.INT.udn
--- a/Engine/Documentation/Source/Engine/MediaFramework/TechReference/MF_TechRef.INT.udn
+++ b/Engine/Documentation/Source/Engine/MediaFramework/TechReference/MF_TechRef.INT.udn
@@ -51,7 +51,7 @@ This table shows the supported URLs based on media player platform.
 
 | | Android Media Player | AVF Media Player | PS4 Media Player | WMF Media Player | Xbox One Media Player |
 | --- | --- | --- |--- | --- | --- |
-| **file://** | Yes | Yes | Yes | Yes | Yes |
+| **file://** | Yes | Yes | Yes | Yes | No |
 | **http://** | Yes | No | Yes | Yes | No |
 | **httpd://** | Yes | No | No | Yes | No |
 | **https://** | Yes | No | Yes | Yes | No |
