#UE4 Docs: minor text and formatting updates

diff --git a/Engine/Documentation/Source/Platforms/Mobile/Materials/MobileMaterials.INT.udn b/Engine/Documentation/Source/Platforms/Mobile/Materials/MobileMaterials.INT.udn
--- a/Engine/Documentation/Source/Platforms/Mobile/Materials/MobileMaterials.INT.udn
+++ b/Engine/Documentation/Source/Platforms/Mobile/Materials/MobileMaterials.INT.udn
@@ -74,7 +74,11 @@ Due to hardware differences between PCs and mobile devices you could run into is
 
 <!--
 [EXCERPT:UFP]
-When enabled, Use Full Precision will make sure that the Material's shader will use highest precision math available when used on a mobile device. This can be used to work around rendering problems caused by low-precision math, such as usage of world coordinates or tight specular reflections on glossy materials. Keep in mind that this will make the Material more expensive.
+When enabled, the Material's shader will use the highest precision math available on a mobile device. This can be used to work around rendering problems caused by low-precision math, such as usage of world coordinates or tight specular reflections on glossy materials. 
+
+[REGION:note]
+Enabling full precision shaders will make the Material more expensive.
+[/REGION]
 
 [OBJECT:ComparisonSlider]
 	[PARAM:before]
