#UE4 Docs: minor text and formatting updates

diff --git a/Engine/Documentation/Source/Engine/UMG/HowTo/CreatingWidgets/CreatingWidgets.INT.udn b/Engine/Documentation/Source/Engine/UMG/HowTo/CreatingWidgets/CreatingWidgets.INT.udn
--- a/Engine/Documentation/Source/Engine/UMG/HowTo/CreatingWidgets/CreatingWidgets.INT.udn
+++ b/Engine/Documentation/Source/Engine/UMG/HowTo/CreatingWidgets/CreatingWidgets.INT.udn
@@ -32,7 +32,7 @@ For this how-to, we are using the **Blueprint Third Person Template**. However,
 
 	![](CreatingWidgets02.png)
 
-	Upon opening a Widget Blueprint you will be presented with the UMG Editor Interface.
+	Upon opening a Widget Blueprint you will be presented with the UMG Editor Interface. The Widget Blueprint opens by default to the **Designer** tab, which is where you perform your visual layout of your UI elements. To script functionality for the Widget Blueprint, you can click the **Graph** button in the top right corner of the editor.
 
 	[REGION:lightbox]
 	[![](CreatingWidgets03.png)](CreatingWidgets03.png)
@@ -40,28 +40,20 @@ For this how-to, we are using the **Blueprint Third Person Template**. However,
 
 	[REGION:caption]
 	Click for full view. 
-	[/REGION]
-
-	The Widget Blueprint opens by default to the **Designer** tab, which is where you perform your visual layout of your UI elements. To script functionality for the Widget Blueprint, you can click the **Graph** button in the top right corner of the editor. 
+	[/REGION] 
 
-1.	Inside the **Palette** window, **left-click** and drag a **Text** widget into the **Visual Designer** graph. 
+1.	Inside the **Palette** window, **left-click** and drag a **Text** widget into the **Visual Designer** graph. The Palette window contains several different types of widgets that you can use to construct your UI elements. 
 
 	![](CreatingWidgets04.png)
 
-	The Palette window contains several different types of widgets that you can use to construct your UI elements. 
-
-1.	In the **Details** panel for the Text Widget, click the [](Engine/UMG/UserGuide/Anchors) dropdown and select to Anchor to the Center. 
+1.	In the **Details** panel for the Text Widget, click the [](Engine/UMG/UserGuide/Anchors) dropdown and select to Anchor to the Center. This will ensure that as the viewport size changes, the Text widget will stay anchored relative to the center of the screen.
 
-	![](CreatingWidgets05.png)
+	![](CreatingWidgets05.png) 
 
-	This will ensure that as the viewport size changes, the Text widget will stay anchored relative to the center of the screen. 
-
-1.	In the **Details** panel click the **Size to Content** option, change the **Text**, **Color and Opacity** and **Font Size** options. 
+1.	In the **Details** panel click the **Size to Content** option, change the **Text**, **Color and Opacity** and **Font Size** options. The Size to Content option as the name suggests will adjust the size of the widget to fit its content. We also changed the text to "Hello World!", changed the color of the text to green, and set the Font Size to 56.
 
 	![](CreatingWidgets06.png)
 
-	The Size to Content option as the name suggests will adjust the size of the widget to fit its content. We also changed the text to "Hello World!", changed the color of the text to green, and set the Font Size to 56.
-
 1.	Click the **Compile** button, then **Save** and close the Widget Blueprint. 
 
 	![](CreatingWidgets09.png)
@@ -82,42 +74,36 @@ For this how-to, we are using the **Blueprint Third Person Template**. However,
 
 1.	**Right-click** on the **Return Value** pin and select **Promote to Variable** and call it **SampleUI**. 
 
-	![](CreatingWidgets13.png)
+	[REGION:tip]
+	Whenever you create the Widget Blueprint (your UI), it's a good idea to promote that to a variable so you can access it later through Blueprint script.
+	[/REGION]
 
-	Whenever you create the Widget Blueprint (your UI), it's a good idea to promote that to a variable so you can access it later through Blueprint script. 
+	![](CreatingWidgets13.png) 
 
-1.	Drag off the out pin on the **Set** node and search for and add the **Add to Viewport** node. 
+1.	Drag off the out pin on the **Set** node and search for and add the **Add to Viewport** node. This is how we tell the Widget Blueprint to be displayed in-game. 
 
 	![](CreatingWidgets14.png)
 
-	This is how we tell the Widget Blueprint to be displayed in-game. 
-
 1.	**Right-click** in the graph and search for and add the **Enter** Keyboard Event. 
 
 	![](CreatingWidgets15.png)
 
-1.	Connect the **Pressed** pin to the executable pin of the **Create Sample UI Widget** node. 
+1.	Connect the **Pressed** pin to the executable pin of the **Create Sample UI Widget** node. At this point we have stated that when we press Enter, create the widget and store it as a variable before displaying it. 
 
 	![](CreatingWidgets16.png)
 
-	At this point we have stated that when we press Enter, create the widget and store it as a variable before displaying it. 
-
 1.	Hold **Ctrl** and drag in the **SampleUI** variable into the graph from the **My Blueprint** panel.
 
 	![](CreatingWidgets17.png)
 
-1.	Drag off the pin of the new node and search for and add the **Remove from Parent** node. 
+1.	Drag off the pin of the new node and search for and add the **Remove from Parent** node. The Remove from Parent node is how we tell Widget Blueprints to be removed from the viewport.
 
-	![](CreatingWidgets18.png)
+	![](CreatingWidgets18.png) 
 
-	The Remove from Parent node is how we tell Widget Blueprints to be removed from the viewport. 
-
-1.	Connect the **Release** pin of the **Enter** node to the **Remove from Parent** executable in pin. 
+1.	Connect the **Release** pin of the **Enter** node to the **Remove from Parent** executable in pin. Now when we press **Enter** the text will be displayed and when we release Enter, the text will be removed. 
 
 	![](CreatingWidgets19.png)
 
-	Now when we press **Enter** the text will be displayed and when we release Enter, the text will be removed. 
-
 1.	Close the **Level Blueprint**, then click the **Play** button from the Main Toolbar. 
 
 	![](CreatingWidgets20.png)
