#UE4 Docs: Applying SME feeback to these documents.

diff --git a/Engine/Documentation/Source/Engine/Performance/ReducingPackageSize/ReducingPackageSize.INT.udn b/Engine/Documentation/Source/Engine/Performance/ReducingPackageSize/ReducingPackageSize.INT.udn
--- a/Engine/Documentation/Source/Engine/Performance/ReducingPackageSize/ReducingPackageSize.INT.udn
+++ b/Engine/Documentation/Source/Engine/Performance/ReducingPackageSize/ReducingPackageSize.INT.udn
@@ -64,7 +64,7 @@ To turn on package compression, you need to do the following in the UE4 Editor.
 
 	![](T_Project_Packaging.png)
 
-1. Click on the small white triangle that is at the bottom of the Packaging settings to expose the **Advanced Project settings**.
+1. Click on the **small white triangle** that is at the bottom of the Packaging settings to expose the **Advanced Project settings**.
 
 	![](T_Expose_Advanced_Props.png)
 
@@ -75,7 +75,7 @@ To turn on package compression, you need to do the following in the UE4 Editor.
 If you have not packaged your game with the Create compressed cooked packages checkmark box enabled you should see a huge difference in size when your project is re-packaged with it enabled.
 In fact it is not uncommon for some project's APK packages to be reduced in size by up to 50% when **Create compressed cooked packages** is enabled.
 
-##  Setting up a Project's Levels
+## Setting up a Project's Levels
 One overlooked area that can lead to the bloating of APK package sizes is failing to properly setup your project's options, like which levels should be loaded by default, or which levels should be used between level transitions. 
 To set which level (or levels) should be used for this type of interaction, you can do the following.
 
@@ -92,6 +92,55 @@ To set which level (or levels) should be used for this type of interaction, you
 	![](T_Project_Settings_Default_Game_Map.png)
 
 1. Repeat this process until you have levels setup for each required map type.
+
+## Selecting which Content to Package or not Package
+In the Packaging section of your project settings you can specify which maps and content should or should not be packaged with your game.  To specify which maps should be included in your project you will need to do the following:
+
+[region:warning]
+If you do not specify which maps should be cooked then all maps will be cooked including any testing maps that you might have. Failing to not specify which maps should be cooked will result in your final packaged game being bigger than it needs to be. 
+[/region]
+
+1. First, open up the project's settings by going to the **Main Toolbar** and selecting **Edit** > **Project Settings**.
+
+	![](T_Project_Settings_Menu.png)
+
+1. Then, under **Project** in the **Packaging** section look for the **Packaging** option.
+
+	![](Content_To_Pack.png)
+
+1. Click on the **small white triangle** at the bottom of the Packaging section to expose the advances packing options.
+
+	![](Advanced_Packing_Section.png)
+
+1. Scroll to the bottom of the advanced Packaging options till you see a checkmark box for **Cook everything in the project content directory(ignore list of maps below)** option. The options that you see here will allow you to specify which content and maps should be packaged or should not be packaged with your project.
+
+	![](Advanced_Packing_Options.png)
+
+1. To include or exclude an item press the white plus sign next to the option you want to use to create a new entry. 
+
+	![](Add_New_Item.png)
+
+1. Then click on the three small white dots next to the new entry and from the windows that is displayed, choose which assets to include or exclude from the build.
+
+	![](Select_Items_To_Cook.png)
+
+1. When you have selected an asset, the entry box will now contain a link to where that asset is in your project's folder. 
+
+	![](Selected_Item_Location.png)
+
+|Property Name | Description|
+|--------------|------------|
+|Cook everything in the project content directory(ignore list of maps below.|Cook all things in the project content directory.|
+|Cook only maps(this only affects cook all).|Cook only maps(this only affect the cookall flag).|
+|Create compressed cooked packages.|Create compressed cooked packages(decreased deployment size.|
+|Do not include editor content in this package may cause game to crash / error if you are using this content|Skip cooking Editor content.|
+|List of maps to include in a packaged build|List of maps to include when no other map list is specified on commandline.|
+|Additional Asset Directories to Cook|Directories containing .uasset files that should always be cooked regardless of whether they're referenced by anything in your project. Note: These paths are relative to your project Content directory.|
+|Directories to never cook|Directories that should never be cooked.|
+|Additional Non-Asset Directories to Package.|Directories containing files that should always be added to the .pak file(if using a .pak file; otherwise they're copied as individual files) This is used to stage additional file that you manually load via the UFS (Unreal File System) file IO API. Note:These paths are relative to your project Content directory.|
+|Additional Non-Asset Directories to Copy.|Directories containing files that should always be copied when packaging your project, but are not supposed to be part of the .pak file. This is used to stage additional files that you manually load with using the UFS (Unreal File System) file IO API, ex, third-party libraries that perform their own internal file IO. Note:These paths are relative to your project Content directory. |
+
+
 ## Checking what Content will be Cooked
 To ensure that only the content that is related to your project is added to your APK file, you can check which content is being added to your project's APK by looking in the **Cooked** folder in your project.
 You can find the Cooked folder by going to **(ProjectName)\Saved\Cooked** and then choosing the folder with the format you cooked your project for.
@@ -109,7 +158,6 @@ Once the biggest assets have been determined, you can then examine the assets in
 ![](T_Cooked_Content.png)
 
 ## Development VS Shipping Build Size
-
 When trying to figure out the final size your project will be, keep in mind that the Development version of your project will be slightly larger than your Shipping build size.
 On the Medieval Match example game, the difference in size between a Development and Shipping build was around 14 percent.
 However, since each project has different requirements, the saving between the two different build types for your project could be more or less than a 14 percent savings.
@@ -136,8 +184,8 @@ If it is not referenced, you can simply delete it, however if it is referenced b
 Disabling any unused plug-ins inside of the UE4 Editor is another way to ensure that unnecessary content and code do not make it into your project's final APK package file.
 This is because some plug-ins require a certain set of base assets and code to work correctly.
 If the plug-in is not disabled, the assets and code that are required to make the plug-in work will not be removed from your project.
-While this might not reduce your project's size as much as reducing or removing a large Texture would, every little bit can help when trying to slim your project down to the 50MB size.
-To disable a plug-in inside of UE4, you need to do the following. 
+While this might not reduce your project's size as much as reducing or removing a large Texture would, every little bit can help when trying to slim your project down to the 100 MB size.
+To disable a plug-in inside of UE4, you need to do the following: 
 
 [region:note]
 	Make sure you fully test your project to ensure that the plug-ins you are disabling do not interfere with the functionality of your project.
@@ -163,7 +211,6 @@ To disable a plug-in inside of UE4, you need to do the following.
 	[/OBJECT]
 
 ## Package Black List
-
 With the release of Unreal Engine 4.9 and later you can now place a text file in your project's **Build/Platform/(Target Platform Folder)** directory that will tell the cooker to exclude partial or complete file paths from being packaged into your project.
 Projects can have multiple Black List files for Debug, Development, Test, and Shipping builds that can be setup to include or exclude whatever project data you want.
 You can even have different Black List files for each platform your project supports like one for Android, one for iOS, ect. 
@@ -183,7 +230,6 @@ An easy and simple way to go about this is to open up all your project's levels
 If you run into further problems and are not sure where they are coming from, check the Logcat for errors.
 
 ## Max Dynamic Point Lights
-
 With the release of UE 4.9 a new option has been introduced that can help reduce the number of shaders that are generated for dynamic lights. 
 Making this feature is particularly useful for mobile games that do not require dynamic lighting. 
 The property, **Max Dynamic Point Lights** can be found under the **Mobile** section of the **Project Settings**.
diff --git a/Engine/Documentation/Source/Platforms/Android/Reference/AndroidGettingStarted.INT.udn b/Engine/Documentation/Source/Platforms/Android/Reference/AndroidGettingStarted.INT.udn
--- a/Engine/Documentation/Source/Platforms/Android/Reference/AndroidGettingStarted.INT.udn
+++ b/Engine/Documentation/Source/Platforms/Android/Reference/AndroidGettingStarted.INT.udn
@@ -14,15 +14,14 @@ type:reference
 %Globals:OSSelect%
 
 ## Overview
-The Android Development Reference pages contains a high level explanation of the various software components and system settings you will need to make sure you have installed and setup on the computer you are using to develop your UE4 Android project.
+The Android Development Reference pages contain a high-level explanation of the various software components and system settings you will need to make sure you have installed and setup on the computer you are using to develop your UE4 Android project.
 
 ## Android SDK Requirements
-Here is what we currently require for Android game development:
+Here is what we currently require for Android game development in UE4:
 
-|Project Type | Components | Environment Variables |
-| ------------- | ---------- | --------------------- |
-| Blueprint Only Projects |  [INCLUDE:#bpcomponents] |      [INCLUDE:#bpenvvar] |
-| C++ Projects |   [INCLUDE:#codecomponents] |  [INCLUDE:#codeenvvar] |
+ | Components | Environment Variables |
+ | ---------- | --------------------- |
+ |   [INCLUDE:#codecomponents] |  [INCLUDE:#codeenvvar] |
 
 <!--
 [EXCERPT:BPComponents]
@@ -33,7 +32,7 @@ Here is what we currently require for Android game development:
 
 [EXCERPT:CodeComponents]
 [INCLUDE:#bpcomponents]
-* **Android NDK**
+* Android NDK
 [/EXCERPT:CodeComponents]
 
 [EXCERPT:BPEnvVar]
@@ -44,7 +43,7 @@ Here is what we currently require for Android game development:
 
 [EXCERPT:CodeEnvVar]
 [INCLUDE:#bpenvvar]
-* **NDKROOT**
+* NDKROOT
 [/EXCERPT:CodeEnvVar]
 -->
 
@@ -92,11 +91,21 @@ The installer for CodeWorks for Android 1R4 is included with Unreal Engine 4 for
 	[ENGINE INSTALL LOCATION]\Engine\Extras\Android\CodeWorksforAndroid-1R4-windows.exe
 
 The important pieces are:
-* SDK
-* NDK
-* JDK
-* Ant 
-* Nsight 
+
+* Android SDK (Base, Platform Tools, Build Tools)
+* Android 4.4.2 (API 19) - SDK Platform
+* Android 5.0.1 (API 21) - SDK Platform
+* Android SDK Support Library
+* Android Toolchain 
+	* Android NDK 11c
+	* Java SDK
+	* Apache Ant
+	* USB Driver
+	* Optional Install 
+		* Gradle
+* Developer Tools ( Optional Install)
+* Nsight Tegra, Visual Studio Edition
+* USB Driver
 	
 If you use the default installation set, a pretty minimal set of components will be installed. If you customize the install, you can add more items (like profilers, etc). Note that during installation, 
 one or two of the sub-installers may appear to hang, but they will eventually finish. The install could take as long as 30 to 45 minutes. If you never plan to write code, you may want to 
@@ -116,15 +125,24 @@ The installer for CodeWorks for Android 1R4 is included with Unreal Engine 4 for
 
 The important pieces are:
 
-* SDK
-* NDK
-* Ant 
+* Android SDK (Base, Platform Tools, Build Tools)
+* Android 4.4.2 (API 19) - SDK Platform
+* Android 5.0.1 (API 21) - SDK Platform
+* Android SDK Support Library
+* Android Toolchain 
+	* Android NDK 11c
+	* Java SDK
+	* Apache Ant
+	* USB Driver
+	* Optional Install 
+		* Gradle
+* Developer Tools 
+* Nsight Tegra, Visual Studio Edition (Optional Install)
+* USB Driver
 
 If you use the default installation set, a pretty minimal set of components will be installed. 
-If you customize the install, you can add more items (like profilers, etc). Note that during installation, one or two of the sub-installers may appear to hang, but they will eventually finish. 
-The install could take as long as 30 to 45 minutes on some computers. 
-If you never plan to write code, you may want to select **Customize**, and only install the SDK, Ant, and Java (skipping the NDK and Nsight debugger) to make for the quickest installation.
-The Nsight portion of CodeWorks for Android 1R4 is used for debugging an Android game on a device so it is not required to compile C++ code.
+If you customize the install, you can add more items (like profilers, etc). Note that during installation, one or two of the sub-installer may appear to hang, but they will eventually finish. The install could take as long as 30 to 45 minutes on some computers. 
+If you never plan to write code, you may want to select **Customize**, and only install the SDK, Ant, and Java (skipping the NDK and Nsigh debugger) to make for the quickest installation. The Nsight portion of CodeWorks for Android 1R4 is used for debugging an Android game on a device so it is not required to compile C++ code.
 
 ![](TADP_Mac_OS.png)
 
@@ -132,34 +150,6 @@ The Nsight portion of CodeWorks for Android 1R4 is used for debugging an Android
 
 -->
 
-### ADT (Android Development Tools)
-
-This is Google's recommended development environment. It contains the SDK, Ant, and Eclipse (which you may want to use for your own tools, but we do not support Eclipse with Unreal Engine 4).
- We recommend you install everything into C:\Android for simplicity.
-
-* [Download it here](http://developer.android.com/sdk/installing/bundle.html).
-
-In addition, you will need to install JDK 6 or later.
-
-* [Download it here](http://www.oracle.com/technetwork/java/javase/downloads/index.html).
-
-For code projects only, you will need to install the NDK:
-* [Download it here](http://developer.android.com/tools/sdk/ndk/index.html).
-
-### Android SDK (Software Development Kit)
-
-This is a more bare bones installation of the SDK (it does not include Eclipse or Ant).
-* [Download it here](http://developer.android.com/sdk/index.html).
-
-In addition, you will need to install JDK 6 or later.
-* [Download it here](http://www.oracle.com/technetwork/java/javase/downloads/index.html).
-
-You will need to install the Ant tool (version 1.8 or later).
-* [Download it here](http://ant.apache.org/) then unzip the .zip file.
-
-For code projects only, you will need to install the NDK:
-* [Download it here](http://developer.android.com/tools/sdk/ndk/index.html).
-
 
 ## Environment Variables
 
@@ -192,13 +182,15 @@ For code projects only, you will need to install the NDK:
 	[/PARAM]
 [/OBJECT]
 
+
+
 <!--
 [EXCERPT:ENVWindowsSetup]
 
 Unless you installed CodeWorks for Android 1R4 from the provided UE4 installer, you will need to verify that some environment variables are set properly. 
 
 [REGION:note]
-Once you have completed the installation (either by setting environment variables in the System Control Panel, or by installing CodeWorks for Android 1R4), you will need to restart both the Editor and Launcher 
+Once you have completed the installation (either by setting environment variables in the System Control Panel or by installing CodeWorks for Android 1R4), you will need to restart both the Editor and Launcher 
 applications. This is due to how Windows applications store environment variables and then pass the old versions along to applications they start. 
 
 If you continue seeing the message about the Android SDK not being installed properly, it is recommended that you reboot your computer.
@@ -227,12 +219,12 @@ Here are the variables that we depend on, and where they should point:
 * Set this to where the JDK is installed. It will probably be called something like `jdk1.6.0_24`.
 * It must contain a directory called `bin` with `javac.exe` in it.
 
-**ANT_HOME** [Not needed if you installed ADT]
+**ANT_HOME**
 * Set this to where you unzipped Ant.
 * It must contain a directory called `bin` with `ant.bat` in it.
 
-**NDKROOT** [Not needed if you do not compile C++ code]
-* Set this to where you unzipped the NDK. It will probably be called something like `android-ndk-r9c`.
+**NDKROOT** 
+* Set this to where you unzipped the NDK. It will probably be called something like `android-ndk-r11c`.
 * It must contain a file called `ndk-build.cmd`
 
 [/EXCERPT:ENVWindowsSetup]
@@ -241,7 +233,7 @@ Here are the variables that we depend on, and where they should point:
 Unless you installed CodeWorks for Android 1R4, you will need to verify that some environment variables are set properly. 
 
 [region:note]
-Once you have completed the installation (either by editing your .bash_profile file, or by installing CodeWorks for Android 1R4), you will need to restart both the Editor and Launcher applications.
+Once you have completed the installation (either by editing your .bash_profile file or by installing CodeWorks for Android 1R4), you will need to restart both the Editor and Launcher applications.
 [/region]
 
 To set an environment variable you need to edit the **.bash_profile** file in your home directory. 
@@ -254,7 +246,7 @@ If you’re not familiar with Unix terminal commands please do the following.
 1. Copy the following lines to the file (replacing paths to point to where you installed SDK and tools) and save it:
 
 	* export ANDROID_HOME="[REDACTED_USER_HOME]/Devel/NVPACK/android-sdk-macosx"
-	* export NDKROOT="[REDACTED_USER_HOME]/Devel/NVPACK/android-ndk-r10c"
+	* export NDKROOT="[REDACTED_USER_HOME]/Devel/NVPACK/android-ndk-r11c"
 	* export ANT_HOME="[REDACTED_USER_HOME]/Devel/NVPACK/apache-ant-1.8.2"
 
 Here are the name of the variables that we depend on and where they should point to.
@@ -262,11 +254,11 @@ Here are the name of the variables that we depend on and where they should point
 * **ANDROID_HOME**
 	* Set this to the directory where you installed the Android SDK or the SDK directory if you installed ADT. 
 	* It must contain a directory called platform-tools.
-* **ANT_HOME** [Not needed if you installed ADT]
+* **ANT_HOME** 
 	* Set this to where you unzipped Ant.
 	* It must contain a directory called bin with a file called ANT in it.
-* **NDKROOT** [Not needed if you do not compile C++ code]
-	* Set this to where you unzipped the NDK and look for something similar to **android-ndk-r9c**.
+* **NDKROOT** 
+	* Set this to where you unzipped the NDK and look for something similar to **android-ndk-r11c**.
 	* It must contain a file called **ndk-build**.
 
 [/EXCERPT:ENVMacOSSetup]
@@ -313,11 +305,11 @@ If you do not see your device listed in the **Launch** menu, here are some thing
 - Open a command prompt (Press **Windows+R**, and run `cmd.exe`) 
 	- Type the following command: `adb devices`
 	- If you do not see any devices listed, your USB driver is not installed properly. How to install the driver depends on the model of your Android device. Search online for your device model and 
-	"usb driver", then install the driver. 
+	"USB driver", then install the driver. 
 	- Additionally, you will need to make sure your device has Developer mode enabled with USB debugging is turned on. See the [Android Developer Using Hardware Devices guide](http://developer.android.com/tools/device.html) for more information (specifically items 2 and 3).  
-	- Afterwards, make sure the `adb devices` command works; without this, the engine will not be able to do anything with your Android device. 
+	- Afterwards make sure the `adb devices` command works; without this, the engine will not be able to do anything with your Android device. 
 	
-If your device shows up in the Launch menu, but gets stuck deploying to the device, it is most likely because the device is not authorized for your PC to talk to it. If you click on the **Show Log** line 
+If your device shows up in the Launch menu but gets stuck deploying to the device, it is most likely because the device is not authorized for your PC to talk to it. If you click on the **Show Log** line 
 in the progress bar, you may see that the last line indicates that it is calling `adb uninstall`. 
 - As above, use `cmd.exe` to run the `adb devices` command. 
 - If the device is listed but shows unauthorized, that indicates that you need to get the device to ask if you want to authorize the PC. 
@@ -333,10 +325,10 @@ If you do not see your device listed in the Launch menu, here are some things to
 	* Type the following command into the terminal: `adb devices`.
 	* If you do not see any devices listed, your USB driver is not installed properly. How to install the driver depends on the model of your Android device. Search online for your device model and **USB driver**, then install the required driver. 
 	* Additionally, you will need to make sure your device has Developer mode enabled with USB debugging is turned on. See the [Android Developer Using Hardware Devices](http://developer.android.com/tools/device.html) guide for more information (specifically items 2 and 3).
-	* Afterwards, make sure the `adb devices` command works; without this, the engine will not be able to do anything with your Android device.
+	* Afterwards make sure the `adb devices` command works; without this, the engine will not be able to do anything with your Android device.
 
-If your device shows up in the **Launch menu**, but gets stuck deploying to the device, it is most likely because the device is not authorized to interact with your Mac. 
-If you click on the **Show Log** line in the progress bar, and see that the last line indicates that it is calling `adb uninstall`, try the following.
+If your device shows up in the **Launch menu** but gets stuck deploying to the device, it is most likely because the device is not authorized to interact with your Mac. 
+If you click on the **Show Log** line in the progress bar and see that the last line indicates that it is calling `adb uninstall`, try the following.
 * Open up the Terminal app and run the `adb devices` command. 
 * If the device is listed but shows unauthorized, that indicates that you need to get the device to ask if you want to authorize the Mac. 
 * It is recommended you check the **Always allow** check box to always allow your Mac to talk to this device.
@@ -351,11 +343,13 @@ Not all Android devices are made the same. In particular, there are 4 different
 [EXCERPT:AndroidFormats]
 |Format | Description|
 |-------|------------|
-|DXT |This is what your desktop computer uses (also used by Tegra-based devices).|
-|ATC |This is used in Qualcomm Snapdragon based devices.|
-|PVR |This is used in ImgTec PowerVR based devices.|
-|ETC1 |This is supported on all devices, but does not support textures with an alpha channel, so those textures are uncompressed, making your game's download large and use more runtime memory.|
-|ETC2 |This is used by MALI based devices and some others.|
+|ETC1 |Supported by all Android based devices but cannot compress alpha textures (they are stored uncompressed).  Recommend using an RGB and a separate alpha texture if need alpha to get better compression.|
+|ETC2 |Supported by all OpenGL 3.x class devices and supports alpha compression.|
+|ATC |Supported by Qualcomm Adreno GPUs and supports alpha compression.|
+|DXT |Supported by Nvidia Tegra GPUs and supports alpha compression.|
+|PVRTC |supported by PowerVR GPUs and supports alpha compression.|
+|ASTC  |Latest Texture compression format allowing more quality control by specifying block size and supports alpha compression. Available on some devices at this point and will be required for Vulkan Level 1.
+|
 [/EXCERPT:AndroidFormats]
 
 When you deploy or package your game for Android, the data will be converted to a format usable by an Android device (we call this process "cooking"). Because different devices support different 
@@ -371,10 +365,5 @@ If you are developing a project that is aimed at using one of these devices, it
 [/region]
 
 ### Choosing a Format
-You most likely will want to choose an optimal format for your device. If you know the type of GPU in your device, you can look at the table above to choose a format.
-
-However, if you do not know, then you can choose **Android** or **Android (ETC1)**. When you run the game, if you tap the screen with 4 fingers at once, a dialog will pop up. This dialog is usually used to 
-enter a console command (like `stat fps`), but it will also display the formats that your device supports. Once you have that list, you can use that to choose a more ideal type when deploying/packaging.
-
-If you choose a format that your device does not support, it will fail to load. 
+You most likely will want to choose an optimal format for your device. If you know the type of GPU in your device, you can look at the table above to choose a format. However, if you do not know, then you can choose **Android** or **Android (ETC1)**. When you run the game, if you tap the screen with 4 fingers at once, a dialog will pop up. This dialog is usually used to enter a console command (like `stat fps`), but it will also display the formats that your device supports. Once you have that list, you can use that to choose an ideal type when deploying/packaging. If you choose a format that your device does not support, it will fail to load. 
 
diff --git a/Engine/Documentation/Source/Platforms/GearVR/Prerequisites/GearVRPrerequisites.INT.udn b/Engine/Documentation/Source/Platforms/GearVR/Prerequisites/GearVRPrerequisites.INT.udn
--- a/Engine/Documentation/Source/Platforms/GearVR/Prerequisites/GearVRPrerequisites.INT.udn
+++ b/Engine/Documentation/Source/Platforms/GearVR/Prerequisites/GearVRPrerequisites.INT.udn
@@ -72,6 +72,8 @@ Then, place that folder in **Engine/Build/Android/Java**.
 	
 	![](Platforms\GearVR\GVR_OSIG_File_Location.png)
 
+* You can also place the OSIG file in your project's **Build/Android/assets** directory instead of Engine/Build/Android/Java. Keep in mind that adding your OSIG file in this manner will result in your OSIG file only being used in the project that it is added to and not for any Gear VR project.
+
 It is very important that you make sure all items listed in the Gear VR Prerequisites section are completed before moving on to the next section. 
 Skipping a step could result in your project not working with the Gear VR Headset. 
 
