#UE4 Docs: minor text and formatting updates

diff --git a/Engine/Documentation/Source/Engine/UMG/HowTo/InWorldWidgetInteraction/InWorldWidgetInteraction.INT.udn b/Engine/Documentation/Source/Engine/UMG/HowTo/InWorldWidgetInteraction/InWorldWidgetInteraction.INT.udn
--- a/Engine/Documentation/Source/Engine/UMG/HowTo/InWorldWidgetInteraction/InWorldWidgetInteraction.INT.udn
+++ b/Engine/Documentation/Source/Engine/UMG/HowTo/InWorldWidgetInteraction/InWorldWidgetInteraction.INT.udn
@@ -45,13 +45,11 @@ For this how-to, we are using the **Blueprint First Person Template** with **Sta
 
 	![](WidgetInteraction_05.png)
 
-1.	Connect to the **Pressed** pin of the **Right Mouse Button** node and set the **Key** to **Left Mouse Button**. 
+1.	Connect to the **Pressed** pin of the **Right Mouse Button** node and set the **Key** to **Left Mouse Button**. By default UMG registers click events as a result of a Left Mouse Button key press, using this node will allow us to simulate that when the Right Mouse Button is pressed. 
+	You could use any Key Input event to call this function (such as a trigger press or button press to simulate a Left Mouse Button click). 
 
 	![](WidgetInteraction_06.png)
 
-	By default UMG registers click events as a result of a Left Mouse Button key press, using this node will allow us to simulate that when the Right Mouse Button is pressed. 
-	You could use any Key Input event to call this function (such as a trigger press or button press to simulate a Left Mouse Button click). 
-
 1.	Add a **Release Pointer Key** node (set to **Left Mouse Button**) and connect it to the **Released** pin of the Right Mouse Button node. 
 
 	![](WidgetInteraction_07.png)
@@ -60,7 +58,7 @@ For this how-to, we are using the **Blueprint First Person Template** with **Sta
 
 	![](WidgetInteraction_08.png)
 
-1.	In the **Visual Designer**, remove the **Canvas Panel** then add a **Button** with a **Text** widget on top of the button. 
+1.	In the **Visual Designer**, remove the **Canvas Panel** then add a **Button** with a **Text** widget on top of the button. For this example we are purely adding a button that can be clicked on. The text on the button we will indicate the number of times the button has been clicked by the player. You can create any type of interactive widget and interact with it using the Widget Interaction component (things like sliders, combo boxes, check boxes, etc. can be interacted with in world space).
 
 	[REGION:lightbox]
 	[![](WidgetInteraction_09.png)](WidgetInteraction_09.png)
@@ -68,10 +66,7 @@ For this how-to, we are using the **Blueprint First Person Template** with **Sta
 
 	[REGION:caption]
 	Click for full view. 
-	[/REGION]
-
-	For this example we are purely adding a button that can be clicked on. The text on the button we will indicate the number of times the button has been clicked by the player. 
-	You can create any type of interactive widget and interact with it using the Widget Interaction component (things like sliders, combo boxes, check boxes, etc. can be interacted with in world space). 
+	[/REGION] 
 
 1.	In the **Details** panel for the **Button**, under **Style/Hovered** change the **Tint** to a different color. 
 
@@ -81,42 +76,32 @@ For this how-to, we are using the **Blueprint First Person Template** with **Sta
 
 	![](WidgetInteraction_11.png)
 
-1.	In the **Graph**, create a **Text** Variable called **Value** then **Compile** and set **Value** to **0**. 
+1.	In the **Graph**, create a **Text** Variable called **Value** then **Compile** and set **Value** to **0**. This is the variable that we will update when the player clicks on the button which we will bind to our Text widget. 
 
 	![](WidgetInteraction_12.png)
 
-	This is the variable that we will update when the player clicks on the button which we will bind to our Text widget. 
-
 1.	On the **Designer** tab for the **Text** widget, click the **Bind** option next to **Text** and select the **Value** property.
 
 	![](WidgetInteraction_13.png)
 
-1.	In the **Details** panel for the **Button** widget under **Events**, click the **+** next to **On Clicked**. 
+1.	In the **Details** panel for the **Button** widget under **Events**, click the **+** next to **On Clicked**. This will fire whenever the button is clicked on which is called from inside our **FirstPersonCharacter** Blueprint. 
 
 	![](WidgetInteraction_14.png)
 
-	This will fire whenever the button is clicked on which is called from inside our **FirstPersonCharacter** Blueprint. 
-
-1.	On the **Graph** tab, hold **Ctrl** and drag in the **Value** variable then off it use **To String** followed by **String To Int**. 
+1.	On the **Graph** tab, hold **Ctrl** and drag in the **Value** variable then off it use **To String** followed by **String To Int**. We are converting the Text to a String, then from a String to an Int so that we can increase the value when the button is clicked. 
 	
 	![](WidgetInteraction_15.png)
 
-	We are converting the Text to a String, then from a String to an Int so that we can increase the value when the button is clicked. 
-
 1.	Off the **String To Int** pin, use an **Integer + Integer** node set to **+ 1**. 
 
-1.	Hold **Alt** and drag in the **Value** variable to set it, then connect the remaining pins to the **OnClicked** event as shown. 
+1.	Hold **Alt** and drag in the **Value** variable to set it, then connect the remaining pins to the **OnClicked** event as shown. The **ToText(int)** node will be automatically created when you connect the out pin of **Integer + Integer** to the text input pin on **Value**. 
 
 	![](WidgetInteraction_16.png)
 
-	The **ToText(int)** node will be automatically created when you connect the out pin of **Integer + Integer** to the text input pin on **Value**. 
-
-1.	In the **Content Browser**, create a new **Blueprint** based on **Actor** called **ExampleWidget** and add a **Widget** component to it. 
+1.	In the **Content Browser**, create a new **Blueprint** based on **Actor** called **ExampleWidget** and add a **Widget** component to it. This will represent the 3D version of our **InteractiveWidget** Widget Blueprint. 
 
 	![](WidgetInteraction_18.png)
 
-	This will represent the 3D version of our **InteractiveWidget** Widget Blueprint. 
-
 1.	In the **Details** panel of the **Widget** component under **User Interface**, change the **Widget Class** to **InteractiveWidget**. 
 
 	![](WidgetInteraction_19.png)
@@ -130,7 +115,7 @@ For this how-to, we are using the **Blueprint First Person Template** with **Sta
 	
 ## End Result
 
-When you play in the editor, you should have something similar to below:
+When you play in the editor, you should have something similar to below. You can use the **Left Mouse Button** to fire your weapon and when pointing at the button, you enter the **Hovered** state (causing it to change its style). When you **Right Mouse Button** click on the button, the text value will increase each time it is pressed. 
 
 [OBJECT:EmbeddedVideo]
 [PARAMLITERAL:width]
@@ -144,11 +129,6 @@ bOZy7H9IMOI
 [/PARAMLITERAL]
 [/OBJECT]
 
-You can use the **Left Mouse Button** to fire your weapon and when pointing at the button, you enter the **Hovered** state (causing it to change its style). 
-When you **Right Mouse Button** click on the button, the text value will increase each time it is pressed. 
-
-As noted during the steps, we are using the Right Mouse Button to register click events however you can use gamepad button clicks to register Pointer Key Presses/Releases or (as shown below) Motion Controller Trigger presses. 
+As noted during the steps, we are using the Right Mouse Button to register click events however you can use gamepad button clicks to register Pointer Key Presses/Releases or (as shown below) Motion Controller Trigger presses. Here we point our Right Motion Controller (with an attached Widget Interaction component) at the button and press the Right Trigger to register a button click.
 
 ![](WidgetInteraction_21.png)
-
-Here we point our Right Motion Controller (with an attached Widget Interaction component) at the button and press the Right Trigger to register a button click.
diff --git a/Engine/Documentation/Source/Engine/UMG/HowTo/VirtualKeyboards/ActorSetup/ActorSetup.INT.udn b/Engine/Documentation/Source/Engine/UMG/HowTo/VirtualKeyboards/ActorSetup/ActorSetup.INT.udn
--- a/Engine/Documentation/Source/Engine/UMG/HowTo/VirtualKeyboards/ActorSetup/ActorSetup.INT.udn
+++ b/Engine/Documentation/Source/Engine/UMG/HowTo/VirtualKeyboards/ActorSetup/ActorSetup.INT.udn
@@ -31,28 +31,19 @@ With everything in place, we give it a test to see the end result which prints *
 
 ## Steps
 
-1.	Open the **KeypadDisplay** Blueprint and add a **Widget** component with class set to **KeypadDisplayWB** and **Draw at Desired Size**. 
+1.	Open the **KeypadDisplay** Blueprint and add a **Widget** component with class set to **KeypadDisplayWB** and **Draw at Desired Size**. Draw at Desired Size is a quick way to cause the render target to automatically match the desired size inside the Widget Blueprint.
 
 	![](Keypad15.png)
 
-	Draw at Desired Size is a quick way to cause the render target to automatically match the desired size inside the Widget Blueprint.
-
 1.	Open the **Keypad** Blueprint and add a **Widget** component with class set to **KeypadWB** and **Draw at Desired Size**. 
 
-1.	Uncheck the **Window Focusable** option for this Widget Component. 
+1.	Uncheck the **Window Focusable** option for this Widget Component. We don't want this component to take focus away from the Keypad Display widget, so this needs to be unchecked. 
 
 	![](Keypad16.png)
 
-	We don't want this component to take focus away from the Keypad Display widget, so this needs to be unchecked. 
-
-1.	Also inside **Keypad**, add a **WidgetInteraction** component with **Enable Hit Testing** set to false. 
-
-	![](Keypad17.png)
+1.	Also inside **Keypad**, add a **WidgetInteraction** component with **Enable Hit Testing** set to false. We add this WidgetInteraction component so that we don't have to get a reference to the WidgetInteraction component in the player character and we disable Hit Testing as we don't want it to generate any hover effects. WidgetInteraction components by default are set to **Virtual User Index 0** when they are created and all of them communicate using the Virtual User Index ID. Since the WidgetInteraction component in the player character is also using Virtaul User Index 0, this is the same as getting a reference to it.
 
-	We add this WidgetInteraction component so that we don't have to get a reference to the WidgetInteraction component in the player character and we disable Hit Testing as we don't want it to generate any hover effects. 
-
-	WidgetInteraction components by default are set to **Virtual User Index 0** when they are created and all of them communicate using the Virtual User Index ID. 
-	Since the WidgetInteraction component in the player character is also using Virtaul User Index 0, this is the same as getting a reference to it. 
+	![](Keypad17.png) 
 
 1.	On the Graph of **Keypad**, off **Event Begin Play** take the **Widget** component and **Get User Widget Object** to **Cast to KeypadWB**. 
 
@@ -60,12 +51,10 @@ With everything in place, we give it a test to see the end result which prints *
 
 1.	Off the **As Keypad WB** pin, **Set Widget Interaction** and pass through the **WidgetInteraction** component from Step 4. 
 
-1.	Also off the **As Keypad WB** pin, call the **Assign Enter Pressed** which will create an **Enter Pressed** Custom Event. 
+1.	Also off the **As Keypad WB** pin, call the **Assign Enter Pressed** which will create an **Enter Pressed** Custom Event. Now when we click on the Enter button this custom event will be called which we can use to check if the correct code was entered. 
 
 	![](Keypad19.png)
 
-	Now when we click on the Enter button this custom event will be called which we can use to check if the correct code was entered. 
-
 1.	Create a **Variable** of the **Keypad Display** Actor type called **TargetDisplay** and set it to **Editable**. 
 
 	![](Keypad20.png)
diff --git a/Engine/Documentation/Source/Engine/UMG/HowTo/VirtualKeyboards/CharacterSetup/CharacterSetup.INT.udn b/Engine/Documentation/Source/Engine/UMG/HowTo/VirtualKeyboards/CharacterSetup/CharacterSetup.INT.udn
--- a/Engine/Documentation/Source/Engine/UMG/HowTo/VirtualKeyboards/CharacterSetup/CharacterSetup.INT.udn
+++ b/Engine/Documentation/Source/Engine/UMG/HowTo/VirtualKeyboards/CharacterSetup/CharacterSetup.INT.udn
@@ -32,27 +32,19 @@ For this how-to, we are using the **Blueprint First Person Template** with **Sta
 
 1.	Inside the **Content/FirstPersonBP/Blueprints** folder, open the **FirstPersonCharacter** Blueprint. 
 
-1.	Add a **WidgetInteraction** component attached to the **Sphere** then zero out the **Location** and **Show Debug** in the **Details** panel. 
+1.	Add a **WidgetInteraction** component attached to the **Sphere** then zero out the **Location** and **Show Debug** in the **Details** panel. Our Widget Interaction component will now point in the direction that player points the gun indicated in-game by a debug marker.
 
 	![](Keypad01.png)
 
-	Our Widget Interaction component will now point in the direction that player points the gun indicated in-game by a debug marker.
-
 1.	In the **Graph**, add a **Right Mouse Button** Event and drag in the **Widget Interaction** Component. 
 
-1.	Off **Widget Interaction** add **Press Pointer Key** and **Release Pointer Key** both set to **Left Mouse Button** and connect as shown. 
-
-	![](Keypad02.png)
-
-	Widgets respond to mouse/touch interaction via Left Mouse Button click or touch controls and we can simulate that using the **Press Pointer Key** function. 
-	In our example we use the **Right Mouse Button** to simulate a Left Mouse Button click, however you could use a gamepad button press or even Motion Controller Trigger press to fire off the Press Pointer Key function. 
+1.	Off **Widget Interaction** add **Press Pointer Key** and **Release Pointer Key** both set to **Left Mouse Button** and connect as shown. Widgets respond to mouse/touch interaction via Left Mouse Button click or touch controls and we can simulate that using the **Press Pointer Key** function. In our example we use the **Right Mouse Button** to simulate a Left Mouse Button click, however you could use a gamepad button press or even Motion Controller Trigger press to fire off the Press Pointer Key function.
 
-1.	In the **Content Browser** create two **Class Blueprints** based on **Actor** and two **Widget Blueprints**. 
+	![](Keypad02.png) 
 
-	![](Keypad03.png)
+1.	In the **Content Browser** create two **Class Blueprints** based on **Actor** and two **Widget Blueprints**. For the two Actor Blueprints, call one **Keypad** (to represent the keypad entry) and the other **KeypadDisplay** (to represent and display the values entered from the Keypad). For the Widget Blueprints use the same names as the Actor Blueprints however to distinguish them, suffix them with **WB** (for Widget Blueprint).
 
-	For the two Actor Blueprints, call one **Keypad** (to represent the keypad entry) and the other **KeypadDisplay** (to represent and display the values entered from the Keypad). 
-	For the Widget Blueprints use the same names as the Actor Blueprints however to distinguish them, suffix them with **WB** (for Widget Blueprint). 
+	![](Keypad03.png) 
 
 
 ## End Result
