#UE4 Docs: removed hardcoded image sizes

diff --git a/Engine/Documentation/Source/Engine/UMG/UserGuide/UserGuide.INT.udn b/Engine/Documentation/Source/Engine/UMG/UserGuide/UserGuide.INT.udn
--- a/Engine/Documentation/Source/Engine/UMG/UserGuide/UserGuide.INT.udn
+++ b/Engine/Documentation/Source/Engine/UMG/UserGuide/UserGuide.INT.udn
@@ -151,7 +151,7 @@ version: 4.9
 	[REGION:topics third]
 	[OBJECT:TopicCompact]
 		[PARAM:image]
-			![](AnchorsTopic.png)(w:283 h:117)
+			![](AnchorsTopic.png)
 		[/PARAM]
 		[PARAM:title]
 			Anchors
@@ -162,7 +162,7 @@ version: 4.9
 	[/OBJECT]
 	[OBJECT:TopicCompact]
 		[PARAM:image]
-			![](AnimationTopic.png)(w:283 h:117)
+			![](AnimationTopic.png)
 		[/PARAM]
 		[PARAM:title]
 			Animation
@@ -173,7 +173,7 @@ version: 4.9
 	[/OBJECT]
 	[OBJECT:TopicCompact]
 		[PARAM:image]
-			![](EventsTopic.png)(w:283 h:117)
+			![](EventsTopic.png)
 		[/PARAM]
 		[PARAM:title]
 			Events
@@ -184,7 +184,7 @@ version: 4.9
 	[/OBJECT]
 	[OBJECT:TopicCompact]
 		[PARAM:image]
-			![](PropertyBindingTopic.png)(w:283 h:117)
+			![](PropertyBindingTopic.png)
 		[/PARAM]
 		[PARAM:title]
 			Property Binding
@@ -195,7 +195,7 @@ version: 4.9
 	[/OBJECT]
 	[OBJECT:TopicCompact]
 		[PARAM:image]
-			![](ScalingTopic.png)(w:283 h:117)
+			![](ScalingTopic.png)
 		[/PARAM]
 		[PARAM:title]
 			DPI Scaling
@@ -206,7 +206,7 @@ version: 4.9
 	[/OBJECT]
 	[OBJECT:TopicCompact]
 		[PARAM:image]
-			![](FontsTopic.png)(w:283 h:117)
+			![](FontsTopic.png)
 		[/PARAM]
 		[PARAM:title]
 			Fonts
@@ -217,7 +217,7 @@ version: 4.9
 	[/OBJECT]
 	[OBJECT:TopicCompact]
 		[PARAM:image]
-			![](StylingTopic.png)(w:283 h:117)
+			![](StylingTopic.png)
 		[/PARAM]
 		[PARAM:title]
 			Styling
@@ -228,7 +228,7 @@ version: 4.9
 	[/OBJECT]
 	[OBJECT:TopicCompact]
 		[PARAM:image]
-			![](SlotTopic.png)(w:283 h:117)
+			![](SlotTopic.png)
 		[/PARAM]
 		[PARAM:title]
 			Slots
