#UE4 Docs: updated metadata

diff --git a/Engine/Documentation/Source/Engine/Blueprints/UserGuide/Variables/Structs/Structs.INT.udn b/Engine/Documentation/Source/Engine/Blueprints/UserGuide/Variables/Structs/Structs.INT.udn
--- a/Engine/Documentation/Source/Engine/Blueprints/UserGuide/Variables/Structs/Structs.INT.udn
+++ b/Engine/Documentation/Source/Engine/Blueprints/UserGuide/Variables/Structs/Structs.INT.udn
@@ -1,4 +1,4 @@
-Availability:Docs
+Availability:Public
 Title:Struct Variables in Blueprints
 Crumbs: %ROOT%, Engine, Engine/Blueprints, Engine/Blueprints/Scripting
 Description: Blueprint struct variables allow you to store different data types that contain related information together.
