#UE4 Docs: minor text and formatting updates

diff --git a/Engine/Documentation/Source/Engine/UI/LevelEditor/Details/LevelEditorDetails.INT.udn b/Engine/Documentation/Source/Engine/UI/LevelEditor/Details/LevelEditorDetails.INT.udn
--- a/Engine/Documentation/Source/Engine/UI/LevelEditor/Details/LevelEditorDetails.INT.udn
+++ b/Engine/Documentation/Source/Engine/UI/LevelEditor/Details/LevelEditorDetails.INT.udn
@@ -75,25 +75,29 @@ If there are properties within an Actor that you frequently change or want quick
 Above, we checked two options as **Favorites**, indicated by the star icon next to their property name for easy access.
 [/REGION]
 
-You can enable **Favorites** from inside the **Editor Preferences**. 
+**To enable Favorites:**
 
-![](EditorPrefs.png)
+1. In the **Edit** menu, select **Editor Preferences**.
 
-Under **Experimental**, check the **Enable Details Panel Favorites** option. 
+	![](EditorPrefs.png)
 
-![](Experimental.png)
+1. Under **Experimental**, check the **Enable Details Panel Favorites** option. 
+
+	![](Experimental.png)
 
 [REGION:note]
 You may need to restart the editor for the changes to be applied. 
 [/REGION]
 
-With the option enabled, you can click the star icon next to a property to mark it as a favorite. 
+**To mark a Property as a Favorite:**
+
+1. With the option enabled, click the star icon next to a property in any Details panel. 
 
-![](ClickFavorite.png)
+	![](ClickFavorite.png)
 
-This will list the property (along with any other marked favorites) under the **Favorites** section of the panel. 
+1. The property (along with any other marked favorites) is listed under the **Favorites** section of the panel. 
 
-![](MarkedFavorites.png)
+	![](MarkedFavorites.png)
 
 [REGION:note]
 Some properties may not offer the ability to favorite them, due to the complexity of their customization. 
