#UE4 Docs: Applying some feedback dealing with image comp.

diff --git a/Engine/Documentation/Source/Enterprise/Automotive/AutomotiveMaterialPack/AutomotiveMaterialPack.INT.udn b/Engine/Documentation/Source/Enterprise/Automotive/AutomotiveMaterialPack/AutomotiveMaterialPack.INT.udn
--- a/Engine/Documentation/Source/Enterprise/Automotive/AutomotiveMaterialPack/AutomotiveMaterialPack.INT.udn
+++ b/Engine/Documentation/Source/Enterprise/Automotive/AutomotiveMaterialPack/AutomotiveMaterialPack.INT.udn
@@ -55,13 +55,13 @@ In the following section, we will take a look at some of the parameters that all
 
 	[OBJECT:ComparisonSlider]
 		[PARAM:before]
-		![Scale Aware Tiling](AP_02_SMP_01_Scale-Aware.png)
+		![Normal Tiling](AP_02_SMP_02_NormalUV.png)
 		[/PARAM]
 		[PARAM:after]
-		![Normal Tiling](AP_02_SMP_02_NormalUV.png)
+		![Scale Aware Tiling](AP_02_SMP_01_Scale-Aware.png)
 		[/PARAM]
 	[/OBJECT]
-
+	
 * **World Aligned:** World Aligned allows objects that have incorrect or poorly setup UV's to correctly display the applied Material.
 
 	[OBJECT:ComparisonSlider]
@@ -133,10 +133,10 @@ The Carbon Fiber Material makes use of the new Dual Normal option that is availa
 
 [OBJECT:ComparisonSlider]
 	[PARAM:before]
-	![Dual Normal On](AP_04_CarbonFiber.png)
+	![Dual Normal Off](AP_04_CarbonFiber_No_2nd_Normal.png)
 	[/PARAM]
 	[PARAM:after]
-	![Dual Normal Off](AP_04_CarbonFiber_No_2nd_Normal.png)
+	![Dual Normal On](AP_04_CarbonFiber.png)
 	[/PARAM]
 [/OBJECT]
 
@@ -222,10 +222,10 @@ When looking at the rubber Material Instance you will notice that it has a tire-
 
 [OBJECT:ComparisonSlider]
 	[PARAM:before]
-	![Correct Way](AP_09_Rubber_04_Right.png)
+	![Wrong Way](AP_09_Rubber_04_Wrong.png)
 	[/PARAM]
 	[PARAM:after]
-	![Wrong Way](AP_09_Rubber_04_Wrong.png)
+	![Correct Way](AP_09_Rubber_04_Right.png)
 	[/PARAM]
 [/OBJECT]
 
@@ -235,13 +235,13 @@ When using the sidewall option for a tire you will need to make sure that the UV
 
 [OBJECT:ComparisonSlider]
 	[PARAM:before]
-	![Correct UV Layout](AP_09_Rubber_05_UV_LineUp.png)
+	![Wrong UV Layout](AP_09_Rubber_06_UV_Wrong.png)
 	[/PARAM]
 	[PARAM:after]
-	![Wrong UV Layout](AP_09_Rubber_06_UV_Wrong.png)
+	![Correct UV Layout](AP_09_Rubber_05_UV_LineUp.png)
 	[/PARAM]
 [/OBJECT]
-
+	
 [REGION:lightbox]
 [![](AP_Tire_Sidewall.png)(w:800)](AP_Tire_Sidewall.png)
 [/REGION]
