#UE4 Docs: minor text and formatting updates

diff --git a/Engine/Documentation/Source/Engine/Sequencer/HowTo/ImportExportEDLs/ImportExportEDLs.INT.udn b/Engine/Documentation/Source/Engine/Sequencer/HowTo/ImportExportEDLs/ImportExportEDLs.INT.udn
--- a/Engine/Documentation/Source/Engine/Sequencer/HowTo/ImportExportEDLs/ImportExportEDLs.INT.udn
+++ b/Engine/Documentation/Source/Engine/Sequencer/HowTo/ImportExportEDLs/ImportExportEDLs.INT.udn
@@ -26,7 +26,7 @@ We are also using Adobe Premiere Pro CS6 as our external editing software; howev
 
 1.	Make sure your cinematic you want to export includes a **Shot Track** containing your shots, then click the **Render Movie** button.
 
-1.	Define an **Output Directory** and **Filename Format** using the **{shot}** token and make sure **Write Edit Decision List** is checked. 
+1.	Define an **Output Directory** and **Filename Format** using the **{shot}** token and make sure **Write Edit Decision List** is checked. Here we show the sample cinematic that we want to export along with an EDL and click the **Render Movie** button to bring up the settings dialog window. Inside the settings window, we create a new folder to export the cinematic to and prefix our {shot} token with the name of our cinematic. By using the {shot} token each of our shots will be exported as individual video files along with the full cinematic.
 
 	[OBJECT:EmbeddedVideo]
 	[PARAMLITERAL:width]
@@ -38,18 +38,15 @@ We are also using Adobe Premiere Pro CS6 as our external editing software; howev
 	[PARAMLITERAL:videoid]
 	c2J89J3b8WA
 	[/PARAMLITERAL]
-	[/OBJECT]
-
-	Above we show the sample cinematic that we want to export along with an EDL and click the **Render Movie** button to bring up the settings dialog window. 
-	Inside the settings window, we create a new folder to export the cinematic to and prefix our {shot} token with the name of our cinematic. 
-	By using the {shot} token each of our shots will be exported as individual video files along with the full cinematic. 
+	[/OBJECT] 
 
 	[REGION:note]
 	This will only work when the **Output Format** is set to **Video Sequence**. It does not work with individual frames. 
 	[/REGION]
 
 
-1.	Inside your software editing suite, **Import** the **EDL** file exported from the previous step. 
+1.	Inside your software editing suite, **Import** the **EDL** file exported from the previous step. In Premiere Pro, you can **Right-click** in the window in the lower-left and choose **Import** then point to the EDL file. 
+	We left the EDL Information and New Sequence dialog windows with the default information. 
 
 	[OBJECT:EmbeddedVideo]
 	[PARAMLITERAL:width]
@@ -63,10 +60,7 @@ We are also using Adobe Premiere Pro CS6 as our external editing software; howev
 	[/PARAMLITERAL]
 	[/OBJECT]
 
-	In Premiere Pro, you can **Right-click** in the window in the lower-left and choose **Import** then point to the EDL file. 
-	We left the EDL Information and New Sequence dialog windows with the default information. 
-
-1.	Next you'll need to link the data from the EDL to the exported media files in your editing software. 
+1.	Next you'll need to link the data from the EDL to the exported media files in your editing software. In Premiere Pro, **Right-click** the media folder and select **Link Media** then point each file to its corresponding media file. 
 
 	[OBJECT:EmbeddedVideo]
 	[PARAMLITERAL:width]
@@ -80,9 +74,7 @@ We are also using Adobe Premiere Pro CS6 as our external editing software; howev
 	[/PARAMLITERAL]
 	[/OBJECT]
 
-	In Premiere Pro, **Right-click** the media folder and select **Link Media** then point each file to its corresponding media file. 
-
-1.	Modify the shots in your editing software as desired, then **Export** as an **EDL**. 
+1.	Modify the shots in your editing software as desired, then **Export** as an **EDL**. For our sample scene, we adjust the length of the shots as well as re-arrange their position in the sequence. 
 
 	[OBJECT:EmbeddedVideo]
 	[PARAMLITERAL:width]
@@ -96,9 +88,7 @@ We are also using Adobe Premiere Pro CS6 as our external editing software; howev
 	[/PARAMLITERAL]
 	[/OBJECT]
 
-	Above for our sample scene, we adjust the length of the shots as well as re-arrange their position in the sequence. 
-
-1.	Return to **Unreal Engine 4** and in your Sequence, **Right-click** your **Shot Track** and **Import EDL**. 
+1.	Return to **Unreal Engine 4** and in your Sequence, **Right-click** your **Shot Track** and **Import EDL**. The changes we made externally are now reflected inside our Sequence in Unreal Engine 4. 
 
 	[OBJECT:EmbeddedVideo]
 	[PARAMLITERAL:width]
@@ -112,8 +102,6 @@ We are also using Adobe Premiere Pro CS6 as our external editing software; howev
 	[/PARAMLITERAL]
 	[/OBJECT]
 
-	The changes we made externally are now reflected inside our Sequence in Unreal Engine 4. 
-
 ## End Result
 
 The import/export EDL option within Sequencer gives you the flexibility to determine where you want to do your editing. 
