#UE4 Docs: minor text and formatting updates

diff --git a/Engine/Documentation/Source/Platforms/Linux/BeginnerLinuxDeveloper/SettingUpAnUnrealWorkflow/2/SettingUpAnUnrealWorkflow_2.INT.udn b/Engine/Documentation/Source/Platforms/Linux/BeginnerLinuxDeveloper/SettingUpAnUnrealWorkflow/2/SettingUpAnUnrealWorkflow_2.INT.udn
--- a/Engine/Documentation/Source/Platforms/Linux/BeginnerLinuxDeveloper/SettingUpAnUnrealWorkflow/2/SettingUpAnUnrealWorkflow_2.INT.udn
+++ b/Engine/Documentation/Source/Platforms/Linux/BeginnerLinuxDeveloper/SettingUpAnUnrealWorkflow/2/SettingUpAnUnrealWorkflow_2.INT.udn
@@ -35,12 +35,10 @@ Engine repository with your Git client. We'll cover both ways on this page.
 
 If don't want to clone the Unreal Engine repository with Git, you can use the **Download ZIP** button on our [GitHub site](http://github.com/EpicGames).
 
-1.	First, select the branch you want to download by clicking the **Branch:** button on the left side of the page.
+1.	First, select the branch you want to download by clicking the **Branch:** button on the left side of the page. For more info on branches in the Unreal Engine GitHub repository, please see the [](GettingStarted/DownloadingUnrealEngine/SourceBranches) page.
 
 	![](SelectRepoBranch.png)
 
-	[INCLUDE:GettingStarted/DownloadingUnrealEngine/SourceBranches]
-
 1.	After selecting the branch you want to download, click the **Clone or download** button located on the right side of the page.
 
 	![](CloneorDownloadButton.png)
