#UE4 Docs: minor text and formatting updates

diff --git a/Engine/Documentation/Source/Engine/MediaFramework/HowTo/StreamMediaSource/UsingStreamMedia.INT.udn b/Engine/Documentation/Source/Engine/MediaFramework/HowTo/StreamMediaSource/UsingStreamMedia.INT.udn
--- a/Engine/Documentation/Source/Engine/MediaFramework/HowTo/StreamMediaSource/UsingStreamMedia.INT.udn
+++ b/Engine/Documentation/Source/Engine/MediaFramework/HowTo/StreamMediaSource/UsingStreamMedia.INT.udn
@@ -26,14 +26,12 @@ For this how-to we are using the **Blueprint Third Person Template** project wit
 
 	![](StreamMedia_02.png)
 
-1.	Open **MediaStream** then enter your desired **Stream URL**. 
+1.	Open **MediaStream** then enter your desired **Stream URL**. If you do not have a file to link to, **Right-click** this [Sample Video](Infiltrator Demo.mp4) and copy the link address and paste it in the Stream URL field.
 
-	![](StreamMedia_03.png)
-
-	If you do not have a file to link to, **Right-click** this [Sample Video](Infiltrator Demo.mp4) and copy the link address and paste it in the Stream URL field. 
+	![](StreamMedia_03.png) 
 	
 	[REGION:note]
-	Above we are using a fake address, however your Stream URL must link directly to a supported format in order to play back the video. 
+	We are using a fake address, however your Stream URL must link directly to a supported format in order to play back the video. 
 	For example, PS4Media (PS4) only supports MP4 through HLS right now, while WmfMedia (Windows) can support a number of other stream sources. 
 	For futher information on supported formats per platform/player plug-in, see the [](Engine\MediaFramework\TechReference\) page. 
 	[/REGION]
@@ -42,40 +40,32 @@ For this how-to we are using the **Blueprint Third Person Template** project wit
 
 	![](StreamMedia_04.png)
 
-1.	In the **Create Media Player** window, check **Audio output SoundWave asset** and **Video output Media Texture asset**. 
+1.	In the **Create Media Player** window, check **Audio output SoundWave asset** and **Video output Media Texture asset**. This will automatically create a Sound Wave and Media Texture asset linked to this Media Player asset which is needed for playback.
 
 	![](StreamMedia_05.png)
+	
+1. Name the new Media PLayer Asset. We have called our Media Player asset **SampleStream** which is also applied to the Sound Wave and Media Texture. 
 
-	This will automatically create a Sound Wave and Media Texture asset linked to this Media Player asset which is needed for playback. 
-
-	![](StreamMedia_06.png)
-
-	We have called our Media Player asset **SampleStream** which is also applied to the Sound Wave and Media Texture. 
-
-1.	In the **Content Browser** under **StarterContent/Shapes** drag a **Shape_Plane** into the level. 
+	![](StreamMedia_06.png) 
 
-	![](StreamMedia_07.png)
+1.	In the **Content Browser** under **StarterContent/Shapes** drag a **Shape_Plane** into the level. We will use this Static Mesh to playback our video on.
 
-	We will use this Static Mesh to playback our video on. 
+	![](StreamMedia_07.png) 
 
 1.	Use the **Translate (W)**, **Rotation (E)** and **Scale (R)** tools to adjust the mesh to your desired size and location.
 
-1.	Hold **Ctrl** then select both the **Sound Wave** and **Media Texture** assets and drop them on top of the **Static Mesh** in the level. 
-
-	![](StreamMedia_08.png)
+1.	Hold **Ctrl** then select both the **Sound Wave** and **Media Texture** assets and drop them on top of the **Static Mesh** in the level. This will automatically create a **Material** and apply it to the Static Mesh.
 
-	This will automatically create a **Material** and apply it to the Static Mesh. 
+	![](StreamMedia_08.png) 
 
 1.	From the Toolbar, click the **Blueprints** button and **Open Level Blueprint**. 
 
 	![](StreamMedia_09.png)
 
-1.	Add a **Variable** of the **Media Player Reference** type called **Media Player** and set it to your **SampleStream** Media Player asset. 
+1.	Add a **Variable** of the **Media Player Reference** type called **Media Player** and set it to your **SampleStream** Media Player asset. You may need to **Compile** before you can assign the **Default Value**. 
 
 	![](StreamMedia_10.png)
 
-	You may need to **Compile** before you can assign the **Default Value**. 
-
 1.	Hold **Ctrl** and drag in the **MediaPlayer** variable and also **Right-click** and add an **Event Begin Play** node. 
 
 	![](StreamMedia_11.png)
