#UE4 Docs: Removing r.FinishCurrentFrame=1 as it is not a valid VR command.

diff --git a/Engine/Documentation/Source/Platforms/VR/ContentSetup/VRContent.INT.udn b/Engine/Documentation/Source/Platforms/VR/ContentSetup/VRContent.INT.udn
--- a/Engine/Documentation/Source/Platforms/VR/ContentSetup/VRContent.INT.udn
+++ b/Engine/Documentation/Source/Platforms/VR/ContentSetup/VRContent.INT.udn
@@ -68,7 +68,6 @@ These . INIsettings, if used, should be added to your projects **Config\ Default
 	r.screenpercentage=130
 	r.SeparateTranslucency=0
 	r.HZBOcclusion=0
-	r.FinishCurrentFrame=1
 	r.MotionBlurQuality=0
 	r.PostProcessAAQuality=3
 	r.BloomQuality=1
