#UE4 Docs: minor text and formatting updates

diff --git a/Engine/Documentation/Source/Engine/MediaFramework/HowTo/FileMediaSource/UsingFileMedia.INT.udn b/Engine/Documentation/Source/Engine/MediaFramework/HowTo/FileMediaSource/UsingFileMedia.INT.udn
--- a/Engine/Documentation/Source/Engine/MediaFramework/HowTo/FileMediaSource/UsingFileMedia.INT.udn
+++ b/Engine/Documentation/Source/Engine/MediaFramework/HowTo/FileMediaSource/UsingFileMedia.INT.udn
@@ -28,12 +28,10 @@ You will also need a [Supported Video File](Engine\MediaFramework\TechReference)
 
 	![](FileMedia_01b.png)
 
-1.	Drag your video file into the **Content/Movies** folder of your project. 
+1.	Drag your video file into the **Content/Movies** folder of your project. Placing your video files inside the **Content/Movies** folder of your project will ensure the videos are properly packaged. 
 
 	![](FileMedia_02.png)
 
-	Placing your video files inside the **Content/Movies** folder of your project will ensure the videos are properly packaged. 
-
 1.	Inside Unreal Engine and your project, **Right-click** in the **Movies** folder and under **Media** select **File Media Source**. 
 
 	![](FileMedia_03.png)
@@ -46,48 +44,38 @@ You will also need a [Supported Video File](Engine\MediaFramework\TechReference)
 
 	![](FileMedia_05.png)
 
-1.	In the **Create Media Player** window, check **Audio output SoundWave asset** and **Video output Media Texture asset**. 
+1.	In the **Create Media Player** window, check **Audio output SoundWave asset** and **Video output Media Texture asset**. This will automatically create a Sound Wave and Media Texture asset linked to this Media Player asset which is needed for playback. 
 
 	![](FileMedia_06.png)
 
-	This will automatically create a Sound Wave and Media Texture asset linked to this Media Player asset which is needed for playback. 
+1.  Name the new Media Player asset. We have called our Media Player asset **SampleMedia** which is also applied to the Sound Wave and Media Texture. 
 
 	![](FileMedia_07.png)
 
-	We have called our Media Player asset **SampleMedia** which is also applied to the Sound Wave and Media Texture. 
-
-	If you open the Media Player asset, you can **Double-click** on your Media Source asset and it will playback. 
+1.  Open the Media Player asset and **Double-click** on your Media Source asset and it will playback. In the lower-right in the **Details** panel the **Output** section will also have auto-assigned the Sound Wave and Media Texture. 
 
 	![](FileMedia_08.png)
 
-	In the lower-right in the **Details** panel the **Output** section will also have auto-assigned the Sound Wave and Media Texture. 
-
-1.	In the **Content Browser** under **StarterContent/Shapes** drag a **Shape_Plane** into the level. 
+1.	In the **Content Browser** under **StarterContent/Shapes** drag a **Shape_Plane** into the level. We will use this Static Mesh to playback our video on. 
 
 	![](FileMedia_09.png)
 
-	We will use this Static Mesh to playback our video on. 
-
 1.	Use the **Translate (W)**, **Rotation (E)** and **Scale (R)** tools to adjust the mesh similar to below. 
 
 	![](FileMedia_10.png)
 
-1.	Hold **Ctrl** then select both the **Sound Wave** and **Media Texture** assets and drop them on top of the **Static Mesh** in the level. 
+1.	Hold **Ctrl** then select both the **Sound Wave** and **Media Texture** assets and drop them on top of the **Static Mesh** in the level. This will automatically create a **Material** and apply it to the Static Mesh. 
 
 	![](FileMedia_11.png)
 
-	This will automatically create a **Material** and apply it to the Static Mesh. 
-
 1.	From the Toolbar, click the **Blueprints** button and **Open Level Blueprint**. 
 
 	![](FileMedia_12.png)
 
-1.	Add a **Variable** of the **Media Player Reference** type called **Media Player** and set it to your **SampleMedia** Media Player asset. 
+1.	Add a **Variable** of the **Media Player Reference** type called **Media Player** and set it to your **SampleMedia** Media Player asset. You may need to **Compile** before you can assign the **Default Value**. 
 
 	![](FileMedia_13.png)
 
-	You may need to **Compile** before you can assign the **Default Value**. 
-
 1.	Hold **Ctrl** and drag in the **MediaPlayer** variable and also **Right-click** and add an **Event Begin Play** node. 
 
 	![](FileMedia_14.png)
