Incorporating peer review feedback - Media Framework Update Guide UEDOC-3213

diff --git a/Engine/Documentation/Source/Engine/MediaFramework/TechReference/Upgrading/Upgrading.INT.udn b/Engine/Documentation/Source/Engine/MediaFramework/TechReference/Upgrading/Upgrading.INT.udn
--- a/Engine/Documentation/Source/Engine/MediaFramework/TechReference/Upgrading/Upgrading.INT.udn
+++ b/Engine/Documentation/Source/Engine/MediaFramework/TechReference/Upgrading/Upgrading.INT.udn
@@ -15,19 +15,37 @@ Your project will still load, but playback will no longer work with existing 4.1
 
 As you transition your project, keep the following items in mind:
 
-* **Media Player** assets no longer store a URL for the media. Instead there are now new asset types called 'media sources' that describe the media. For more information on these new asset types, see [](Engine/MediaFramework/Overview).
+* **Media Player** assets no longer store a URL for the media. Instead there are now new asset types called 'media sources' that describe the media. 
 
 	![](MediaTypes.png)
+	
+	[REGION:tip]
+	For more information on these new asset types, see the [](Engine/MediaFramework/Overview).
+	[/REGION]
 
-* **File Media Source** assets are used for media files stored on your device or in local network file shares. For an example of how to use a **File Media Source** asset, see [](Engine/MediaFramework/HowTo/FileMediaSource).
+* **File Media Source** assets are used for media files stored on your device or in local network file shares. 
 
 	![](FileMediaSource.png)
-* **Stream Media Source** assets are used for streaming media over the internet. For an example of how to use a **Stream Media Source** asset, see [](Engine/MediaFramework/HowTo/StreamMediaSource).
+	
+	[REGION:tip]
+	For an example of how to use a **File Media Source** asset, try [](Engine/MediaFramework/HowTo/FileMediaSource).
+	[/REGION]
+	
+* **Stream Media Source** assets are used for streaming media over the internet. 
 
 	![](StreamMediaSource.png)
-* **Media Playlist** assets are used to combine multiple media sources and play them one after another. Playlists will automatically cycle back to the beginning. For an example of working with playlists, see [](Engine/MediaFramework/HowTo/MediaPlaylists).
+	
+	[REGION:tip]
+	For an example of how to use a **Stream Media Source** asset, try [](Engine/MediaFramework/HowTo/StreamMediaSource).
+	[/REGION]
+	
+* **Media Playlist** assets are used to combine multiple media sources and play them one after another. Playlists will automatically cycle back to the beginning. 
 
 	![](MediaPlaylist.png)
+	
+	[REGION:tip]
+	For an example of working with playlists, try [](Engine/MediaFramework/HowTo/MediaPlaylists).
+	[/REGION]
 
 * **Media Texture** and **Media Sound Wave** assets no longer hold a reference to the **Media Player**. Instead, the Media Player now references the texture and the sound assets. Additionally, there is no longer a many-to-one relationship between **Media Texture** / **Media Sound Wave** and 
 **Media Player** assets. Every **Media Player** now references at most one of each.
@@ -41,4 +59,6 @@ If you do not set this option, you must initiate playback via Blueprints or C++.
 
 * The **Loop** option will loop the currently opened media source. Please note that, while looping is enabled, playlists do not advance.
 
-	![](Loop.png)
\ No newline at end of file
+	![](Loop.png)
+
+* For specifics on the content types and URLs supported by the updated Media Framework system, see the [](Engine/MediaFramework/TechReference).
\ No newline at end of file
