Edits to the document not related to the listed JIRA that were made during a peer review.

diff --git a/Engine/Documentation/Source/Engine/Editor/VR/Controls/VREditorControls.INT.udn b/Engine/Documentation/Source/Engine/Editor/VR/Controls/VREditorControls.INT.udn
--- a/Engine/Documentation/Source/Engine/Editor/VR/Controls/VREditorControls.INT.udn
+++ b/Engine/Documentation/Source/Engine/Editor/VR/Controls/VREditorControls.INT.udn
@@ -43,7 +43,7 @@ skilllevel: intermediate
 [/OBJECT]
 [/VAR]
 
-There are various controls to enable you to navigate the scene, select and manipulate Actors, and workw ith editor panels and windows.
+There are various controls for navigating the scene, selecting and manipulating Actors, and working with editor panels and windows.
 
 ## The Controller
 
