#UE4 Docs: minor formatting updates

diff --git a/Engine/Documentation/Source/Platforms/Oculus/Oculus.INT.udn b/Engine/Documentation/Source/Platforms/Oculus/Oculus.INT.udn
--- a/Engine/Documentation/Source/Platforms/Oculus/Oculus.INT.udn
+++ b/Engine/Documentation/Source/Platforms/Oculus/Oculus.INT.udn
@@ -17,8 +17,11 @@ tags:Oculus Rift
 
 ## Oculus Rift Topics
 
-[REGION:buttonlist]
 [DIR(output:"listbutton" type:"overview" parent:"Platforms/Oculus")]
+
+## Oculus Rift Resources
+
+[REGION:buttonlist]
 * [Lessons from Integrating the Oculus Rift into Unreal Engine 4](https://de45xmedrsdbp.cloudfront.net/Resources/files/UE4-Integration-and-Demos_OC-100270768.pptx "Oculus Connect 2014 Presentation")
 [/REGION]
 
@@ -45,7 +48,6 @@ While you're developing your UE4 VR project, it's best to enable **Unknown Sourc
 	![](Rift_UKS_04.png)
 
 
-
 ## Oculus Talks
  **Oculus Connect: Integrating the Oculus Rift into Unreal Engine 4**
 [OBJECT:EmbeddedVideo]
