#UE4 Docs: Applying peer feedback.

diff --git a/Engine/Documentation/Source/Enterprise/Automotive/AutomotiveMaterialPack/AutomotiveMaterialPack.INT.udn b/Engine/Documentation/Source/Enterprise/Automotive/AutomotiveMaterialPack/AutomotiveMaterialPack.INT.udn
--- a/Engine/Documentation/Source/Enterprise/Automotive/AutomotiveMaterialPack/AutomotiveMaterialPack.INT.udn
+++ b/Engine/Documentation/Source/Enterprise/Automotive/AutomotiveMaterialPack/AutomotiveMaterialPack.INT.udn
@@ -7,6 +7,7 @@ Parent:Enterprise/Automotive
 tags:Enterprise
 type:overview
 order:1
+Related:Engine/Rendering/Materials/PhysicallyBased/
 Related:Engine/Rendering/Materials/MaterialProperties/LightingModels
 Related:Engine/Rendering/Materials/HowTo/ClearCoatDualNormal
 Related:Engine/Content/Types/Textures/Cubemaps
@@ -16,7 +17,24 @@ Related:Engine/UI/LevelEditor/Modes/MeshPaintMode/VertexColor
 
 ![](AP_Intro.png)
 
-The Automotive Material Pack is a collection of automotive-themed PBR Materials and Textures that have been setup for use in Unreal Engine 4(UE4). The pack features 21 Master Materials and 75 Material Instances taken from Epic's 2016 GDC Mclaren 570S demo. In the following document we will take a look at what Materials the Automotive Material Pack contains and any special information you might need to know in order to use these Materials in your UE4 projects.
+The Automotive Material Pack is a collection of automotive-themed [PBR Materials](Engine/Rendering/Materials/PhysicallyBased/) and Textures that have been setup for use in Unreal Engine 4(UE4). The pack features 21 Master Materials and 75 Material Instances taken from Epic's 2016 GDC Mclaren 570S demo. In the following document, we will take a look at what Materials the Automotive Material Pack contains and any special information you might need to know in order to use these Materials in your UE4 projects.
+
+[OBJECT:EmbeddedVideo]
+	[PARAMLITERAL:width]
+	640
+	[/PARAMLITERAL]
+	[PARAMLITERAL:height]
+	360
+	[/PARAMLITERAL]
+	[PARAMLITERAL:videoid]
+	l4FiJ1A5veY
+	[/PARAMLITERAL]
+[/OBJECT]
+
+
+[region:note]
+For a more in-depth explanation of what Physically Based Rendering is and how to use it in UE4 make sure to check out the [Physically Based Materials](Engine/Rendering/Materials/PhysicallyBased/) documentation pages for more information.
+[/region]
 
 ## What's in the Automotive Material Pack
 
@@ -45,7 +63,8 @@ If you are unsure what a certain parameter in the Material does, hovering your m
 The vehicle meshes used in the documentation below **are not** included with the asset pack.
 [/region]
 
-## Standard Material Parameters
+
+## General Material Parameters
 
 In the following section, we will take a look at some of the parameters that all of the Materials share, what those parameters do and how you can use them.
 
@@ -66,10 +85,10 @@ In the following section, we will take a look at some of the parameters that all
 
 	[OBJECT:ComparisonSlider]
 		[PARAM:before]
-		![World Aligned OFF and no UVs](AP_02_SMP_03_WA_OFF.png)
+		![World Aligned OFF and no UV's](AP_02_SMP_03_WA_OFF.png)
 		[/PARAM]
 		[PARAM:after]
-		![World Aligned ON and no UVs](AP_02_SMP_04_WA_ON.png)
+		![World Aligned ON and no UV's](AP_02_SMP_04_WA_ON.png)
 		[/PARAM]
 	[/OBJECT]
 
@@ -88,7 +107,7 @@ In the following section, we will take a look at some of the parameters that all
 	[/OBJECT]
 	[/region]
 
-## Car Paint Material
+### Car Paint
 
 The Car Paint Material is a very robust Material that has 29 unique parameters that can be adjusted. When trying to modify this Material for your project, we recommend you starting adjusting and experimenting with the following parameters:
 
@@ -114,20 +133,15 @@ The following image sequence shows off the variety of paint colors and surface l
 [INCLUDE:Enterprise\Automotive\AutomotiveMaterialPack\#CP]
 [/REGION]
 
-## Orange Peel Material Parameter
+### Orange Peel Material Parameter
 
-The Clearcoat lacquer that is applied on top of automotive paint often has a slightly bumpy surface similar to the peel of an Orange. Because of this resemblance this layer of lacquer is referred to as Orange Peel. You can control the intensity of the Orange Peel using the **Orange Peel Strength** and **Orange Peel Scale** parameters. To disable the Orange Peel effect, set the **Orange Peel Strength** to **0**.
+The Clearcoat lacquer that is applied on top of automotive paint often has a slightly bumpy surface similar to the peel of an Orange. Because of this resemblance, this layer of lacquer is referred to as Orange Peel. You can control the intensity of the Orange Peel using the **Orange Peel Strength** and **Orange Peel Scale** parameters. To disable the Orange Peel effect, set the **Orange Peel Strength** to **0**.
 
-[OBJECT:ComparisonSlider]
-	[PARAM:before]
-	![Orange Peel Off](AP_03_Carpaint_03_OrangPeel_OFF.png)
-	[/PARAM]
-	[PARAM:after]
-	![Orange Peel On](AP_03_Carpaint_04_OrangPeel_ON.png)
-	[/PARAM]
-[/OBJECT]
+[REGION:fullwidth]
+[INCLUDE:Enterprise\Automotive\AutomotiveMaterialPack\#OP]
+[/REGION]
 
-## Carbon Fiber Material
+### Carbon Fiber
 
 The Carbon Fiber Material makes use of the new Dual Normal option that is available when using a  Clearcoat Material. The Dual Normal option allows you to add a 2nd Normal map to the Clearcoat Material so that both the bottom and top layers of the Material will react to changes to lighting. Without this feature, the glossy effect of the Carbon Fiber Material is lost, making the Carbon Fiber look not as realistic as it could be. 
 
@@ -144,7 +158,9 @@ The Carbon Fiber Material makes use of the new Dual Normal option that is availa
 For more information on using Dual Normals with a Clearcoat Material make sure to check out the [Dual Normals](Engine/Rendering/Materials/HowTo/ClearCoatDualNormal) Material page.
 [/region]
 
-## Caliper Material
+## Wheels
+
+### Caliper
 
 The caliper is part of a car's braking system and is usually placed on the brake disc in between the wheel rim and the brake pad. 
 
@@ -164,7 +180,7 @@ When using the brake caliper Material you can adjust the following two parameter
 	[/PARAM]
 [/OBJECT]
 
-## Rim Material
+### Rim
 
 There are five different Materials that can be used for the rims of the wheels including a Two Tone option. When using the Two Tone Rim keep in mind that the Two Tone Rim is made up of two different Materials with one Material being blue-black and the other material is silver. The **Red Channel** of the Vertex Color on the Rim Static Mesh is used to define which areas of the wheel should be silver. Any areas that are not using red Vertex color should use black as the Vertex Color so that they will display the 2nd Material.
 
@@ -178,9 +194,9 @@ If you are unfamiliar with how to use Vertex Paint in UE4 please check out the [
 [INCLUDE:Enterprise\Automotive\AutomotiveMaterialPack\#RM]
 [/REGION]
 
-## Brake Rotor Material
+### Brake Rotor
 
-Since the brake rotor Material uses a radial Texture pattern the UV's of the brake rotor need to use planar UV project mapping to display the brake rotor Material correctly. Failing to setup the meshes UVs in this manner will result in the brake rotor Material not being displayed correctly when applied to the brake rotor mesh.  The image below shows an example of how a Texture should look when placed on an object that has planar projection mapping.
+Since the brake rotor, Material uses a radial Texture pattern the UV's of the brake rotor need to use planar UV project mapping to display the brake rotor Material correctly. Failing to set up the brake rotor meshes UV's in this manner will result in the brake rotor Material not being displayed correctly when applied to the brake rotor mesh.  The image below shows an example of how a Texture should look when placed on an object that has planar projection mapping.
 
 ![](AP_07_BrakeRotor_01.png)
 
@@ -195,21 +211,7 @@ The following images show examples of the different types of brake rotor Materia
 	[/PARAM]
 [/OBJECT]
 
-## Reflector Material
-
-The reflector Material is used for the plastic  reflectors you find next to turn singles and on lights.  When looking at the reflector Material you will find that it comes with the following three reflector pattern presets:
-
-* **Hexagon**
-* **Honeycomb**
-* **Triangle**
-
-	![](AP_08_Reflector_02.png)
-
-To change which reflector pattern is used,  open up the **Reflector** Material Instance and change the **Reflector(Bottom Layer) Normal** to the reflector pattern you want to use.
-
-![](AP_08_Reflector_01.png)
-
-## Rubber / Tire Material
+### Rubber / Tire 
 
 The rubber Material that comes with the Automotive Material Pack can be used for any general rubber surface you might find in a car. For example, the following surfaces would make good use of the rubber Material.
 
@@ -222,23 +224,23 @@ When looking at the rubber Material Instance you will notice that it has a tire-
 
 [OBJECT:ComparisonSlider]
 	[PARAM:before]
-	![Wrong Way](AP_09_Rubber_04_Wrong.png)
+	![Misaligned UV Layout](AP_09_Rubber_04_Wrong.png)
 	[/PARAM]
 	[PARAM:after]
-	![Correct Way](AP_09_Rubber_04_Right.png)
+	![Aligned UV Layout](AP_09_Rubber_04_Right.png)
 	[/PARAM]
 [/OBJECT]
 
 [region:note]
-When using the sidewall option for a tire you will need to make sure that the UV for the tire sidewall mesh line up with the Normal Map like in the image marked Correct UV Layout.
+When using the sidewall option for a tire you will need to make sure that the UV's for the tire sidewall mesh lines up with the Normal Map like in the image marked Correct UV Layout.
 [/region]
 
 [OBJECT:ComparisonSlider]
 	[PARAM:before]
-	![Wrong UV Layout](AP_09_Rubber_06_UV_Wrong.png)
+	![Misaligned UV Layout](AP_09_Rubber_06_UV_Wrong.png)
 	[/PARAM]
 	[PARAM:after]
-	![Correct UV Layout](AP_09_Rubber_05_UV_LineUp.png)
+	![Aligned UV Layout](AP_09_Rubber_05_UV_LineUp.png)
 	[/PARAM]
 [/OBJECT]
 	
@@ -250,7 +252,9 @@ When using the sidewall option for a tire you will need to make sure that the UV
 Click for full image.
 [/REGION]
 
-## Leather / Perforated Material
+## Interior 
+
+### Leather / Perforated 
 
 The leather Material that comes with the Automotive Material Pack has 12 different preset leather colors to choose from and also contains five different perforated leather presets. 
 
@@ -263,15 +267,15 @@ The leather Material that comes with the Automotive Material Pack has 12 differe
 	[/PARAM]
 [/OBJECT]
 
-## Alcantara Material
+### Alcantara 
 
-The alcantara Material that comes with the Automotive Material Pack is used to display the fabrics you find in car's. When using the alcantara Material you can adjust the following parameters to manipulate the look of the fabric.
+The Alcantara Material that comes with the Automotive Material Pack is used to display the fabrics you find in cars. When using the Alcantara Material you can adjust the following parameters to manipulate the look of the fabric.
 
 * **Fabric Scale / Strength**
 * **Cloth Folds Scale / Strength**
 * **Fabric Detail Scale / Strength**
 
-In the following set of images, you can see the effect that the **Cloth Folds** parameter has on the look of the alcantara fabric Material.
+In the following set of images, you can see the effect that the **Cloth Folds** parameter has on the look of the Alcantara fabric Material.
 
 [OBJECT:ComparisonSlider]
 	[PARAM:before]
@@ -282,15 +286,15 @@ In the following set of images, you can see the effect that the **Cloth Folds**
 	[/PARAM]
 [/OBJECT]
 
-## Glass Material
+## Glass 
 
-The glass Material that comes with the Automotive Material Pack has the ability to display two completely different reflections in one Material. With this Material you can select what the exterior glass should reflect as well as what the interior glass should reflect. 
+The glass Material that comes with the Automotive Material Pack has the ability to display two completely different reflections in one Material. With this Material, you can select what the exterior glass should reflect as well as what the interior glass should reflect. 
 
 * **Outer Reflection:** Helps to reflect the scene around the exterior of the glass. 
 
 	![](AP_12_Glass_06_Ex.png)
 
-* **Inner Reflection:** Helps to reflect the scene around the interor of the glass.
+* **Inner Reflection:** Helps to reflect the scene around the interior of the glass.
 
 	![](AP_12_Glass_07_Ex.png)
 
@@ -305,7 +309,7 @@ Here is an example of what the glass Material looks like with the inner reflecti
 	[/PARAM]
 [/OBJECT]
 
-When using the inner reflections for a car windshield the best workflow is to capture a Cubemap from the inside of the car. Then in the glass Material Instance, enable the **Use Glass Inner Reflections Cubemap** option and then input the Cubemap you want to use into the **Cubemap** input.
+When using the inner reflections for a car windshield the best workflow is to capture a Cubemap from the inside of the car. Then, in the glass Material Instance, enable the **Use Glass Inner Reflections Cubemap** option and then input the Cubemap you want to use into the **Cubemap** input.
 
 ![](AP_12_Glass_01.png)
 
@@ -314,6 +318,22 @@ If you are unfamiliar with how to capture and use Cubemap in UE4 please check ou
 [/region]
 
 
+## Reflectors 
+
+The reflector Material is used for the plastic  reflectors you find next to turn signals as well as on and around lights.  When looking at the reflector Material you will find that it comes with the following three reflector pattern presets:
+
+* **Hexagon**
+* **Honeycomb**
+* **Triangle**
+
+	[REGION:fullwidth]
+	[INCLUDE:Enterprise\Automotive\AutomotiveMaterialPack\#RF]
+	[/REGION]
+
+To change which reflector pattern is used,  open up the **Reflector** Material Instance and change the **Reflector(Bottom Layer) Normal** to the reflector pattern you want to use.
+
+![](AP_08_Reflector_01.png)
+
 <!-----
 [EXCERPT:RM]
 [OBJECT:SliderSequence]
@@ -373,4 +393,54 @@ If you are unfamiliar with how to capture and use Cubemap in UE4 please check ou
 	[/PARAM]
 [/OBJECT]
 [/EXCERPT:CP]
+-->
+
+<!-----
+[EXCERPT:RF]
+[OBJECT:SliderSequence]
+	[PARAMLITERAL:max]
+	3
+	[/PARAMLITERAL]
+	[PARAMLITERAL:ext]
+	jpg
+	[/PARAMLITERAL]
+	[PARAMLITERAL:id]
+	RF
+	[/PARAMLITERAL]
+	[PARAM:caption]
+	**Drag the slider to see the Hexagon, Honeycomb and Triangle reflector patterns.**
+	[/PARAM]
+	[PARAM:image_list]
+	![](RF_1.png)
+	![](RF_2.png)
+	![](RF_3.png)
+	[/PARAM]
+[/OBJECT]
+[/EXCERPT:RF]
+-->
+
+<!-----
+[EXCERPT:OP]
+[OBJECT:SliderSequence]
+	[PARAMLITERAL:max]
+	5
+	[/PARAMLITERAL]
+	[PARAMLITERAL:ext]
+	jpg
+	[/PARAMLITERAL]
+	[PARAMLITERAL:id]
+	OP
+	[/PARAMLITERAL]
+	[PARAM:caption]
+	**Drag the slider to see what happens to the Orange Peel when its intensity is gradually increased from 0.0 to 0.8.**
+	[/PARAM]
+	[PARAM:image_list]
+	![](OP_1.png)
+	![](OP_2.png)
+	![](OP_3.png)
+	![](OP_4.png)
+	![](OP_5.png)
+	[/PARAM]
+[/OBJECT]
+[/EXCERPT:OP]
 -->
\ No newline at end of file
