Accepting peer review changes on UEDOC-3209

diff --git a/Engine/Documentation/Source/Engine/Foliage/Foliage.INT.udn b/Engine/Documentation/Source/Engine/Foliage/Foliage.INT.udn
--- a/Engine/Documentation/Source/Engine/Foliage/Foliage.INT.udn
+++ b/Engine/Documentation/Source/Engine/Foliage/Foliage.INT.udn
@@ -55,7 +55,7 @@ order:9
 
 [EXCERPT:Intro]
 
-The Foliage Instanced Mesh system allows you to quickly paint or erase sets of Static Meshes on Landscape Actors, other Static Mesh Actors, and BSP. 
+The Foliage Instanced Mesh system allows you to quickly paint or erase sets of Static Meshes on Landscape Actors, other Static Mesh Actors, and BSP geometry. 
 These meshes are automatically grouped together into batches that are rendered using hardware instancing, meaning that many instances can be rendered with only a single draw call.
 
 [/EXCERPT:Intro]
@@ -85,10 +85,10 @@ To add a Static Mesh to the Mesh List, you must drag the Static Mesh from the **
 [/REGION]
 
 [REGION:tip]
-The yellow Foliage Type will only be displayed when you are trying to add the correct actor type. If you do not see this then what you are trying to add will not work.
+The yellow Foliage Type will only be displayed when you are trying to add the correct actor type. If you do not see this, then what you are trying to add will not work.
 [/REGION]
 
-To adjust the properties of a Static Mesh that is used as foliage, first, select the Static Mesh in the Mesh List and then press the small **Magnifying Glass** icon in the lower right-hand corner of the foliage tool. This will expose the selected Static Meshes details. 
+To adjust the properties of a Static Mesh that is used as foliage, select the Static Mesh in the Mesh List and then press the small **Magnifying Glass** icon in the lower right-hand corner of the foliage tool. This will expose the selected Static Meshes details. 
 
 [REGION:lightbox]
 [![](Expose_Instanc_Options.png)(w:190 h:350)](Expose_Instanc_Options.png)
@@ -108,7 +108,7 @@ Foliage Meshes in the Mesh list can be enabled or disabled by hovering the mouse
 ![](Foliage_Selecting_Meshes.png)
 
 You can also select Foliage Meshes by clicking on them in the Mesh list. 
-This will put an orange highlight around them and then in the options panel below you can adjust various properties. 
+This will put an orange highlight around them and then, in the options panel below, you can adjust various properties. 
 
 [REGION:lightbox]
 [![](Foliage_Selected_Meshes.png)(w:190 h:350)](Foliage_Selected_Meshes.png)
@@ -127,14 +127,14 @@ In the Meshes list (pictured above), the Mesh Edit Tool displays properties rela
 | **Paint Density** |  [INCLUDE:Engine/Foliage#PaintDensity] |
 | **Erase Density** |  [INCLUDE:Engine/Foliage#EraseDensity] |
 
-The checkmark boxes you see with the names **Landscape**, **Static Meshes**, **BSP** and **Translucent** determine which type of Actor type the Foliage meshes will be placed on when the tool is used. 
+The checkmark boxes you see with the names **Landscape**, **Static Meshes**, **BSP** and **Translucent** determine which Actor type the Foliage meshes will be placed on when the tool is used. 
 Checking a box will indicate that a Foliage mesh should be added should you use the tool on that Actor type.
 Unchecking a box will indicate that a foliage mesh should not be added to that Actor type.
 
 <!---
 [EXCERPT:PaintDensity]
 The target density at which you will add instances when using **Left Mouse** button.
-* This is a value from 0 to 1 where 1 paint mesh instances at the maximum density listed in the Mesh Properties for each mesh (see below).
+* This is a value from 0 to 1 where 1 paints mesh instances at the maximum density listed in the Mesh Properties for each mesh (see below).
 * If the density of meshes in the brush is already greater than this, no meshes will be added.
 [/EXCERPT:PaintDensity]
 --->
@@ -142,7 +142,7 @@ The target density at which you will add instances when using **Left Mouse** but
 <!---
 [EXCERPT:EraseDensity]
 The target density you would like to achieve when you erase instances using **Ctrl+Shift+click**.
-* This is a value from 0 to 1 where 0 represents no meshes. If there are fewer meshes than this desired erase density, no more meshes will be erased. 
+* This is a value from 0 to 1 where 0 represents no meshes and 1 represents the maximum density listed in the Mesh Properties for each mesh (see below). If there are fewer meshes than this desired erase density, no more meshes will be erased. 
 * This allows you to thin out meshes already painted without completely removing all instances.
 [/EXCERPT:EraseDensity]
 --->
@@ -150,7 +150,7 @@ The target density you would like to achieve when you erase instances using **Ct
 <!---
 [EXCERPT:Filter]
 The kinds of objects you want to place instances on. 
-* Note the foliage sphere brush will disappear when it is not over an object that matches the filter type. 
+* Note that the foliage sphere brush will disappear when it is not over an object that matches the filter type. 
 * Also note that you can only paint on objects that are in the active streaming level.
 [/EXCERPT:Filter]
 --->
@@ -173,7 +173,7 @@ There are five tools available by selecting the appropriate button in the Toolba
 ### Paint Tool
 ![Paint Tool](ToolPaint.png)
 
-The Paint tool is Foliage Edit Mode's default tool. You can use it to paint instances of the Static Meshes in the Meshes list in your level. When Foliage Edit Mode is active, a transparent sphere brush is drawn in the level indicating where the foliage brush will operate. In the image below you can see an example of what the Foliage paint brush looks like.
+The Paint tool is Foliage Edit Mode's default tool. You can use it to paint instances of the Static Meshes in the Meshes list in your level. When Foliage Edit Mode is active, a transparent sphere brush is drawn in the level indicating where the foliage brush will operate. In the image below, you can see an example of what the Foliage paint brush looks like.
 
 ![](Foliage_Paint_Tool_Gimzo.png)
 
@@ -200,7 +200,7 @@ The properties are accessible by selecting the Static Mesh or Static Meshes you
 [![](Foliage_PaintSettings.png)(w:190 h:350)](Foliage_PaintSettings.png)
 [/REGION]
 
-Here is a high-level description of what each section is used for.
+Here is a high-level description of what each section is used for:
 
 | **Property** | **Description** |
 | --- | --- |
