Fixing capitalization error for VR Editor Controls. UEDOC-3294

diff --git a/Engine/Documentation/Source/Engine/Editor/VR/Controls/VREditorControls.INT.udn b/Engine/Documentation/Source/Engine/Editor/VR/Controls/VREditorControls.INT.udn
--- a/Engine/Documentation/Source/Engine/Editor/VR/Controls/VREditorControls.INT.udn
+++ b/Engine/Documentation/Source/Engine/Editor/VR/Controls/VREditorControls.INT.udn
@@ -83,7 +83,7 @@ Each Controller emits a beam called the laser pointer which is used to determine
 | Control |  Description |
 | ------- | ----------- |
 | ![Aim Controller at Actor + Squeeze Trigger](select_actor.png)(convert:false) | Aim the controller at an Actor in the world and squeeze the Trigger to select the Actor. |
-| [INCLUDE:#multiselect] | Aim one controller at an Actor in the world and squeeze the Trigger halfway to select the Actor and enter multi-select mode. While holding the Trigger halfway, Aim the other controller at other Actors and squeeze its Trigger to add or remove them from the selection. |
+| [INCLUDE:#multiselect] | Aim one controller at an Actor in the world and squeeze the Trigger halfway to select the Actor and enter multi-select mode. While holding the Trigger halfway, aim the other controller at other Actors and squeeze its Trigger to add or remove them from the selection. |
 [/REGION]
 
 ## Transform Controls
