#UE4 Docs: Adding an image that was missing.

diff --git a/Engine/Documentation/Source/Enterprise/Automotive/AutomotiveMaterialPack/AutomotiveMaterialPack.INT.udn b/Engine/Documentation/Source/Enterprise/Automotive/AutomotiveMaterialPack/AutomotiveMaterialPack.INT.udn
--- a/Engine/Documentation/Source/Enterprise/Automotive/AutomotiveMaterialPack/AutomotiveMaterialPack.INT.udn
+++ b/Engine/Documentation/Source/Enterprise/Automotive/AutomotiveMaterialPack/AutomotiveMaterialPack.INT.udn
@@ -57,6 +57,7 @@ Inside of the Automotive Material Pack you will find the following Materials typ
 
 [region:tip]
 If you are unsure what a certain parameter in the Material does, hovering your mouse over the parameter will display additional information.
+![](AP-Param-Description.png)
 [/region]
 
 [region:note]
