Initial pass at Sequencer Recorder page.

diff --git a/Engine/Documentation/Source/Engine/Sequencer/HowTo/SequenceRecorder/SequenceRecorder.INT.udn b/Engine/Documentation/Source/Engine/Sequencer/HowTo/SequenceRecorder/SequenceRecorder.INT.udn
--- a/Engine/Documentation/Source/Engine/Sequencer/HowTo/SequenceRecorder/SequenceRecorder.INT.udn
+++ b/Engine/Documentation/Source/Engine/Sequencer/HowTo/SequenceRecorder/SequenceRecorder.INT.udn
@@ -1,7 +1,7 @@
 Availability: Public
-Title:Recording Gameplay
+Title:Recording to Sequencer
 Crumbs: %ROOT%, Engine, Engine/Sequencer, Engine/Sequencer/HowTo
-Description:Record your gameplay actions with the Sequence Recorder and use that data inside Sequencer.
+Description:Record your actions with the Sequence Recorder and use that data inside Sequencer.
 Version: 4.12
 SkillLevel: Advanced
 parent:Engine/Sequencer/HowTo
@@ -10,7 +10,7 @@ checkpoint: editorqs
 Related: Engine/Sequencer/Overview
 Related: Engine\Sequencer\ReferenceEditor\
 
-The **Sequence Recorder** allows you to capture specified Actors during gameplay that are then saved as a new **Level Sequence** and can be edited inside **Sequencer**. 
+The **Sequence Recorder** allows you to capture specified Actors during level editing or gameplay that can be saved as a new **Level Sequence** and can be edited inside **Sequencer**. 
 This is useful for quickly capturing content for scenes as you can take a playable character, perform some actions during gameplay while recording with the Sequence Recorder then take that data into Sequencer and create a cinematic around it. 
 
 In this example, we will use the Sequence Recorder to record our playable character's movement which we can then edit. 
@@ -22,9 +22,15 @@ In this example, we will use the Sequence Recorder to record our playable charac
 For this how-to, we are using the **Blueprint Third Person Template** project with **Starter Content Enabled**. 
 [/REGION]
 
-1.	From the **Menu Bar** under **Window**, select **Sequence Recorder**. The Sequence Recorder window will automatically open. There are some options under **Sequence Recording** which will determine how (and where) the new Level Sequence asset will be saved. You can choose to record Actors that are spawned (such as particle effects, other characters, etc.), determine the naming convention and save location, and change other settings like the duration to record or delay before starting to record. 
+1.	From the **Menu Bar** under **Window**, select **Sequence Recorder**. 
 
 	![](RecorderStep1.png)
+
+	The Sequence Recorder window will automatically open. There are some options under **Sequence Recording** which will determine how (and where) the new Level Sequence asset will be saved. You can choose to record Actors that are spawned (such as particle effects, other characters, etc.), determine the naming convention and save location, and change other settings like the duration to record or delay before starting to record. 
+
+	[REGION:note]
+	See the [](Engine\Sequencer\Features\SequenceRecorder) page for more information about each option.
+	[/REGION]
 	
 1.	From the Main Editor Toolbar, click the **Play** button to start a Play in Editor session. 
 
@@ -67,7 +73,7 @@ For this how-to, we are using the **Blueprint Third Person Template** project wi
 
 ## End Result
 
-You can open the **RecordedSequence** asset and begin editing it as you would a normal Level Sequence. Below is our recorded sequence to which we could add cameras and a [**Camera Cuts**](Engine/Sequencer/HowTo/TracksCameraCut) track to provide multiple angles, music, effects or anything else we wanted. We could even take this sequence and embed it in another sequences as part of a [**Shots Track**](Engine/Sequencer/HowTo/TracksShot). 
+You can open the **RecordedSequence** asset and begin editing it as you would a normal Level Sequence. Below is our recorded sequence to which we could add cameras and a [Camera Cuts](Engine/Sequencer/HowTo/TracksCameraCut) track to provide multiple angles, music, effects or anything else we wanted. We could even take this sequence and embed it in another sequences as part of a [Shots Track](Engine/Sequencer/HowTo/TracksShot). 
 
 [OBJECT:EmbeddedVideo]
 [PARAMLITERAL:width]
@@ -81,15 +87,28 @@ hR5zBjYO5Oc
 [/PARAMLITERAL]
 [/OBJECT]
 
+In addition to recording gameplay, you can record your actions during level editing by assigning an Actor to Record. 
 
+[OBJECT:EmbeddedVideo]
+[PARAMLITERAL:width]
+640
+[/PARAMLITERAL]
+[PARAMLITERAL:height]
+360
+[/PARAMLITERAL]
+[PARAMLITERAL:videoid]
+rC_ZStHXm38
+[/PARAMLITERAL]
+[/OBJECT]
 
+Above we have placed a cube in our level and instructed the Sequence Recorder to record the cube. 
+We then moved the cube around and the Sequence Recorder captured the movements we entered through keyframes in a newly created Level Sequence. 
+When we playback our Level Sequence, a new Cube Actor is created as a spawnable in the Level Sequence which is why a second Cube appears when the sequence is active.
 
+[REGION:note]
+Only properties that can be keyframed can be captured and recorded when Level Editing recording. 
+[/REGION]
 
 
 
 
-
-
-
-	
-
diff --git a/Engine/Documentation/Source/Engine/Sequencer/Workflow/SequenceRecorder/SequenceRecorder.INT.udn b/Engine/Documentation/Source/Engine/Sequencer/Workflow/SequenceRecorder/SequenceRecorder.INT.udn
new file mode 100644
--- /dev/null
+++ b/Engine/Documentation/Source/Engine/Sequencer/Workflow/SequenceRecorder/SequenceRecorder.INT.udn
@@ -0,0 +1,164 @@
+Availability:Docs
+Title: Sequence Recorder
+Crumbs: %ROOT%, Engine
+Description:Overview of the Sequence Recorder tool and recording options to capture gameplay. 
+Navigation:topic
+parent:Engine/Sequencer/Workflow
+order:1
+Version: 4.13
+Related: Engine\Sequencer\Overview
+Related: Engine/Sequencer/HowTo
+
+The **Sequence Recorder** allows you to record Actor(s) from gameplay or level editing into Skeletal Animation, Transform, Event and Audio Tracks into a Level Sequence which can then be edited inside the **Sequencer Editor**. 
+
+[TOC(start:2 end:2)]
+
+
+## Opening the Sequence Recorder
+
+You can open the Sequence Recorder from the main **Menu Bar** under the **Window** tab: 
+
+![](EnableSequenceWindow.png)
+
+After selecting Sequence Recorder, the Sequence Recorder window will open: 
+
+![](NoActorRecordWindow.png)
+
+## Tracking Actors for Recording
+
+When you first open the Sequence Recorder, you will notice three buttons across the top of the window: 
+
+![](RecorderToolbar1.png)
+
+These are used to start recording, add Actors to track for recording or remove Actors from the recording list. 
+After clicking the **Add** button, both the **Record** option as well as the **Remove** option will become accessible. 
+This will also add an **Actor Recording** field to the Actor Tracking window with the Actor to track listed as **None**. 
+
+![](RecorderToolbar2.png)
+
+Selecting the Actor Recording in the window will populate the Sequence Recorder with additional options related to that recording.
+
+![](SequenceRecorderWindow.png)
+
+Above we have set the **Actor to Track** to the **ThirdPersonCharacter** Actor (which is our playable character). 
+
+[REGION:note]
+Please see the [](Engine\Sequencer\HowTo\SequenceRecorder) page for a step-by-step guide to recording Actors with the Sequence Recorder. 
+[/REGION]
+
+
+## Sequence Recorder Options
+
+Once you have the Sequence Recorder open, there are several properties you can adjust to determine how the Sequence Recorder records its Actors. 
+When first opening the Sequence Recorder, you will only be presented with the default **Sequence Recording** options. After adding Actor(s) to track for recording, the **Actor Recording** and **Animation Recording** sections become available on a per Actor basis. 
+
+Please see each of the tables below for a list of properties along with a description of each:
+
+### Sequence Recording
+
+| Option | Description |
+| ---- | ----------- |
+| **Create Level Sequence** | Whether to create a Level Sequence when recording or not. Actors and animations will be inserted into this sequence if enabled. |
+| **Immersive Mode** | Whether to maximize the viewport when recording. |
+| **Sequence Length** | Defines the maximum length of the sequence to record. |
+| **Recording Delay** | Delay that will be used before starting the recording. |
+| **Sequence Name** | The name of the sequence to record to. This name will also be used to auto-generate any assets created by this recording. |
+| **Sequence Recording Base Path** | Defines the location within the project to save the recording. Sub-assets will be created in subdirectories within the base path specified. |
+| **Record Nearby Spawned Actors** | Whether or not to record nearby spawned Actors as part of this recording. |
+| **Nearby Actor Recording Proximity** | Defines the proximity to currently recorded Actors to record newly spawned Actors. |
+| **Record World Settings Actor** | Whether to record the World Settings Actor in the sequence (some games use this to attach world SFX). |
+| **Actor Filter** | Filter to check spawned Actors against to see if they should be recorded. You can specify which Actor classes to record from the expandable option. |
+| **Level Sequence Actors to Trigger** | Sequence Actors to trigger playback on when recording starts (for example, this can be used for recording synchronized sequences). |
+| **Default Animation Settings** | The default settings applied to animation recording. Expand to see options for recording in World Space, Removing Root Animation whether or not to Auto Save the asset when recording is complete and the ability to Sample Rate and Length.  |
+| **Classes and Properties to Record** | The properties to record for specified classes. Component classes specified here will be recorded. If an Actor does not contain one of these classes it will be ignored. Expanded to define the classes and associated properties. |
+| **Animation Sub Directory** | The name of the subdirectory animations will be placed in. Leave this empty to place them into the same directory as the Sequence base path. |
+| **Record Sequencer Spawned Actors** | Whether to record Actors that are spawned by Sequencer itself (this is usually disabled as results can be unexpected). |
+| **Per Actor Settings** | Settings applied to Actors of a specified class. For each class you can define whether to record to a possessable (level-owned) or spawnable (sequence-owned).  |
+
+### Actor Recording
+
+| Option | Description |
+| ---- | ----------- |
+| **Record Visibility** | Whether to record the Actors visibility or not. |
+| **Record Transforms** | Whether to record Actor transforms. This can be useful if you want the Actors to end up in specific locations after the sequence. By default we rely on animations to provide transforms, but this can be changed using the **Record in World Space** animation setting. |
+| **Actor to Record** | This option is used to define which Actor you want to record. |
+
+### Animation Recording
+
+| Option | Description |
+| ---- | ----------- |
+| **Specify Target Animation** | Whether we should specify the target animation or auto-create it. |
+| **Target Animation** | Enabled when **Specify Target Animation** is enabled and is used to specify the animation you want to record to. |
+| **Animation Settings** | See **Default Animation Settings** under **Sequence Recording**. These are applied to the **Actor to Record**.  |
+| **Record to Possessable** | Whether to record possessable (level-owned) or spawnable (sequence-owned) Actors. Defaults to the global setting. |
+
+
+## Tips & Tricks
+
+This section contains some workflow tips & tricks you can use when working with the Sequence Recorder. 
+
+### Right-click Recording
+
+If you have an active **Level Sequence** open, rather than clicking the **Add** button inside the Sequence Recorder, you can **Right-click** on an Actor in your level and initiate recording of the Actor from the context menu (or press **Alt + R**). 
+This will add the Actor to the Actor Tracking window of the Sequence Recorder and automatically start recording after the set **Recording Delay** expires. 
+
+![](RightClickRecord.png)
+
+[REGION:note]
+Please note that this will initiate Editor Recording and not [Gameplay Recording](Engine\Sequencer\HowTo\SequenceRecorder) with the Sequence Recorder.
+[/REGION]
+
+
+### Subscene Recording Gameplay
+
+While it is perfectly okay to use the Sequence Recorder to assign an Actor and initiate gameplay recording, there is another method that you can take to record your gameplay using Sequence Recorder. 
+If you create a [Subscene Track](Engine\Sequencer\HowTo\TracksSub) inside a Level Sequence, you can add a primed **Recorded Sequence** on the Subscene Track during gameplay set to record your playable character.
+
+[OBJECT:EmbeddedVideo]
+[PARAMLITERAL:width]
+640
+[/PARAMLITERAL]
+[PARAMLITERAL:height]
+360
+[/PARAMLITERAL]
+[PARAMLITERAL:videoid]
+Wab8bB69TIo
+[/PARAMLITERAL]
+[/OBJECT]
+
+Above we create a new Level Sequence along with a Subscene Track, then during gameplay we hit **Shift + F1** to gain Mouse Control and add a primed Recorded Sequence based on our current player pawn. 
+When can then press the **Record** button in the viewport while playing our game to start recording and hit **Esc** to end recording. 
+
+
+### Recording Class Properties
+
+Want to record a change in a specific property value for an Actor during your Sequence Recording? 
+You can use the **Classes and Properties to Record** setting in the Sequence Recorder window to specify which class and its associated property (or properties) you want to track for recording.
+
+![](PropertyRecording.png)
+
+If that property changes at all during recording it will automatically be keyframed reflecting the changes during the recording. 
+If the Actor you are tracking does not contain one of the defined classes it will just be ignored. There are also some default elements added to this option which you can use or delete. 
+
+### Recording to Possessable
+
+By default when you elect to perform a recording the Actors you have tagged for tracking will automatically have a **Spawnable** version of that Actor created and added to your recording. 
+While the Spawnable version will only exist while the Level Sequence is open, there may be times when you want to "write" directly to an Actor in your level instead of creating a new Spawnable. 
+You can do so by checking the **Record to Possessable** option under the **Animation Recording** section in the Sequence Recorder. 
+
+![](RecordToPossessable.png)
+
+Checking this option will no longer create a Spawnable but instead will be using a Possessable Actor to apply the recording to. 
+
+[REGION:note]
+Please see [Possessables vs. Spawnables](https://docs.unrealengine.com/latest/INT/Engine/Sequencer/Overview/index.html#possessablesvsstopspawnables) for more information on each.
+[/REGION]
+### Immersive Mode
+
+When recording your gameplay, you may want a larger viewport to "play" your game while recording. 
+The **Immersive Mode** option inside the Sequence Recorder window will automatically maximize your viewport for you while recording is active. 
+
+![](ImmersiveMode.png)
+
+Once you stop recording, the viewport will return back to its normal size. 
+
diff --git a/Engine/Documentation/Source/Engine/Sequencer/Workflow/Workflow.INT.udn b/Engine/Documentation/Source/Engine/Sequencer/Workflow/Workflow.INT.udn
--- a/Engine/Documentation/Source/Engine/Sequencer/Workflow/Workflow.INT.udn
+++ b/Engine/Documentation/Source/Engine/Sequencer/Workflow/Workflow.INT.udn
@@ -7,9 +7,9 @@ parent:Engine/Sequencer
 order:5
 Version: 4.12
 
-There are several built-in shortcuts and keybindings you can use to speed up the rate at which you create your content in Sequencer. There are also several different workflows you can take depending upon the composition of your team. 
+In addition to the several built-in shortcuts and key bindings that you can use to speed up the rate at which you create your content in Sequencer, there are also features related to Sequencer that you can use to generate content for your cinematics such as the **Sequence Recorder**. 
 
-With an understanding of how to use Sequencer, the following pages may improve your productivity or workflows: 
+This page will highlight different features and tools that are related to Sequencer as well as some of the tips, tricks and workflows you can use during your cinematic creation process. 
 
 [REGION:note]
 If you are new to Sequencer it is recommended that you first check out the [**Overview**](Engine\Sequencer\Overview) page. 
@@ -20,10 +20,23 @@ If you are new to Sequencer it is recommended that you first check out the [**Ov
 		![](%ROOT%/start_icon.png)(convert:false)
 	[/PARAM]
 	[PARAM:title]
+	Features & Tools
 	[/PARAM]
 	[PARAM:links]
-		* [](Engine\Sequencer\Workflow\CheatSheet "%Engine\Sequencer\Workflow\CheatSheet:description%")
+		* [](Engine\Sequencer\Workflow\SequenceRecorder "%Engine\Sequencer\Workflow\SequenceRecorder:description%")
 		* [](Engine\Sequencer\Workflow\EditorPrefs "%Engine\Sequencer\Workflow\EditorPrefs:description%")
+	[/PARAM]
+[/OBJECT]
+
+[OBJECT:TopicButtonList]
+	[PARAM:icon]
+		![](%ROOT%/related_icon.png)(convert:false)
+	[/PARAM]
+	[PARAM:title]
+	Tips & Workflows
+	[/PARAM]
+	[PARAM:links]
+		* [](Engine\Sequencer\Workflow\CheatSheet "%Engine\Sequencer\Workflow\CheatSheet:description%")
 		* [](Engine\Sequencer\Workflow\ToolUsage "%Engine\Sequencer\Workflow\ToolUsage:description%")
 	[/PARAM]
 [/OBJECT]
