Applying SME feedback.

diff --git a/Engine/Documentation/Source/Engine/MediaFramework/Overview/MF_Overview.INT.udn b/Engine/Documentation/Source/Engine/MediaFramework/Overview/MF_Overview.INT.udn
--- a/Engine/Documentation/Source/Engine/MediaFramework/Overview/MF_Overview.INT.udn
+++ b/Engine/Documentation/Source/Engine/MediaFramework/Overview/MF_Overview.INT.udn
@@ -11,14 +11,14 @@ Related: Engine\MediaFramework\HowTo
 
 [TOC(start:2 end:2)]
 
-**Media Framework** uses several assets to allow the playback of videos inside Unreal Engine 4. Videos can be scrubbed, paused or rewound inside a **Media Player** asset as well as controlled through [C++](Programming) or [](Engine/Blueprints). 
-Whether you want a monitor in your level playing a video or you want to create a UI based interface with [](Engine/UMG) to allow players control over your video playback, you will need to specify the **Media Source** (where to find the media) to allow the Media Player the ability to play the media back. 
+**Media Framework** uses several assets to enable the playback of videos inside Unreal Engine 4. Videos can be scrubbed, paused, or rewound inside a **Media Player** asset, as well as controlled through [C++](Programming) or [](Engine/Blueprints). 
+Whether you want a surface in your level playing a video, or you want to create a [](Engine/UMG) interface to give players control over your video playback, you will first need to specify the **Media Source** so that the engine can locate your media asset.
 
 ## Media Source Types
 
 There are several different Media Source assets that you can use to define where the media is coming from. 
-Whether you have a file that is included as part of your project, streaming media from a website or using platform specific media you will need to define a source before you can play the video.
-You can add a Media Source inside the **Content Browser** by clicking the **Add New** Button then specifying which source type under the **Media** section.  
+Whether you have a file that is included as part of your project, streaming media from a website, or using platform specific media, you will need to define a source before you can play the video.
+You can add a Media Source inside the **Content Browser** by clicking the **Add New** Button then specifying which source type you want under the **Media** section.  
 
 ![](MediaSource.png)
 
@@ -34,10 +34,10 @@ The **File Media Source** asset is used for media files stored on your device or
 While a media file can exist anywhere, it's typically a good practice to move your media files to the **Content/Movies** folder of your project and then point to the media file. 
 This will ensure that your media is packaged correctly with your project. You will also see a yellow exclamation point when pointing to a file outside of this location as a warning. 
 
-This type of source asset also allows you to load the entire media file into memory and play it from there by enabling the **Precache File** option. 
+This type of source asset can be used to load the entire media file into memory and play it from there by enabling the **Precache File** option. 
 Based on the size of your file, caching time may vary, so take that into consideration when choosing this option. 
 
-Each of the source assets allows you to provide the name of the native media player with the **Default Player** field. If you leave this field empty, the media player will be detected automatically. 
+Each source asset supports a name, in the **Default Player** field, to identify the appropriate native media player for that asset. If you leave this field empty, the media player will be detected automatically. 
 
 [REGION:note]
 Please see the [](Engine\MediaFramework\HowTo\FileMediaSource) how-to for working with File Media Source assets. 
@@ -47,8 +47,8 @@ Please see the [](Engine\MediaFramework\HowTo\FileMediaSource) how-to for workin
 
 ![](MediaSource_Stream.png)
 
-This **Stream Media Source** asset allows you to specify a URL as the source of your media which is pulled from the internet. Currently, linking to a YouTube or DailyMotion style URL is not supported and you will need to link directly to a hosted file. 
-Above we are using a fake address with a fake hosted MP4 file. 
+This **Stream Media Source** asset takes a URL as the source of your media, and pulls it from the internet. Currently, linking to a YouTube or DailyMotion style URL is not supported, and you will need to link directly to a hosted file. 
+Please note that the URL shown above is just an example, and cannot be used.
 
 [REGION:note]
 Please see the [](Engine\MediaFramework\HowTo\StreamMediaSource) how-to for working with Stream Media Source assets. 
@@ -58,29 +58,29 @@ Please see the [](Engine\MediaFramework\HowTo\StreamMediaSource) how-to for work
 
 ![](MediaSource_Platform.png)
 
-The **Platform Media Source** asset allows overriding Media Sources on a per-platform basis. You can specify the Media Source to be used if no source was specified for a target platform as well as define the Media Sources per platform. 
+The **Platform Media Source** asset supports overriding Media Sources on a per-platform basis. You can specify the Media Source to be used if no source was specified for a target platform, as well as define the Media Sources per platform. 
 
-For example, supppose you want a particluar video to play only on Android Devices or only on PS4. Inside the **Media Sources** section you can designate which videos play on which platforms where leaving the Media Source blank will cause the **Default Source** to be used. 
+For example, supppose you want a particluar video to play only on Android Devices or only on PS4. Inside the **Media Sources** section, you can designate which videos play on which platforms. Leaving the Media Source blank will cause the **Default Source** to be used. 
 
 [REGION:warning]
-The Platform Media Source asset is extermely experimental does not work for 4.13 and are planned for a 4.14 release.
+The Platform Media Source asset is experimental and is not expected to work in version 4.13 of the engine. It is planned for a 4.14 release.
 [/REGION]
 
 
 ## Media Playlists
 
-While you can play individual Media Source assets individually, one asset you can use to play back multiple Media Sources is the **Media Playlist** asset. 
-After creating a few Media Source assets, you can add them to a Media Playlist which will automatically cycle through and play each back in the order you define.
+While you can play Media Source assets individually, the **Media Playlist** asset can play multiple Media Sources in order. 
+After creating a few Media Source assets, you can add them to a Media Playlist, which will automatically cycle through and play each asset in the order you define.
 You can create Media Playlists inside the **Content Browser** from the **Add New** Button under the **Media** section. 
 
 ![](MediaPlaylistButton.png)
 
-In the Media Playstlist window (pictured below) under the **Playlist** section you can define the source assets to include and playback order. 
+In the Media Playstlist window (pictured below), under the **Playlist** section, you can define the source assets to include and specific the playback order. 
 Any Media Source assets you have created will be displayed in the lower **Media Library** window which you can **Double-click** on to add to the end of the playlist. 
 
 ![](CustomizedPlaylist.png)
 
-Media Playlists can use a mixture of each Media Source types and playback in order unless specified through Blueprint Script. 
+Media Playlists can use a mixture of each Media Source type and playback in order unless specified through Blueprint Script. 
 
 [REGION:note]
 Please see the [](Engine\MediaFramework\HowTo\MediaPlaylists) how-to for working with Media Playlist assets. 
@@ -98,7 +98,7 @@ From the **Add New** button under **Media**, select the **Media Sound Wave** ass
 
 The Media Sound Wave asset provides **Debug**, **Concurrency** and **Attenuation** settings that can be adjusted to define how the asset is processed. 
 When linked to a Media Player asset, audio attached to a video source will automatically play back along with the video. 
-Typically you can leave the default settings for a newly created Media Sound Wave asset. However, if you need more control over how the sound is perceived you can adjust the Concurrency and Attenuation settings. 
+Typically, you can leave the default settings for a newly created Media Sound Wave asset. However, if you need more control over how the sound is perceived, you can adjust the Concurrency and Attenuation settings. 
 
 [REGION:note]
 Please see the [](Engine/Audio) section for more details on working with audio in Unreal Engine 4. 
@@ -118,7 +118,7 @@ Media Texture assets can be included within [](Engine/Rendering/Materials/) that
 
 Above, a Material including a Media Texture asset has been created and applied to a static mesh in the level. 
 Inside the [Texture Editor](Engine/Content/Types/Textures/Properties/Interface), the Media Texture displays the same playback position as it appears in the Material in the level. 
-In addition to the standard [](Engine/Content/Types/Textures/Properties), Media Textures allow you to define the X and Y Axis **Addressing** to either **Clamped**, **Mirror** or **World** values. 
+In addition to the standard [](Engine/Content/Types/Textures/Properties), Media Textures accept X and Y Axis **Addressing** values of **Clamped**, **Mirror** or **World** values. 
 
 [REGION:note]
 Please see the [](Engine\MediaFramework\HowTo) section for guides on how to setup and use Media Textures to play videos. 
@@ -142,19 +142,19 @@ This allows you to automatically create and assign a Media Texture and Media Sou
 ![](MediaPlayerWindow.png)
 
 Inside the Media Player Window, you can playback and preview any Media Source assets you have created (created source media are displayed the **Media Library** panel). 
-You can play, pause, rewind or skip to the beginning or end of a video from the toolbar across the top. Any Media Playlists that are currently playing videos will be displayed in the upper right **Playlist** window. 
+You can play, pause, rewind, or skip to the beginning or end of a video from the toolbar across the top. Any Media Playlists that are currently playing videos will be displayed in the upper right **Playlist** window. 
 
-In the lower-right **Details** panel you can define options for **Playback** or assign the **Sound Wave** and **Video Textures** (or an **Image Texture**) associated with the Media Player. 
+In the lower-right **Details** panel, you can define options for **Playback** or assign the **Sound Wave** and **Video Textures** (or an **Image Texture**) associated with the Media Player. 
 There is also an **Info** panel which will display info about the source such as the stream type, codecs, sample rate, dimensions and more. 
 
 [REGION:warning]
-Still Image Texture output is currently not implemented by any of the Media Player plug-in and this feature is considered experimental.
+Still Image Texture output is currently not implemented by any of the Media Player plug-ins and this feature is considered experimental.
 [/REGION]
 
 ## Media Assets & Scripting
 
-After your Media Player asset is set up and linked with the appropriate assets, in order to playback the media during game play you will need to first open the Media Source asset through Blueprint or C++. 
-Before you can open your Media Source, you will first need to get a reference to your Media Player asset which you can create a [Variable](Engine/Blueprints/UserGuide/Variables/) for inside any Blueprint. 
+Once your Media Player asset is set up and linked with the appropriate assets, it can be played during a game session. In order to play media in-game, you will first need to open the Media Source asset through Blueprint or C++. 
+To do this, create a [Variable](Engine/Blueprints/UserGuide/Variables/) inside any Blueprint that contains a reference to your Media Player asset. 
 
 ![](MediaPlayerDetailsSet.png)
 
@@ -176,7 +176,7 @@ If you uncheck the Play on Open option (pictured below), you will then need to c
 
 ![](UncheckPlayOnOpen.png)
 
-Below on **Event BeginPlay** we open a Media Playlist and play the **Media Player** when the **Right Mouse Button** is **Pressed**. 
+Below on **Event BeginPlay**, we open a Media Playlist and play the **Media Player** when the **Right Mouse Button** is **Pressed**. 
 
 ![](PlayWhenPressed.png)
 
@@ -186,10 +186,12 @@ This is because it may take some time for the Media Source to open and the Play
 [/REGION]
 
 You can call other functions from a Media Player reference such as [Binding Events](Engine/Blueprints/UserGuide/EventDispatcher/BindingAndUnbinding/) to various states of playback (such if the video is paused or if the video ends). 
-You can also get or set options like if a video is playing or can play if a video can loop, or play rate information. With a Media Player reference, you can drag off and see these options under the **Media Player** section of the Blueprint Context Menu. 
+You can also call several different Functions off a Media Player reference like checking if a video can be paused, is a video playing, is a video set to loop, accessing play rate information and others. 
 
 ![](MediaPlayerNodes.png)
 
+To see these options, drag off a Media Player reference then look under the **Media Player** section of the Blueprint Context Menu. 
+
 
 
 
