Applying Editorial Review Feedback.
- Split Xbox/PS4 sections off into Platform Specfic pages which link to the Media Framework section.

diff --git a/Engine/Documentation/Source/Engine/MediaFramework/TechReference/MF_TechRef.INT.udn b/Engine/Documentation/Source/Engine/MediaFramework/TechReference/MF_TechRef.INT.udn
--- a/Engine/Documentation/Source/Engine/MediaFramework/TechReference/MF_TechRef.INT.udn
+++ b/Engine/Documentation/Source/Engine/MediaFramework/TechReference/MF_TechRef.INT.udn
@@ -9,7 +9,7 @@ Version: 4.13
 Related: Engine\MediaFramework\Overview
 Related: Engine\MediaFramework\HowTo
 
-[TOC(start:2 end:3)]
+[TOC(start:2 end:2)]
 
 This page outlines the currently supported file extensions and URLs as well as outlines some troubleshooting tips on how to enable Media Framework with different platforms. 
 There are some known platform specific issues and limitations which we are currently investigating and trying to resolve, each of which have been identified in the Troubleshooting Platform sections. 
@@ -86,61 +86,6 @@ You can also check information about a Media Source inside the Media Player asse
 
 If **Verbose Logging** is enabled for the Media Assets log category, this information is also added to the Output Log. 
 
-### PlayStation 4
-
-The PS4Media player plug-in is disabled by default and must be enabled in the project’s [Plugin Settings](Programming/Plugins/).
-
-The video decoder on PS4 consumes a lot of Onion memory, especially with HLS. You may run out of memory when trying to play more than one video at the same time, because we allocate only 200 MB by default. 
-Currently, these default settings are not exposed in the UI but you can modify them by hand.
-
-To enable, add the following to the **PS4Engine.ini**: 
-
-	[SystemSettings]
-	r.GarlicHeapSizeMB=2600
-	r.OnionHeapSizeMB=200
-
-In your game's ** *.Target.cs**:
-
-	public override void SetupGlobalEnvironment(
-            	TargetInfo Target,
-            	ref LinkEnvironmentConfiguration OutLinkEnvironmentConfiguration,
-            	ref CPPEnvironmentConfiguration OutCPPEnvironmentConfiguration)
-	{
-            	// ...
-
-            	if (Target.Platform == UnrealTargetPlatform.PS4)
-            	{
-	OutCPPEnvironmentConfiguration.Definitions
-	.Add("GARLIC_HEAP_SIZE= (2900ULL * 1024 * 1024)");
-
-	OutCPPEnvironmentConfiguration.Definitions
-	.Add("ONION_HEAP_SIZE= (100ULL * 1024 * 1024)");
-            	}
-
-            	// ...
-	}
-
-To improve playback performance on PS4, your video file should be encoded as interleaved, so that audio and video can be read in a single sequential operation instead of seeking to different areas of the media file. 
-This can be done without re-encoding with [**mp4box**](https://gpac.wp.mines-telecom.fr/mp4box/) by adding the following: 
-
-	mp4box.exe -add non-interleaved-source.mp4 --inter 100 -delay 2=0 -new output.mp4
-
-One should also double check that the encoding is optimum for the PS4 as this will significantly reduce the CPU and GPU overhead of the player. 
-This can be done using [**FFMPEG**](https://ffmpeg.org/). For example, for a 1080p mp4 file we suggest the following: 
-
-	ffmpeg.exe -i "highqualitysource.mp4" -c:v libx264 -crf 20 -preset veryslow -maxrate 15000k -bufsize 30000k -slices 24 -x264opts b-pyramid=0:no-cabac:weightb=0:weightp=0 -refs 1 -pix_fmt yuv420p "non-interleaved-output.mp4"
-
-### Xbox One
-
-The MfMedia player plug-in is disabled by default and must be enabled in the project’s Plugin settings.
-It is possible to compile and test the MfMedia plug-in on Windows 7 and up if the Engine is compiled against Windows 7 (by default we compile against Windows Vista). 
-
-This can be accomplished by changing the following code in the `SetUpCompileEnvironment()` function in **UEBuildWindows.cs**:
-
-	// Windows Vista or higher required	
-	InBuildTarget.GlobalCompileEnvironment.Config.Definitions.Add("_WIN32_WINNT=0x0601");	
-	InBuildTarget.GlobalCompileEnvironment.Config.Definitions.Add("WINVER=0x0601");
-
 ### Windows
 
 The Windows Media Foundation (WMF) based player plug-in WmfMedia has a number of limitations for MP4 containers. 
@@ -164,8 +109,6 @@ See Microsoft's [MPEG-4 File Source](https://msdn.microsoft.com/en-us/library/wi
 * **Editor**
 	* The Playlist UI inside the Media Player only shows the currently opened Media Playlist asset. 
 		*	Workaround: Edit the Playlist contents by **Double-clicking** on the MediaPlaylist asset to open its own editor. 
-* **Playstation 4**
-	* Playlists currently do not work.
 * **Sequencer**
 	* Media Playback will not render out properly when recording cinematics from Sequencer. 
 		* We have this scheduled to release in the near future as we know this is important.
