Applying a few small peices of SME feedback.

diff --git a/Engine/Documentation/Source/Enterprise/Automotive/AutomotiveMaterialPack/AutomotiveMaterialPack.INT.udn b/Engine/Documentation/Source/Enterprise/Automotive/AutomotiveMaterialPack/AutomotiveMaterialPack.INT.udn
--- a/Engine/Documentation/Source/Enterprise/Automotive/AutomotiveMaterialPack/AutomotiveMaterialPack.INT.udn
+++ b/Engine/Documentation/Source/Enterprise/Automotive/AutomotiveMaterialPack/AutomotiveMaterialPack.INT.udn
@@ -110,7 +110,7 @@ In the following section, we will take a look at some of the parameters that all
 
 ### Car Paint
 
-The Car Paint Material is a very robust Material that has 29 unique parameters that can be adjusted. When trying to modify this Material for your project, we recommend you starting adjusting and experimenting with the following parameters:
+The Car Paint Material is a very robust Material that has 29 unique parameters that can be adjusted. When trying to modify this Material for your project, we recommend you start by adjusting and experimenting with the following parameters:
 
 * **Base Color/Flake**
 * **Scale/Flake**
@@ -136,7 +136,7 @@ The following image sequence shows off the variety of paint colors and surface l
 
 ### Orange Peel Material Parameter
 
-The Clearcoat lacquer that is applied on top of automotive paint often has a slightly bumpy surface similar to the peel of an Orange. Because of this resemblance, this layer of lacquer is referred to as Orange Peel. You can control the intensity of the Orange Peel using the **Orange Peel Strength** and **Orange Peel Scale** parameters. To disable the Orange Peel effect, set the **Orange Peel Strength** to **0**.
+The Clearcoat lacquer that is applied on top of automotive paint often has a slightly bumpy surface similar to the peel of an Orange. Because of this resemblance, this dimpled or uneven part of the Clearcoat is referred to as the Orange Peel. You can control the intensity of the Orange Peel using the **Orange Peel Strength** and **Orange Peel Scale** parameters. To disable the Orange Peel effect, set the **Orange Peel Strength** to **0**.
 
 [REGION:fullwidth]
 [INCLUDE:Enterprise\Automotive\AutomotiveMaterialPack\#OP]
@@ -163,7 +163,7 @@ For more information on using Dual Normals with a Clearcoat Material make sure t
 
 ### Caliper
 
-The caliper is part of a car's braking system and is usually placed on the brake disc in between the wheel rim and the brake pad. 
+The caliper is part of a car's braking system and is placed on the brake rotor in between the wheel rim and the brake pad. 
 
 ![](AP_05_Caliper_01.png)
 
@@ -221,7 +221,7 @@ The rubber Material that comes with the Automotive Material Pack can be used for
 * **Interior Pieces**
 * **Tires**
 
-When looking at the rubber Material Instance you will notice that it has a tire-specific option called **Use Sidewall**. Enabling this option will add a new Normal Map Texture that is for display the lettering that can be found on the tire sidewalls. Failing to properly line up the UV of the tire sidewall to the Normal map of the tire will result in the tire lettering and logos not show up correctly on the tire sidewall like in the image below. 
+When looking at the rubber Material Instance you will notice that it has a tire-specific option called **Use Sidewall**. Enabling this option will add a new Normal Map Texture that is used to display the text that can be found on the tire sidewalls. Failing to properly line up the UV of the tire sidewall to the Normal map of the tire will result in the tire lettering and logos not show up correctly on the tire sidewall like in the image below. 
 
 [OBJECT:ComparisonSlider]
 	[PARAM:before]
