Fixing the Alcantara description as it was wrong.

diff --git a/Engine/Documentation/Source/Enterprise/Automotive/AutomotiveMaterialPack/AutomotiveMaterialPack.INT.udn b/Engine/Documentation/Source/Enterprise/Automotive/AutomotiveMaterialPack/AutomotiveMaterialPack.INT.udn
--- a/Engine/Documentation/Source/Enterprise/Automotive/AutomotiveMaterialPack/AutomotiveMaterialPack.INT.udn
+++ b/Engine/Documentation/Source/Enterprise/Automotive/AutomotiveMaterialPack/AutomotiveMaterialPack.INT.udn
@@ -270,7 +270,7 @@ The leather Material that comes with the Automotive Material Pack has 12 differe
 
 ### Alcantara 
 
-The Alcantara Material that comes with the Automotive Material Pack is used to display the fabrics you find in cars. When using the Alcantara Material you can adjust the following parameters to manipulate the look of the fabric.
+The Alcantara Material that comes with the Automotive Material Pack is a synthetic suede Materail that is found in high end sports cars. When using the Alcantara Material you can adjust the following parameters to manipulate the look of the fabric.
 
 * **Fabric Scale / Strength**
 * **Cloth Folds Scale / Strength**
