Applying SME feedback.

diff --git a/Engine/Documentation/Source/Engine/MediaFramework/HowTo/MediaPlaylists/UsingMediaPlaylists.INT.udn b/Engine/Documentation/Source/Engine/MediaFramework/HowTo/MediaPlaylists/UsingMediaPlaylists.INT.udn
--- a/Engine/Documentation/Source/Engine/MediaFramework/HowTo/MediaPlaylists/UsingMediaPlaylists.INT.udn
+++ b/Engine/Documentation/Source/Engine/MediaFramework/HowTo/MediaPlaylists/UsingMediaPlaylists.INT.udn
@@ -116,5 +116,6 @@ When pressing the **2** key, the video playback will jump to the second video in
 Our **Media Player** asset is set to **Play on Open** which is why each video is played when we call the open playlist function. 
 If you want to open a Media Source but not start playback immediately, you can uncheck the Play on Open option inside the Media Player asset then use the **Play** function to play an open Media Source asset. 
 
-
+You can also navigate an opend playlist with the **Next** function (to pick the next item in the playlist) and **Previous** (to pick the previous item in the playlist). 
+The Media Player also has a property called **Shuffle** that when enabled will pick a random item from the playlist when the next source is to be opened (also affects Previous and Next). 
 
