#UE4 Docs: Setting pages to publish

diff --git a/Engine/Documentation/Source/Engine/MediaFramework/HowTo/ControllingPlayback/ControllingPlayback.INT.udn b/Engine/Documentation/Source/Engine/MediaFramework/HowTo/ControllingPlayback/ControllingPlayback.INT.udn
--- a/Engine/Documentation/Source/Engine/MediaFramework/HowTo/ControllingPlayback/ControllingPlayback.INT.udn
+++ b/Engine/Documentation/Source/Engine/MediaFramework/HowTo/ControllingPlayback/ControllingPlayback.INT.udn
@@ -1,4 +1,4 @@
-Availability:Docs
+Availability:Public
 Title: Rewind & Fast Forward Playback
 Crumbs:%ROOT%, Engine
 Description:Example of how you can control playback of your media assets.
diff --git a/Engine/Documentation/Source/Engine/MediaFramework/HowTo/FileMediaSource/UsingFileMedia.INT.udn b/Engine/Documentation/Source/Engine/MediaFramework/HowTo/FileMediaSource/UsingFileMedia.INT.udn
--- a/Engine/Documentation/Source/Engine/MediaFramework/HowTo/FileMediaSource/UsingFileMedia.INT.udn
+++ b/Engine/Documentation/Source/Engine/MediaFramework/HowTo/FileMediaSource/UsingFileMedia.INT.udn
@@ -1,4 +1,4 @@
-Availability:Docs
+Availability:Public
 Title: Playing Videos in a Level
 Crumbs:%ROOT%, Engine
 Description:Describes how to use the File Media asset to play videos inside a level. 
diff --git a/Engine/Documentation/Source/Engine/MediaFramework/HowTo/MF_HowTo.INT.udn b/Engine/Documentation/Source/Engine/MediaFramework/HowTo/MF_HowTo.INT.udn
--- a/Engine/Documentation/Source/Engine/MediaFramework/HowTo/MF_HowTo.INT.udn
+++ b/Engine/Documentation/Source/Engine/MediaFramework/HowTo/MF_HowTo.INT.udn
@@ -1,4 +1,4 @@
-Availability:Docs
+Availability:Public
 Title: Media Framework How-To
 Crumbs:%ROOT%, Engine
 Description:Links to step-by-step guides for working with Media Framework in Unreal Engine 4. 
diff --git a/Engine/Documentation/Source/Engine/MediaFramework/HowTo/MediaPlaylists/UsingMediaPlaylists.INT.udn b/Engine/Documentation/Source/Engine/MediaFramework/HowTo/MediaPlaylists/UsingMediaPlaylists.INT.udn
--- a/Engine/Documentation/Source/Engine/MediaFramework/HowTo/MediaPlaylists/UsingMediaPlaylists.INT.udn
+++ b/Engine/Documentation/Source/Engine/MediaFramework/HowTo/MediaPlaylists/UsingMediaPlaylists.INT.udn
@@ -1,4 +1,4 @@
-Availability:Docs
+Availability:Public
 Title: Working with Media Playlists
 Crumbs:%ROOT%, Engine
 Description:Example of how you can use a Media Playlist to playback multiple Media Source assets. 
@@ -26,12 +26,10 @@ You will also need a few [Supported Video Files](Engine\MediaFramework\TechRefer
 
 	![](Playlist_00.png)
 
-1.	**Right-click** on the **Movies** folder and choose **Show in Explorer** then place your video files in the folder. 
+1.	**Right-click** on the **Movies** folder and choose **Show in Explorer** then place your video files in the folder. Placing your video files inside the **Content/Movies** folder of your project will ensure the videos are properly packaged. 
 
 	![](Playlist_01.png)
 
-	Placing your video files inside the **Content/Movies** folder of your project will ensure the videos are properly packaged. 
-
 1.	Create two **File Media Source** assets (called **Video_01** and **Video_02**) point the **File Path** to each respective video. 
 
 	![](Playlist_02.png)
@@ -40,23 +38,18 @@ You will also need a few [Supported Video Files](Engine\MediaFramework\TechRefer
 
 	![](Playlist_03.png)
 
-1.	Inside **MyPlaylist**, **Double-click** on each **File Media Source** to add it to the playlist. 
+1.	Inside **MyPlaylist**, **Double-click** on each **File Media Source** to add it to the playlist. You can also use the **+** button to add an empty element which you can then assign your source media. In this example we are using **File Media Source** assets, however you can also include **Stream Media Source** assets in a Playlist. 
 
 	![](Playlist_04.png)
 
-	You can also use the **+** button to add an empty element which you can then assign your source media. 
-	In this example we are using **File Media Source** assets, however you can also include **Stream Media Source** assets in a Playlist. 
-
 1.	Create a **Media Player** asset called **VideoPlayer**. 
 
 	![](Playlist_05.png)
 
-1.	In the **Create Media Player** window, check **Audio output SoundWave asset** and **Video output Media Texture asset**. 
+1.	In the **Create Media Player** window, check **Audio output SoundWave asset** and **Video output Media Texture asset**. This will automatically create a Sound Wave and Media Texture asset linked to this Media Player asset which is needed for playback. 
 
 	![](Playlist_06.png)
 
-	This will automatically create a Sound Wave and Media Texture asset linked to this Media Player asset which is needed for playback. 
-
 	![](Playlist_07.png)
 
 1.	In the **Content Browser** under **StarterContent/Shapes** drag a **Shape_Plane** into the level.  
@@ -65,22 +58,18 @@ You will also need a few [Supported Video Files](Engine\MediaFramework\TechRefer
 
 1.	Use the **Translate (W)**, **Rotation (E)** and **Scale (R)** tools to adjust the size and location of the Static Mesh as desired. 
 
-1.	Hold **Ctrl** then select both the **Sound Wave** and **Media Texture** assets and drop them on top of the **Static Mesh** in the level. 
+1.	Hold **Ctrl** then select both the **Sound Wave** and **Media Texture** assets and drop them on top of the **Static Mesh** in the level. This will automatically create a **Material** and apply it to the Static Mesh. 
 
 	![](Playlist_10.png)
 
-	This will automatically create a **Material** and apply it to the Static Mesh. 
-
 1.	From the Toolbar, click the **Blueprints** button and **Open Level Blueprint**. 
 
 	![](Playlist_11.png)
 
-1.	Add a **Variable** of the **Media Player Reference** type called **Media Player** and set it to your **VideoPlayer** Media Player asset. 
+1.	Add a **Variable** of the **Media Player Reference** type called **Media Player** and set it to your **VideoPlayer** Media Player asset. You may need to **Compile** before you can assign the **Default Value**. 
 
 	![](Playlist_12.png)
 
-	You may need to **Compile** before you can assign the **Default Value**. 
-
 1.	**Right-click** in the **Event Graph** and add a **1 Keyboard Event**, **2 Keyboard Event** and **3 Keyboard Event**.  
 
 	![](Playlist_13.png)
@@ -98,6 +87,15 @@ You will also need a few [Supported Video Files](Engine\MediaFramework\TechRefer
 	
 ## End Result
 
+When playing in the editor, the video playback will not initially start when the level is loaded. When the player presses the **1** key, the playlist will start playing from the first video in the playlist. 
+When pressing the **2** key, the video playback will jump to the second video in the playlist. When pressing **3**, we are opening and playing the video from the start of the playlist which is the same as opening the playlist at index 0. 
+
+Our **Media Player** asset is set to **Play on Open** which is why each video is played when we call the open playlist function. 
+If you want to open a Media Source but not start playback immediately, you can uncheck the Play on Open option inside the Media Player asset then use the **Play** function to play an open Media Source asset. 
+
+You can also navigate an opend playlist with the **Next** function (to pick the next item in the playlist) and **Previous** (to pick the previous item in the playlist). 
+The Media Player also has a property called **Shuffle** that when enabled will pick a random item from the playlist when the next source is to be opened (also affects Previous and Next). 
+
 [OBJECT:EmbeddedVideo]
 [PARAMLITERAL:width]
 640
@@ -109,13 +107,3 @@ You will also need a few [Supported Video Files](Engine\MediaFramework\TechRefer
 kjbgHVnQxBI
 [/PARAMLITERAL]
 [/OBJECT]
-
-When playing in the editor, the video playback will not initially start when the level is loaded. When the player presses the **1** key, the playlist will start playing from the first video in the playlist. 
-When pressing the **2** key, the video playback will jump to the second video in the playlist. When pressing **3**, we are opening and playing the video from the start of the playlist which is the same as opening the playlist at index 0. 
-
-Our **Media Player** asset is set to **Play on Open** which is why each video is played when we call the open playlist function. 
-If you want to open a Media Source but not start playback immediately, you can uncheck the Play on Open option inside the Media Player asset then use the **Play** function to play an open Media Source asset. 
-
-You can also navigate an opend playlist with the **Next** function (to pick the next item in the playlist) and **Previous** (to pick the previous item in the playlist). 
-The Media Player also has a property called **Shuffle** that when enabled will pick a random item from the playlist when the next source is to be opened (also affects Previous and Next). 
-
diff --git a/Engine/Documentation/Source/Engine/MediaFramework/HowTo/StreamMediaSource/UsingStreamMedia.INT.udn b/Engine/Documentation/Source/Engine/MediaFramework/HowTo/StreamMediaSource/UsingStreamMedia.INT.udn
--- a/Engine/Documentation/Source/Engine/MediaFramework/HowTo/StreamMediaSource/UsingStreamMedia.INT.udn
+++ b/Engine/Documentation/Source/Engine/MediaFramework/HowTo/StreamMediaSource/UsingStreamMedia.INT.udn
@@ -1,4 +1,4 @@
-Availability:Docs
+Availability:Public
 Title: Streaming Videos in a Level
 Crumbs:%ROOT%, Engine
 Description:Describes how to use the Stream Media asset to stream videos from the internet in Unreal Engine 4. 
diff --git a/Engine/Documentation/Source/Engine/MediaFramework/MediaFramework.INT.udn b/Engine/Documentation/Source/Engine/MediaFramework/MediaFramework.INT.udn
--- a/Engine/Documentation/Source/Engine/MediaFramework/MediaFramework.INT.udn
+++ b/Engine/Documentation/Source/Engine/MediaFramework/MediaFramework.INT.udn
@@ -1,10 +1,10 @@
-Availability:Docs
+Availability:Public
 Title: Media Framework
 Crumbs:%ROOT%, Engine
 Description:Working with Media Playback functionality in Unreal Engine 4. 
 Navigation:topic
 parent:Engine
-order:1
+order:12
 Version: 4.13
 
 [VAR:Topic]
@@ -76,5 +76,5 @@ to resolve those issues as well as add additional features and support.
 
 ## Essentials
 
-[DIR(output:"listbutton" parent:"Engine\MediaFramework" org:"hierarchy" end:"1")]
+[DIR(output:"listbutton" parent:"Engine/MediaFramework" org:"hierarchy" end:"1")]
 
diff --git a/Engine/Documentation/Source/Engine/MediaFramework/Overview/MF_Overview.INT.udn b/Engine/Documentation/Source/Engine/MediaFramework/Overview/MF_Overview.INT.udn
--- a/Engine/Documentation/Source/Engine/MediaFramework/Overview/MF_Overview.INT.udn
+++ b/Engine/Documentation/Source/Engine/MediaFramework/Overview/MF_Overview.INT.udn
@@ -1,4 +1,4 @@
-Availability:Docs
+Availability:Public
 Title: Media Framework Overview
 Crumbs:%ROOT%, Engine
 Description:Overview of Media Framework functionality in Unreal Engine 4. 
diff --git a/Engine/Documentation/Source/Engine/Rendering/LightingAndShadows/LightingChannels/LightingChannels.INT.udn b/Engine/Documentation/Source/Engine/Rendering/LightingAndShadows/LightingChannels/LightingChannels.INT.udn
--- a/Engine/Documentation/Source/Engine/Rendering/LightingAndShadows/LightingChannels/LightingChannels.INT.udn
+++ b/Engine/Documentation/Source/Engine/Rendering/LightingAndShadows/LightingChannels/LightingChannels.INT.udn
@@ -1,4 +1,4 @@
-Availability:Docs
+Availability:Public
 Title: Lighting Channels
 Crumbs: %ROOT%, Engine
 Description:Lighting Channels ia a rendering feature that enables you to choose what renderable surfaces are lit by specific lights.
diff --git a/Engine/Documentation/Source/Engine/Sequencer/HowTo/ImportExportEDLs/ImportExportEDLs.INT.udn b/Engine/Documentation/Source/Engine/Sequencer/HowTo/ImportExportEDLs/ImportExportEDLs.INT.udn
--- a/Engine/Documentation/Source/Engine/Sequencer/HowTo/ImportExportEDLs/ImportExportEDLs.INT.udn
+++ b/Engine/Documentation/Source/Engine/Sequencer/HowTo/ImportExportEDLs/ImportExportEDLs.INT.udn
@@ -1,4 +1,4 @@
-Availability: Docs
+Availability: Public
 Title:Importing & Exporting Edit Decision Lists (EDLs)
 Crumbs: %ROOT%, Engine, Engine/Sequencer, Engine/Sequencer/HowTo
 Description:Illustrates how you can import and export EDLs for use in external video editing software applications. 
diff --git a/Engine/Documentation/Source/Engine/UMG/HowTo/InWorldWidgetInteraction/InWorldWidgetInteraction.INT.udn b/Engine/Documentation/Source/Engine/UMG/HowTo/InWorldWidgetInteraction/InWorldWidgetInteraction.INT.udn
--- a/Engine/Documentation/Source/Engine/UMG/HowTo/InWorldWidgetInteraction/InWorldWidgetInteraction.INT.udn
+++ b/Engine/Documentation/Source/Engine/UMG/HowTo/InWorldWidgetInteraction/InWorldWidgetInteraction.INT.udn
@@ -1,4 +1,4 @@
-Availability: Docs
+Availability: Public
 Title:Creating 3D Widget Interaction
 Crumbs: %ROOT%, Engine, Engine/UMG, Engine/UMG/HowTo
 Description:An example of how to create, display and remove UI elements from the screen during gameplay. 
diff --git a/Engine/Documentation/Source/Engine/UMG/HowTo/VirtualKeyboards/ActorSetup/ActorSetup.INT.udn b/Engine/Documentation/Source/Engine/UMG/HowTo/VirtualKeyboards/ActorSetup/ActorSetup.INT.udn
--- a/Engine/Documentation/Source/Engine/UMG/HowTo/VirtualKeyboards/ActorSetup/ActorSetup.INT.udn
+++ b/Engine/Documentation/Source/Engine/UMG/HowTo/VirtualKeyboards/ActorSetup/ActorSetup.INT.udn
@@ -1,4 +1,4 @@
-Availability: Docs
+Availability: Public
 Title:3. Actor Setup & Testing
 Crumbs: %ROOT%, Engine, Engine/UMG, Engine/UMG/HowTo
 Description:Demonstrates how you can create and interaction with a 3D Virtual Keyboard with UMG. 
diff --git a/Engine/Documentation/Source/Engine/UMG/HowTo/VirtualKeyboards/CharacterSetup/CharacterSetup.INT.udn b/Engine/Documentation/Source/Engine/UMG/HowTo/VirtualKeyboards/CharacterSetup/CharacterSetup.INT.udn
--- a/Engine/Documentation/Source/Engine/UMG/HowTo/VirtualKeyboards/CharacterSetup/CharacterSetup.INT.udn
+++ b/Engine/Documentation/Source/Engine/UMG/HowTo/VirtualKeyboards/CharacterSetup/CharacterSetup.INT.udn
@@ -1,4 +1,4 @@
-Availability: Docs
+Availability: Public
 Title:1. Character Setup
 Crumbs: %ROOT%, Engine, Engine/UMG, Engine/UMG/HowTo
 Description:Here we add the Widget Interaction component to our character and create the other assets we will need. 
diff --git a/Engine/Documentation/Source/Engine/UMG/HowTo/VirtualKeyboards/CreatingVirtualKeyboards.INT.udn b/Engine/Documentation/Source/Engine/UMG/HowTo/VirtualKeyboards/CreatingVirtualKeyboards.INT.udn
--- a/Engine/Documentation/Source/Engine/UMG/HowTo/VirtualKeyboards/CreatingVirtualKeyboards.INT.udn
+++ b/Engine/Documentation/Source/Engine/UMG/HowTo/VirtualKeyboards/CreatingVirtualKeyboards.INT.udn
@@ -1,4 +1,4 @@
-Availability: Docs
+Availability: Public
 Title:Creating Virtual Keyboard Interactions
 Crumbs: %ROOT%, Engine, Engine/UMG, Engine/UMG/HowTo
 Description:Demonstrates how you can create and interaction with a 3D Virtual Keyboard with UMG. 
diff --git a/Engine/Documentation/Source/Engine/UMG/HowTo/VirtualKeyboards/WidgetSetup/WidgetSetup.INT.udn b/Engine/Documentation/Source/Engine/UMG/HowTo/VirtualKeyboards/WidgetSetup/WidgetSetup.INT.udn
--- a/Engine/Documentation/Source/Engine/UMG/HowTo/VirtualKeyboards/WidgetSetup/WidgetSetup.INT.udn
+++ b/Engine/Documentation/Source/Engine/UMG/HowTo/VirtualKeyboards/WidgetSetup/WidgetSetup.INT.udn
@@ -1,4 +1,4 @@
-Availability: Docs
+Availability: Public
 Title:2. Widget Blueprint Setup
 Crumbs: %ROOT%, Engine, Engine/UMG, Engine/UMG/HowTo
 Description:Here we create the keypad and keypad display widgets that will appear in the level. 
diff --git a/Engine/Documentation/Source/Engine/UMG/UserGuide/WidgetInteraction/WidgetInteraction.INT.udn b/Engine/Documentation/Source/Engine/UMG/UserGuide/WidgetInteraction/WidgetInteraction.INT.udn
--- a/Engine/Documentation/Source/Engine/UMG/UserGuide/WidgetInteraction/WidgetInteraction.INT.udn
+++ b/Engine/Documentation/Source/Engine/UMG/UserGuide/WidgetInteraction/WidgetInteraction.INT.udn
@@ -1,4 +1,4 @@
-Availability:Docs
+Availability:Public
 Title: Widget Interaction Component
 Crumbs:%ROOT%, Engine, Engine/UMG, Engine/UMG/UserGuide
 Description:Outlines how to use the Widget Interaction Component to interact with 3D Widgets.
diff --git a/Engine/Documentation/Source/Platforms/Linux/BeginnerLinuxDeveloper/SupportedGfxDrivers/SupportedGfxDrivers.INT.udn b/Engine/Documentation/Source/Platforms/Linux/BeginnerLinuxDeveloper/SupportedGfxDrivers/SupportedGfxDrivers.INT.udn
--- a/Engine/Documentation/Source/Platforms/Linux/BeginnerLinuxDeveloper/SupportedGfxDrivers/SupportedGfxDrivers.INT.udn
+++ b/Engine/Documentation/Source/Platforms/Linux/BeginnerLinuxDeveloper/SupportedGfxDrivers/SupportedGfxDrivers.INT.udn
@@ -1,4 +1,4 @@
-Availability: Docs
+Availability: Public
 Title: Supported Graphics Drivers
 Crumbs: %ROOT%, Platforms, Platforms/Linux, Platforms/Linux/BeginnerLinuxDeveloper
 Description:This page provides an overview of supported graphics drivers for Linux game developers.
diff --git a/Engine/Documentation/Source/Platforms/iOS/Windows/Windows.INT.udn b/Engine/Documentation/Source/Platforms/iOS/Windows/Windows.INT.udn
--- a/Engine/Documentation/Source/Platforms/iOS/Windows/Windows.INT.udn
+++ b/Engine/Documentation/Source/Platforms/iOS/Windows/Windows.INT.udn
@@ -1,4 +1,4 @@
-Availability: Docs
+Availability: Public
 Title: Building for iOS on Windows
 Crumbs:%ROOT%, Platforms
 Description:How to remotely compile on a Mac from a Windows machine to create an iOS build. 
