#UE4 Docs: set page to publish

diff --git a/Engine/Documentation/Source/Engine/MediaFramework/TechReference/Upgrading/Upgrading.INT.udn b/Engine/Documentation/Source/Engine/MediaFramework/TechReference/Upgrading/Upgrading.INT.udn
--- a/Engine/Documentation/Source/Engine/MediaFramework/TechReference/Upgrading/Upgrading.INT.udn
+++ b/Engine/Documentation/Source/Engine/MediaFramework/TechReference/Upgrading/Upgrading.INT.udn
@@ -1,4 +1,4 @@
-Availability:Docs
+Availability:Public
 Title: Upgrading Media Framework Content to 4.13
 Crumbs:%ROOT%, Engine
 Description:Updating projects containing Media Framework assets and playback logic to Unreal Engine 4.13
