Shelved response to BenM's SME review comments. There was some shuffling of content, so I'm leaving this for RobertG to have a chance to confirm that this is acceptable. (Submitting as per SeanG's comment on the JIRA.)

diff --git a/Engine/Documentation/Source/Programming/Development/BuildGraph/BuildGraph.INT.udn b/Engine/Documentation/Source/Programming/Development/BuildGraph/BuildGraph.INT.udn
--- a/Engine/Documentation/Source/Programming/Development/BuildGraph/BuildGraph.INT.udn
+++ b/Engine/Documentation/Source/Programming/Development/BuildGraph/BuildGraph.INT.udn
@@ -10,7 +10,7 @@ tags:BuildGraph
 **BuildGraph** is a script-based build automation system that features graphs of building blocks common to Unreal Engine 4 (UE4) projects. 
 BuildGraph integrates with UnrealBuildTool, AutomationTool, and the editor, and can be extended and customized for your projects.
 
-BuildGraph scripts are written with [XML](https://www.w3.org/TR/NOTE-dcd]), specifying a graph of user-defined nodes with 
+BuildGraph scripts are written in XML, specifying a graph of user-defined nodes with 
 dependencies between them. Each node consists of tasks executed in sequence to produce some sort of output (e.g., compile 
 this project, then cook, run this test, and so on). When asked to build a target (that is, a node or named output) BuildGraph 
 will execute all of the nodes in the graph required to make that happen.
@@ -32,8 +32,47 @@ tooltips, validation, and Intellisense features while editing.
 
 ## Writing BuildGraph Scripts
 
+If you want to learn how to write your own **BuildGraph** scripts, it helps to know the various parts that go into making a Graph. 
+Graphs are created with the following elements:
+
+* Tasks: Actions that are executed as part of a build process (compiling, cooking, etc.).
+* Nodes: A named sequence of ordered tasks that are executed to produce outputs. Before they can be executed, nodes may depend 
+on other nodes executing their tasks first.
+* Agents: A group of nodes that are executed on the same machine (if running as part of a build system). Agents have no effect 
+when building locally. 
+* Triggers: Container for groups that should only be executed after manual intervention. 
+* Aggregates: Groups of nodes and named outputs that can be referred to with a single name.
+
+Scripts typically make frequent use of properties for reusable or conditionally defined values. Properties are declared with 
+the `<Property>` element, and are scoped to the point of their first declaration. Properties referenced with the `$(Property 
+Name)` notation are valid within all attribute strings, and will be expanded when the script is read. Properties that can be 
+supplied by the user on the command-line can be declared with the `<Option>` element, and environment variables can be 
+imported into properties using the `<EnvVar>` element.
+
+Any element may be conditionally defined via the "If" attribute. See below for the syntax of conditional expressions.
+
+BuildGraph is typically used for packaged games, so filtering and file manipulation are natively supported. Any attribute 
+that accepts a list of files may consist of a Perforce style wildcard (matching any number of "...", "*" and "?" patterns in 
+any location), a full path name, or a reference to a tagged collection of files. Attributes are denoted by prefixing a `#` 
+character. Files may be added to a tag set using the `<Tag>` task, which also allows the performance of union/difference style 
+operations. Each node can declare multiple outputs in the form of a list of named tags, which other nodes can then depend on.
+
+Graphs may be executed in parallel as part of a build system. To do so, the initial graph configuration is generated by 
+running with the `-Export=<Filename>` argument (producing a JSON file listing the nodes and dependencies to execute). Each 
+participating agent should be synced to the same changelist, and AutomationTool (UAT) should be re-run with the appropriate 
+`-SingleNode=<Name>` argument. Outputs from different nodes are transferred between agents via shared storage (typically, a 
+network share) the path to which can be specified on the command line using the `-SharedStorageDir=<Path>` argument. Note that 
+the allocation of machines (and coordination between them) is assumed to be managed by an external system.
+
+The syntax of elements used to define BuildGraph constructs is listed in the following sections.
+
+## BuildGraph Scripting Reference
+
 [REGION:topics third]
-%Programming/Development/BuildGraph/ScriptAnatomy:topiccompact%
+%Programming/Development/BuildGraph/ScriptAnatomy/Types:topiccompact%
+%Programming/Development/BuildGraph/ScriptAnatomy/Elements:topiccompact%
+%Programming/Development/BuildGraph/ScriptAnatomy/Conditions:topiccompact%
+%Programming/Development/BuildGraph/ScriptAnatomy/Tasks:topiccompact%
 [/REGION]
 
 ## Using BuildGraph
diff --git a/Engine/Documentation/Source/Programming/Development/BuildGraph/ScriptAnatomy/BuildGraphScriptAnatomy.INT.udn b/Engine/Documentation/Source/Programming/Development/BuildGraph/ScriptAnatomy/BuildGraphScriptAnatomy.INT.udn
deleted file mode 100644
diff --git a/Engine/Documentation/Source/Programming/Development/BuildGraph/ScriptAnatomy/Conditions/BuildGraphScriptConditions.INT.udn b/Engine/Documentation/Source/Programming/Development/BuildGraph/ScriptAnatomy/Conditions/BuildGraphScriptConditions.INT.udn
--- a/Engine/Documentation/Source/Programming/Development/BuildGraph/ScriptAnatomy/Conditions/BuildGraphScriptConditions.INT.udn
+++ b/Engine/Documentation/Source/Programming/Development/BuildGraph/ScriptAnatomy/Conditions/BuildGraphScriptConditions.INT.udn
@@ -43,11 +43,7 @@ BuildGraph script conditions consist of atoms and operators that evaluate to `tr
 
 Atoms can be numbers, strings, or identifiers that are coerced to their appropriate type for the oprator using them. 
 Atoms may be quoted with single (') or double (") quotes. They can also be an unquoted sequence of letters, digits, and 
-underscore characters. 
-
-### Literals
-
-All literals are considered identical, regardless of how they're declared. Additionally, literals are considered 
+underscore characters. All atoms are considered the same type, regardless of how they're declared. Additionally, atoms are considered 
 case-insensitive for comparisons, which means that the strings "True" and 'true' are identical to the identifier `true` 
 (despite the presence of quotes and differences in case).
 
@@ -71,7 +67,7 @@ A list of operators is specified below:
 | **x or y**              | True if either argument is `true`.                                          | 4               |
 
 [REGION:note]
-The `'<'` and `'>'` characters must be escaped as `"&lt;"` and `"&gt;"` in XML (referencing [XML Schema v1.0](https://www.w3.org/TR/NOTE-dcd])).
+The `'<'` and `'>'` characters must be escaped as `"&lt;"` and `"&gt;"` in XML.
 [/REGION]
 
 [/EXCERPT:BuildGraphScriptConditions]
\ No newline at end of file
diff --git a/Engine/Documentation/Source/Programming/Development/BuildGraph/ScriptAnatomy/Elements/BuildGraphScriptElements.INT.udn b/Engine/Documentation/Source/Programming/Development/BuildGraph/ScriptAnatomy/Elements/BuildGraphScriptElements.INT.udn
--- a/Engine/Documentation/Source/Programming/Development/BuildGraph/ScriptAnatomy/Elements/BuildGraphScriptElements.INT.udn
+++ b/Engine/Documentation/Source/Programming/Development/BuildGraph/ScriptAnatomy/Elements/BuildGraphScriptElements.INT.udn
@@ -32,10 +32,7 @@ tags:BuildGraph
 
 [EXCERPT:BuildGraphScriptElements]
 
-BuildGraph scripts have elements that are similar to the sort of elements that you'd expect to find in an XML document. 
-As you read through this section, you'll learn how to define a graph along with the syntax that goes into manipulating 
-properties. You'll also learn the syntax needed to control the flow of your logic and finally, we'll show you elements designed 
-to help you with script diagnostics.
+BuildGraph scripts are written in XML. This document covers the types of data found within BuildGraph scripts, as well as their overall structure.
 
 ## Elements
 
@@ -43,7 +40,7 @@ Elements describe the data that they contain, making them one of the basic build
 The following tables contain items that are provided as metadata, which are exported when running on the build system. Although 
 they're not used directly by BuildGraph when executing tasks locally, they're tagged with `[META]`.
 
-## Defining a Graph
+## Graph Structure
 
 A BuildGraph script is typically defined with the following elements:
 
@@ -233,4 +230,32 @@ the target being executed.
 | Message          | String          | Required     | Text to be printed to the log.                                                              |
 | If               | Condition       | Optional     | Condition to be evaluated. The element is ignored unless the condition evaluates to `true`. |
 
-[/EXCERPT:BuildGraphScriptElements]
\ No newline at end of file
+[/EXCERPT:BuildGraphScriptElements]
+
+[EXCERPT:BuildGraphScriptTasks]
+### Tasks
+
+Typically, project-specific utility scripts are implemented as AutomationTool BuildCommand classes. There are times, however, when you'll want 
+to automate custom tasks for your project. This is where BuildGraph can help. BuildGraph can be extended with any number of tasks, enabling you 
+to customize your build automation tasks in ways that suit your specific needs. The following section goes over ways of creating custom 
+tasks using BuildGraph, including a listing of predefined tasks that'll give you a great starting point for creating custom tasks.
+
+## Custom Tasks
+
+To create a new custom task, implement a class derived from the `CustomTask` class and apply the `[TaskElement]` attribute to it. 
+The `TaskElement` constructor takes two arguments; the name of the XML element that it's represented by, and the type of class 
+containing its parameters (which is passed to the constructor at load time).
+
+If you want to read parameter class fields from an XML file, attach the `[TaskParameter]` attribute to the task. Attaching the 
+`TaskParameter` attribute indicates if the parameter is required or optional, including additional validation that should be 
+applied to the argument.
+
+## Predefined Tasks
+
+If you'd like to have a starting point for creating tasks, we're providing a variety of predefined tasks as templates for you to work from. 
+If you're looking for a simple predefined task to start from, we recommend using `LogTask` as a good starting point. If you want to keep up with 
+our latest improvements to BuildGraph; bookmark this page, because the following list of predefined tasks will be updated periodically. 
+
+[INCLUDE:Programming/Development/BuildGraph/ScriptAnatomy/Tasks/PredefinedTasks]
+
+[/EXCERPT:BuildGraphScriptTasks]
\ No newline at end of file
diff --git a/Engine/Documentation/Source/Programming/Development/BuildGraph/ScriptAnatomy/Tasks/BuildGraphScriptTasks.INT.udn b/Engine/Documentation/Source/Programming/Development/BuildGraph/ScriptAnatomy/Tasks/BuildGraphScriptTasks.INT.udn
--- a/Engine/Documentation/Source/Programming/Development/BuildGraph/ScriptAnatomy/Tasks/BuildGraphScriptTasks.INT.udn
+++ b/Engine/Documentation/Source/Programming/Development/BuildGraph/ScriptAnatomy/Tasks/BuildGraphScriptTasks.INT.udn
@@ -32,7 +32,7 @@ tags:BuildGraph
 
 [EXCERPT:BuildGraphScriptTasks]
 
-Typically, project-specific utility scripts are implemented as AutomationTool BuildCommand classes. There are times; however, when you'll want 
+Typically, project-specific utility scripts are implemented as AutomationTool BuildCommand classes. There are times, however, when you'll want 
 to automate custom tasks for your project. This is where BuildGraph can help. BuildGraph can be extended with any number of tasks, enabling you 
 to customize your build automation tasks in ways that suit your specific needs. The following section goes over ways of creating custom 
 tasks using BuildGraph, including a listing of predefined tasks that'll give you a great starting point for creating custom tasks.
