Updating Merge Actors for 4.13 release

diff --git a/Engine/Documentation/Source/Engine/Actors/Merging/ActorMerging.INT.udn b/Engine/Documentation/Source/Engine/Actors/Merging/ActorMerging.INT.udn
--- a/Engine/Documentation/Source/Engine/Actors/Merging/ActorMerging.INT.udn
+++ b/Engine/Documentation/Source/Engine/Actors/Merging/ActorMerging.INT.udn
@@ -19,6 +19,7 @@ The Merge Actor feature is officially supported in 4.12 and later. In 4.8 to 4.1
 
 ## Actor Merging Worklow
 
+<!--
 [OBJECT:EmbeddedVideo]
 [PARAMLITERAL:videoid]
 _fTSGN8HGu8
@@ -30,6 +31,7 @@ _fTSGN8HGu8
 568
 [/PARAMLITERAL]
 [/OBJECT]
+-->
 
 Merging Static Meshes in your scene is handled like so:
 
@@ -52,13 +54,24 @@ If you have one or fewer meshes selected, you will see a yellow warning near the
 
 | Option | Description |
 |--------|-------------|
-|[REGION:tablesection]Mesh Settings[/REGION]|||
+|[REGION:tablesection]Mesh Components List[/REGION]|This area contains a checklist where you can chose which of the selected meshes will be included in the merge.|
+|[REGION:tablesection]Mesh Settings[/REGION]||
 |**Pivot Point at Zero**|If set to true, the pivot point of the final merged actor will be set to the scene's origin <0,0,0>. If set to false, the pivot will be located at the pivot location of the first Static Mesh Actor selected for the merge.
 |**Merge Physics Data**|If true, this will merge together all of the physics primitives into a single physics object.|
-|**Export Specific LOD**| Use this to force the selected meshes to a specific LOD prior to merging.|
+|**Bake Vertex Data**|Bakes vertex data including vertex color and normals into the new resulting mesh.|
+|**LOD Selection Type**| Use this to control which LODs are added to the merged object.|
+|[REGION:tablesection]Advanced Options[/REGION]|These options are found by expanding the small white triangle icon in the middle of the options list|
+|**Generate Lightmap UVs**|Creates a new lightmap UV layout for the merged mesh, using the existing UV islands. This new layout will be placed in the Target Light Map UVChannel.|
+|**Target Light Map UVChannel**|Specifies the UV channel into which the new lightmap UVs will be saved.|
+|**Target Light Map Resolution**|Specifies the resolution for the new lightmap UVs, which controls the distance calculated between each island within the layout. Higher resolutions will yield tighter spacing between each island.|
 |[REGION:tablesection]Material Settings[/REGION]||
 |**Merge Materials**|If checked, the Materials of all the merged assets will be combined into a single Material with new UV layout. The following subsettings apply only if Merge Materials is checked. |
 | |[INCLUDE:#materialsettings]|
+|[REGION:tablesection]Landscape Culling[/REGION]||
+|**Use Landscape Culling**|If checked, existing landscape geometry will be used to cull fully occluded triangles. For example, if you had a tree mesh that was planted such that the bottommost triangles were well below the surface of a landscape, then those occluded triangles would be removed upon merge.|
+|[REGION:tablesection]Standard Options[/REGION]||
+|**Replace Source Actors**|If checked, the Actors selected in the viewport will be removed and replaced with the new merged version. This will not afffect the the original objects in the Content Browser.| 
+
 
 
 <!--
@@ -75,7 +88,7 @@ If you have one or fewer meshes selected, you will see a yellow warning near the
 | **Specular Map** |If checked, a Specular map will be baked out for the resulting merged Actor. This nullifies the Specular Constant property.|
 | **Specular Constant**|Only enabled if Specular Map is unchecked, this sets a constant value for the Specular property of the resulting merged Actor's Material.|
 | **Emissive Map** |If checked, an Emissive map will be baked out for the resulting merged Actor.|
-|**Bake Vertex Data**|Bakes vertex data including vertex color and normals into the new resulting mesh.|
+|**Use Vertex Data for Baking Material**|Whether or not to use vertex data such as vertex color to bake Material result. Useful if you have a Material that is being controlled in some way by vertex color.|
 [/REGION]
 [/EXCERPT:MaterialSettings]
 -->
