#UE4 Docs: set page to publish

diff --git a/Engine/Documentation/Source/Platforms/Linux/BeginnerLinuxDeveloper/SettingUpAnUnrealWorkflow/1/SettingUpAnUnrealWorkflow_1.INT.udn b/Engine/Documentation/Source/Platforms/Linux/BeginnerLinuxDeveloper/SettingUpAnUnrealWorkflow/1/SettingUpAnUnrealWorkflow_1.INT.udn
--- a/Engine/Documentation/Source/Platforms/Linux/BeginnerLinuxDeveloper/SettingUpAnUnrealWorkflow/1/SettingUpAnUnrealWorkflow_1.INT.udn
+++ b/Engine/Documentation/Source/Platforms/Linux/BeginnerLinuxDeveloper/SettingUpAnUnrealWorkflow/1/SettingUpAnUnrealWorkflow_1.INT.udn
@@ -1,4 +1,4 @@
-Availability: Docs
+Availability: Public
 Title: 1 - Required Setup
 Crumbs:%ROOT%, Platforms, Platforms/Linux, Platforms/Linux/BeginnerLinuxDeveloper, Platforms/Linux/BeginnerLinuxDeveloper/SettingUpAnUnrealWorkflow
 Description:Learn how to set up Git on Linux.
diff --git a/Engine/Documentation/Source/Platforms/Linux/BeginnerLinuxDeveloper/SettingUpAnUnrealWorkflow/2/SettingUpAnUnrealWorkflow_2.INT.udn b/Engine/Documentation/Source/Platforms/Linux/BeginnerLinuxDeveloper/SettingUpAnUnrealWorkflow/2/SettingUpAnUnrealWorkflow_2.INT.udn
--- a/Engine/Documentation/Source/Platforms/Linux/BeginnerLinuxDeveloper/SettingUpAnUnrealWorkflow/2/SettingUpAnUnrealWorkflow_2.INT.udn
+++ b/Engine/Documentation/Source/Platforms/Linux/BeginnerLinuxDeveloper/SettingUpAnUnrealWorkflow/2/SettingUpAnUnrealWorkflow_2.INT.udn
@@ -1,4 +1,4 @@
-Availability: Docs
+Availability: Public
 Title: 2 - Downloading UE4 on Linux
 Crumbs:%ROOT%, Platforms, Platforms/Linux, Platforms/Linux/BeginnerLinuxDeveloper, Platforms/Linux/BeginnerLinuxDeveloper/SettingUpAnUnrealWorkflow
 Description:Learn how to fork and clone Unreal Engine source code.
diff --git a/Engine/Documentation/Source/Platforms/Linux/BeginnerLinuxDeveloper/SettingUpAnUnrealWorkflow/3/SettingUpAnUnrealWorkflow_3.INT.udn b/Engine/Documentation/Source/Platforms/Linux/BeginnerLinuxDeveloper/SettingUpAnUnrealWorkflow/3/SettingUpAnUnrealWorkflow_3.INT.udn
--- a/Engine/Documentation/Source/Platforms/Linux/BeginnerLinuxDeveloper/SettingUpAnUnrealWorkflow/3/SettingUpAnUnrealWorkflow_3.INT.udn
+++ b/Engine/Documentation/Source/Platforms/Linux/BeginnerLinuxDeveloper/SettingUpAnUnrealWorkflow/3/SettingUpAnUnrealWorkflow_3.INT.udn
@@ -1,4 +1,4 @@
-Availability: Docs
+Availability: Public
 Title: 3 - Building UE4 on Linux
 Crumbs:%ROOT%, Platforms, Platforms/Linux, Platforms/Linux/BeginnerLinuxDeveloper, Platforms/Linux/BeginnerLinuxDeveloper/SettingUpAnUnrealWorkflow
 Description:Learn how to build Unreal Engine from source.
diff --git a/Engine/Documentation/Source/Platforms/Linux/BeginnerLinuxDeveloper/SettingUpAnUnrealWorkflow/4/SettingUpAnUnrealWorkflow_4.INT.udn b/Engine/Documentation/Source/Platforms/Linux/BeginnerLinuxDeveloper/SettingUpAnUnrealWorkflow/4/SettingUpAnUnrealWorkflow_4.INT.udn
--- a/Engine/Documentation/Source/Platforms/Linux/BeginnerLinuxDeveloper/SettingUpAnUnrealWorkflow/4/SettingUpAnUnrealWorkflow_4.INT.udn
+++ b/Engine/Documentation/Source/Platforms/Linux/BeginnerLinuxDeveloper/SettingUpAnUnrealWorkflow/4/SettingUpAnUnrealWorkflow_4.INT.udn
@@ -1,4 +1,4 @@
-Availability: Docs
+Availability: Public
 Title: 4 - Running UE4 on Linux
 Crumbs:%ROOT%, Platforms, Platforms/Linux, Platforms/Linux/BeginnerLinuxDeveloper, Platforms/Linux/BeginnerLinuxDeveloper/SettingUpAnUnrealWorkflow
 Description:Learn how to run Unreal Engine on Linux.
diff --git a/Engine/Documentation/Source/Platforms/Linux/BeginnerLinuxDeveloper/SettingUpAnUnrealWorkflow/5/SettingUpAnUnrealWorkflow_5.INT.udn b/Engine/Documentation/Source/Platforms/Linux/BeginnerLinuxDeveloper/SettingUpAnUnrealWorkflow/5/SettingUpAnUnrealWorkflow_5.INT.udn
--- a/Engine/Documentation/Source/Platforms/Linux/BeginnerLinuxDeveloper/SettingUpAnUnrealWorkflow/5/SettingUpAnUnrealWorkflow_5.INT.udn
+++ b/Engine/Documentation/Source/Platforms/Linux/BeginnerLinuxDeveloper/SettingUpAnUnrealWorkflow/5/SettingUpAnUnrealWorkflow_5.INT.udn
@@ -1,4 +1,4 @@
-Availability: Docs
+Availability: Public
 Title: 5 - On Your Own!
 Crumbs:%ROOT%, Platforms, Platforms/Linux, Platforms/Linux/BeginnerLinuxDeveloper, Platforms/Linux/BeginnerLinuxDeveloper/SettingUpAnUnrealWorkflow
 Description:Do some exercises on your own.
@@ -7,7 +7,7 @@ type:QuickStart
 parent:Platforms/Linux/BeginnerLinuxDeveloper/SettingUpAnUnrealWorkflow
 order: 5
 tags:Linux
-SkillLevel:Beginner
+SkillLevel:Beginner 
 
 [VAR:Steps]
 [OBJECT:Navigation]
diff --git a/Engine/Documentation/Source/Platforms/Linux/BeginnerLinuxDeveloper/SettingUpAnUnrealWorkflow/SettingUpAnUnrealWorkflow.INT.udn b/Engine/Documentation/Source/Platforms/Linux/BeginnerLinuxDeveloper/SettingUpAnUnrealWorkflow/SettingUpAnUnrealWorkflow.INT.udn
--- a/Engine/Documentation/Source/Platforms/Linux/BeginnerLinuxDeveloper/SettingUpAnUnrealWorkflow/SettingUpAnUnrealWorkflow.INT.udn
+++ b/Engine/Documentation/Source/Platforms/Linux/BeginnerLinuxDeveloper/SettingUpAnUnrealWorkflow/SettingUpAnUnrealWorkflow.INT.udn
@@ -1,4 +1,4 @@
-Availability: Docs
+Availability: Public
 Title: Linux Quick Start Guide
 Crumbs:%ROOT%, Platforms, Platforms/Linux, Platforms/Linux/BeginnerLinuxDeveloper
 Description:Learn how to download, build, and run UE4 on Linux.
diff --git a/Engine/Documentation/Source/Platforms/Linux/Linux.INT.udn b/Engine/Documentation/Source/Platforms/Linux/Linux.INT.udn
--- a/Engine/Documentation/Source/Platforms/Linux/Linux.INT.udn
+++ b/Engine/Documentation/Source/Platforms/Linux/Linux.INT.udn
@@ -17,12 +17,8 @@ you'll want to work your way through the Quick Start Guide before moving onto mo
 
 ## Getting Started
 
-[REGION:topics third]
 %Platforms/Linux/BeginnerLinuxDeveloper/SettingUpAnUnrealWorkflow:topic%
-[/REGION]
 
 ## Reference
 
-[REGION:topics third]
 %Platforms/Linux/BeginnerLinuxDeveloper/SupportedGfxDrivers:topic%
-[/REGION]
\ No newline at end of file
