#UE4 Docs: updated metadata

diff --git a/Engine/Documentation/Source/Engine/UMG/HowTo/VirtualKeyboards/CreatingVirtualKeyboards.INT.udn b/Engine/Documentation/Source/Engine/UMG/HowTo/VirtualKeyboards/CreatingVirtualKeyboards.INT.udn
--- a/Engine/Documentation/Source/Engine/UMG/HowTo/VirtualKeyboards/CreatingVirtualKeyboards.INT.udn
+++ b/Engine/Documentation/Source/Engine/UMG/HowTo/VirtualKeyboards/CreatingVirtualKeyboards.INT.udn
@@ -4,8 +4,8 @@ Crumbs: %ROOT%, Engine, Engine/UMG, Engine/UMG/HowTo
 Description:Demonstrates how you can create and interaction with a 3D Virtual Keyboard with UMG. 
 Version: 4.13
 SkillLevel: Advanced
-parent:Engine/Editor
-order:1
+parent:Engine/UMG/HowTo
+order:4
 checkpoint: editorqs
 
 If you are looking to create a 3D keypad or virtual keyboard that allows a player to enter specified characters, you can do so with a combination of the [Widget Interaction](Engine\UMG\UserGuide\WidgetInteraction) component and [Widget](Engine/Components/Widget) components. 
