[NO NON-PII COMMIT MESSAGE]

diff --git a/Engine/Documentation/Source/Engine/Sequencer/Workflow/ImportExport/ImportExportFBX.INT.udn b/Engine/Documentation/Source/Engine/Sequencer/Workflow/ImportExport/ImportExportFBX.INT.udn
new file mode 100644
--- /dev/null
+++ b/Engine/Documentation/Source/Engine/Sequencer/Workflow/ImportExport/ImportExportFBX.INT.udn
@@ -0,0 +1,158 @@
+Availability:Docs
+Title: Exporting and Importing FBX files
+Crumbs: %ROOT%, Engine
+Description:Describes how you can export FBX files from Sequencer and import FBX files to Sequencer. 
+Navigation:topic
+parent:Engine/Sequencer/Workflow
+order:1
+Version: 4.13
+Related: Engine\Sequencer\Overview
+Related: Engine/Sequencer/HowTo
+
+This page describes how you can export a scene from Sequencer as an FBX file, make alterations to that file inside a 3rd party application like 3dsMax or Maya, then import that FBX file (along with your changes) back into Unreal Engine into Sequencer as part of your existing scene. 
+
+There are some caveats to this process to be aware of such as:
+
+* When exporting, you can export all the objects in your scene with animations to your FBX file. 
+* Importing FBX files will only import animation (it will not create new objects). 
+* Export/Import does not work with **Master Sequences**, **Shots** within a Master, or **Subscenes**. 
+* Make sure that your time settings in your 3rd party application matches the time frame used for your scene in Unreal Engine. 
+
+As long as you export from a source Level Sequence (not exporting from a shot that sits inside a Master) and are only altering the animation of existing assets, you will be able to import your changes directly back into your scene inside Unreal Engine. 
+Also, using the same time frames in both Unreal Engine and your 3rd party application will ensure that they line up and that the imported scene is not resampled at a different time rate. 
+
+## Example Workflow
+
+Below we export a scene from the **Sequencer** Showcase (pictured below) available for free from the **Learn** tab of the Launcher, take it into Maya and alter the camera movement in the scene, then import that FBX along with our changes back into Unreal Engine where our existing scene is automatically updated to reflect our changes. 
+
+![](SequencerSample.png)
+
+
+### Exporting from Sequencer
+
+First thing you will need to do is to determine which Level Sequence to export. 
+
+1.	In the clip below we pick the shot we want to alter and open it up. We then go into the **General Options** and choose **Export**. 
+
+	[OBJECT:EmbeddedVideo]
+	[PARAMLITERAL:width]
+	640
+	[/PARAMLITERAL]
+	[PARAMLITERAL:height]
+	360
+	[/PARAMLITERAL]
+	[PARAMLITERAL:videoid]
+	SM77fHV_wu8
+	[/PARAMLITERAL]
+	[/OBJECT]
+
+	We have a Master Track that includes **shot0040_001** so we open that particular Level Sequence to export it. 
+	
+### External Editing
+
+Once we have an exported FBX from Sequencer, we open our 3rd party application to import it so we can start editing it. 
+
+1.	Before we import our FBX into Maya, we set our working units for **Time** to match our **Time Rate** used in Unreal Engine. 
+
+	[OBJECT:EmbeddedVideo]
+	[PARAMLITERAL:width]
+	640
+	[/PARAMLITERAL]
+	[PARAMLITERAL:height]
+	360
+	[/PARAMLITERAL]
+	[PARAMLITERAL:videoid]
+	cwldMhWAqBY
+	[/PARAMLITERAL]
+	[/OBJECT]
+
+1.	We then import our FBX into Maya and set up our scene so that we can begin making edits. 
+
+	[OBJECT:EmbeddedVideo]
+	[PARAMLITERAL:width]
+	640
+	[/PARAMLITERAL]
+	[PARAMLITERAL:height]
+	360
+	[/PARAMLITERAL]
+	[PARAMLITERAL:videoid]
+	0Nc4S-T1t1I
+	[/PARAMLITERAL]
+	[/OBJECT]
+
+	Once our scene is imported, we set our Perspective under **Panel** to use the camera **shotCamA** in our scene and hide any elements obstructing the view (such as the godray mesh) with it selected and **Ctrl + H**. 
+	Lastly under **View** we **Select Camera** to display the properties and keyframes for our **shotCamA** which we can then begin to edit. 
+
+1.	Next we make edits to the movement of our camera. 
+
+	[OBJECT:EmbeddedVideo]
+	[PARAMLITERAL:width]
+	640
+	[/PARAMLITERAL]
+	[PARAMLITERAL:height]
+	360
+	[/PARAMLITERAL]
+	[PARAMLITERAL:videoid]
+	3asklRbBagA
+	[/PARAMLITERAL]
+	[/OBJECT]
+
+	Above we select the Translation and Rotation properties of our **shotCamA** and **Right-click** and **Break Connections**. 
+	We then pick a new starting point for the camera to be in to start our scene and press **S** to keyframe the new position. 
+	We scrub ahead slightly and move the camera to the ending position and press **S** again to keyframe the end position so that our camera moves along a new path. 
+
+1.	With our changes in place, we then **Export** our scene using **Export Selection**.
+
+	[OBJECT:EmbeddedVideo]
+	[PARAMLITERAL:width]
+	640
+	[/PARAMLITERAL]
+	[PARAMLITERAL:height]
+	360
+	[/PARAMLITERAL]
+	[PARAMLITERAL:videoid]
+	sj-gWprQFNk
+	[/PARAMLITERAL]
+	[/OBJECT]
+
+	During the export process from Maya, we first change the export type to **FBX export**.
+	Then, under **Advanced Options** and **Axis Conversion**, we change the **Up Axis** to follow **Z** to allow for a compatible import into Unreal Engine. 
+	Future versions of the engine may address this so that you do not have to convert the axis, but you will need to do this for now to avoid any importing problems. 
+
+
+### Importing to Sequencer
+
+Now that we have an edited version in FBX format we can import it and the changes into Unreal Engine, updating our scene.
+
+1.	Inside our **shot0040_001** that we wanted to edit, we **Right-click** on our **shotCamA** and select **Import**. 
+
+	[OBJECT:EmbeddedVideo]
+	[PARAMLITERAL:width]
+	640
+	[/PARAMLITERAL]
+	[PARAMLITERAL:height]
+	360
+	[/PARAMLITERAL]
+	[PARAMLITERAL:videoid]
+	HvNrIL9_GZw
+	[/PARAMLITERAL]
+	[/OBJECT]
+
+	By **Right-clicking** on the camera it will take the Transform data inside the FBX and apply it to the selected object in Unreal Engine. 
+	This can be useful if you want to apply the same Transform data to multiple objects inside your scene. You can also import the FBX from the **General Options** window instead of using the **Right-click** method. 
+
+	![](ImportOption.png)
+
+	When choosing this method, Unreal Engine will try to match the names of your objects (for example, we have **shotCamA** in both our FBX file and our Sequence so it will apply the data from the FBX file to that object inside Unreal Engine). 
+
+In our example we only applied Transform changes to our camera however we could have keyframed new values for our **Focal Length** to produce a different looking shot as well.
+You can experiment with this workflow and use what works best in your given situation. There may be times when you want to export to a 3D package like Maya or 3dsMax to fine tune aspects of your scene and then import your changes back into your scene in Unreal Engine. 
+	
+
+
+
+
+
+
+
+
diff --git a/Engine/Documentation/Source/Engine/Sequencer/Workflow/Workflow.INT.udn b/Engine/Documentation/Source/Engine/Sequencer/Workflow/Workflow.INT.udn
--- a/Engine/Documentation/Source/Engine/Sequencer/Workflow/Workflow.INT.udn
+++ b/Engine/Documentation/Source/Engine/Sequencer/Workflow/Workflow.INT.udn
@@ -37,6 +37,7 @@ If you are new to Sequencer it is recommended that you first check out the [**Ov
 	[/PARAM]
 	[PARAM:links]
 		* [](Engine\Sequencer\Workflow\CheatSheet "%Engine\Sequencer\Workflow\CheatSheet:description%")
+		* [](Engine\Sequencer\Workflow\ImportExport "%Engine\Sequencer\Workflow\ImportExport:description%")
 		* [](Engine\Sequencer\Workflow\ToolUsage "%Engine\Sequencer\Workflow\ToolUsage:description%")
 	[/PARAM]
 [/OBJECT]
