Creating new Formatting Localized Text page, moving related info from FText page, adding related page links. UEDOC-3287

diff --git a/Engine/Documentation/Source/Gameplay/Localization/Commandlets/Commandlets.INT.udn b/Engine/Documentation/Source/Gameplay/Localization/Commandlets/Commandlets.INT.udn
--- a/Engine/Documentation/Source/Gameplay/Localization/Commandlets/Commandlets.INT.udn
+++ b/Engine/Documentation/Source/Gameplay/Localization/Commandlets/Commandlets.INT.udn
@@ -3,6 +3,8 @@ Title: Commandlets
 Crumbs:%ROOT%, Gameplay, Gameplay/Localization
 Description:Descriptions and arguments for the commandlets used when localizing your project. 
 version: 4.9
+Parent:Gameplay/Localization
+Order:4
 
 
 [TOC (start:2)]
diff --git a/Engine/Documentation/Source/Gameplay/Localization/Formatting/Formatting.INT.udn b/Engine/Documentation/Source/Gameplay/Localization/Formatting/Formatting.INT.udn
new file mode 100644
--- /dev/null
+++ b/Engine/Documentation/Source/Gameplay/Localization/Formatting/Formatting.INT.udn
@@ -0,0 +1,91 @@
+﻿Availability:Docs
+Title:Formatting User-Facing Text
+Crumbs:%ROOT%, Gameplay, Gameplay/Localization
+Description:Using the Format function to present localized text to your users
+Parent:Gameplay/Localization
+Order:3
+version: 4.13
+
+FText is used to create user-facing text, which is often localized. This means that user interface text that changes dynamically over the course of gameplay can also be subject to grammar rules from a number of different languages. Using `FText::Format` and the corresponding 
+**Format Text** Blueprint node enables you to construct reusable format patterns to make the localization process easier and less prone to edge-case errors. Both functions support ordered arguments, named arguments, and specialized arguments based on plural forms, gender forms, 
+and Hangul post-positions.
+
+FText::Format is also available in Blueprints. The Format text can be a literal or connected to a pin. When used as a literal, it will auto-generate arguments for formatting. So by entering the literal "There are {number} days until {day}.", a Number pin and a Day pin will be 
+automatically added to the node. These pins can be filled in as a literal or from a variable and the output Text will replace the arguments with whatever values are hooked to the pins.
+
+When the Format pin is connected, the Format can not be deciphered because it is variable. In this case, you can add pins yourself and name them after the possible arguments. In this way, you can set up the Format Text node to be ready for multiple possibilities. Whenever an argument 
+matches an argument in the Format variable, it will utilize the provided information.
+
+Pin names can be changed, added, and removed from the Details Panel.
+
+[REGION:tip]
+For users using a source build of Unreal Engine, `TextTest.cpp` can show you examples of how to perform valid culture aware formatting, comparison, and sorting.
+[/REGION]
+
+## Ordered Arguments
+ 
+ When building a format argument with ordered arguments, curly braces denote the beginning and end of the format argument and the number indicates the parameter passed.
+ Arguments start at {0} and proceed sequentially ( {0}, {1}, {2}, etc.). The `Format` function and **Format Text** node slot variables into the format argument in the order 
+ they are provided. 
+
+ The following C++ and Blueprint snippets show an ordered format argument to display the player's current health and maximum health.
+
+### C++ Example
+
+	FText::Format( LOCTEXT( "HealthFormatStr",  "{0} / {1}" ), FText::AsNumber( CurrentHealth ), FText::AsNumber( MaxHealth )  ) 
+		
+### Blueprint Example
+
+![](OrderedArgs.png)
+
+## Named Arguments 
+
+When building a format argument with named arguments, curly braces denote the beginning and end of the format argument and the text in the braces is the name of the argument to be found in the arguments passed in.
+Common argument examples are {Name}, {ID}, {Health}, etc.
+
+The following C++ and Blueprint snippets show a named format argument to display the player's current health.
+
+### C++ Example
+
+	FFormatNamedArguments Arguments;
+	Arguments.Add(TEXT("CurrentHealth"), FText::AsNumber( CurrentHealth ));
+	EditorErrors.Error(FText::Format(LOCTEXT("ExampleFText", "You currently have {CurrentHealth} health left."), Arguments ) );
+
+### Blueprint Example
+			
+![](CurrentHealth.png)
+
+## Specialized Arguments
+
+To allow you to pass in the numeric/gender values needed for plural/gender form support, all of the FText::Format(...) family of functions now take their values as FFormatArgumentValue rather than FText. This can be implicitly constructed from any numeric type, ETextGender, or FText.
+
+The ability to set these value types in Blueprints has been exposed using wildcard pins on the **Format Text** node.
+
+Following from the examples above, these format patterns can be supplied as arguments to `FText::Format` or to **Format Text**. 
+
+### Plural Forms
+
+Plural forms allow you to use different text based upon a numeric variable given to your text format. Plural forms may be cardinal, such as "There is 1 cat" or "There are 4 cats", or ordinal, such as "You came 1st!" or "You came 2nd!".
+
+Plural forms are specified as key->value pairs, and support any of the following keywords (as defined for your culture by the [CLDR data](http://www.unicode.org/cldr/charts/29/supplemental/language_plural_rules.html)): zero, one, two, few, many, other. Values are an optionally quoted string 
+that may also contain format markers.
+
+Cardinal Format Example: `"There {NumCats}|plural(one=is,other=are) {NumCats} {NumCats}|plural(one=cat,other=cats)"`
+
+Ordinal Format Example: `"You came {Place}{Place}|ordinal(one=st,two=nd,few=rd,other=th)!"`
+
+
+### Gender Forms
+
+Gender forms allow you to use different text based upon an ETextGender value given to your text format, such as "Le guerrier est fort" or "La guerrière est forte".
+Gender forms are specified as a list of values in the order of [masculine, feminine, neuter], where neuter is optional. Values are an optionally quoted string that may also contain format markers.
+
+Format Example: `"{Gender}|gender(Le,La) {Gender}|gender(guerrier,guerrière) est {Gender}|gender(fort,forte)"`
+
+### Hangul Post-Positions
+
+Hangul post-positions help you deal with the grammar rules present in Korean, and will insert the correct glyph(s) based upon whether the value being inserted ends in a consonant or a verb, such as "사람은" or "사자는".
+Hangul post-positions are specified as a list of values in the order of [consonant, verb]. Values are an optionally quoted string.
+
+Format Example: `"{Arg}|hpp(은,는)"`
+
diff --git a/Engine/Documentation/Source/Gameplay/Localization/Localization.INT.udn b/Engine/Documentation/Source/Gameplay/Localization/Localization.INT.udn
--- a/Engine/Documentation/Source/Gameplay/Localization/Localization.INT.udn
+++ b/Engine/Documentation/Source/Gameplay/Localization/Localization.INT.udn
@@ -5,31 +5,13 @@ Description:Information about how to localize your project.
 Related: Programming/UnrealArchitecture/StringHandling/FText
 version: 4.9
 
-[OBJECT:TopicButtonList]
-	[PARAM:icon]
-		![](%ROOT%/font_icon.png)(convert:false)
-	[/PARAM]
-	[PARAM:title]
-	Localizing Text
-	[/PARAM]
-	[PARAM:links]
-	* [Overview](Gameplay/Localization/Overview "%Gameplay/Localization/Overview:description%")
-	* [Setup](Gameplay/Localization/Setup "%Gameplay/Localization/Setup:description%")
-	* [Commandlets](Gameplay/Localization/Commandlets "%Gameplay/Localization/Commandlets:description%")
-	[/PARAM]
-[/OBJECT]
 
-[OBJECT:TopicButtonList]
-	[PARAM:icon]
-		![](%ROOT%/cb_icon.png)(convert:false)
-	[/PARAM]
-	[PARAM:title]
-	Localizing Assets
-	[/PARAM]
-	[PARAM:links]
-	[REGION:tip]
-	Asset localization will be coming in the future just before proper audio localization support. 
-	The implementation will revolve around swapping out the entire uasset package with locale specific versions. 
-	[/REGION:tip]
-	[/PARAM]
-[/OBJECT]
+## Localizing Text
+[DIR(output:"listbutton" parent:"Gameplay/Localization")]
+
+
+
+[REGION:tip]
+Asset localization will be coming in the future just before proper audio localization support. 
+The implementation will revolve around swapping out the entire uasset package with locale specific versions. 
+[/REGION:tip]
diff --git a/Engine/Documentation/Source/Gameplay/Localization/Overview/Overview.INT.udn b/Engine/Documentation/Source/Gameplay/Localization/Overview/Overview.INT.udn
--- a/Engine/Documentation/Source/Gameplay/Localization/Overview/Overview.INT.udn
+++ b/Engine/Documentation/Source/Gameplay/Localization/Overview/Overview.INT.udn
@@ -3,6 +3,8 @@ Title:Localization System Overview
 Crumbs:%ROOT%, Gameplay, Gameplay/Localization
 Description:An overview of the elements involved in localizing your project. 
 version: 4.9
+Parent:Gameplay/Localization
+Order:1
 
 Localization of a project involves a few key concepts:
 
diff --git a/Engine/Documentation/Source/Gameplay/Localization/Setup/Setup.INT.udn b/Engine/Documentation/Source/Gameplay/Localization/Setup/Setup.INT.udn
--- a/Engine/Documentation/Source/Gameplay/Localization/Setup/Setup.INT.udn
+++ b/Engine/Documentation/Source/Gameplay/Localization/Setup/Setup.INT.udn
@@ -3,6 +3,8 @@ Title: Setup
 Crumbs:%ROOT%, Gameplay, Gameplay/Localization
 Description:How to set up a project for localization. 
 version: 4.9
+Parent:Gameplay/Localization
+Order:2
 
 [TOC (start:2)]
 
diff --git a/Engine/Documentation/Source/Programming/UnrealArchitecture/StringHandling/FText/FText.INT.udn b/Engine/Documentation/Source/Programming/UnrealArchitecture/StringHandling/FText/FText.INT.udn
--- a/Engine/Documentation/Source/Programming/UnrealArchitecture/StringHandling/FText/FText.INT.udn
+++ b/Engine/Documentation/Source/Programming/UnrealArchitecture/StringHandling/FText/FText.INT.udn
@@ -3,6 +3,8 @@ Title:FText
 Crumbs:%ROOT%, Programming, Programming/UnrealArchitecture/Reference, Programming/UnrealArchitecture/Reference/Properties, Programming/UnrealArchitecture/StringHandling
 Description:
 Version: 4.9
+Related:Gameplay/Localization
+Related:Gameplay/Localization/Formatting
 
 
 
@@ -100,36 +102,7 @@ You can compare two FText strings by using **FText::EqualTo()** or **FText::Equa
 FText does not expose operators for these comparison operations currently because their operations are not simple ordinal comparisons like FStrings. We opted to not use the normal operators to hint 
 at the operation's expense. 
 
-## Formatting FText
-  
-You can use the **FText::Format()** function to concatenate different FText together, and if desired, save it to an FText variable. 
-  
-	FText::Format( LOCTEXT( "HelloMessageFormat",  "Hello {0}," ), PlayerName ) 
-  
-Currently, the Format function only takes FText as parameters.
-
-FText::Format supports two different kinds of format arguments, ordered and named.
-
-### Ordered Arguments
- 
-* Arguments = {0}, {1}, {2}, etc... 
-* Curly braces denote the beginning and end of the format argument and the number indicates the parameter passed.
-
-		FText::Format( LOCTEXT( "HealthFormatStr",  "{0} / {1}" ), FText::AsNumber( CurrentHealth ), FText::AsNumber( MaxHealth )  ) 
 
-### Named Arguments 
-
-* Arguments = {Name}, {ID}, {Health}, etc... 
-* Curly braces denote the beginning and end of the format argument and the text between the braces is the name of argument to be found in the FFormatNamedArgs collection passed in.
-
-		FFormatNamedArguments Arguments;
-		Arguments.Add(TEXT("CurrentHealth"), FText::AsNumber( CurrentHealth ));
-		EditorErrors.Error(FText::Format(LOCTEXT("ExampleFText", "You currently have {CurrentHealth} health left."), Arguments ) );
-			
-
-[PUBLISH:Licensee]
-TextTest.cpp can show you examples of how to perform valid culture aware formatting, comparison, and sorting.
-[/PUBLISH:Licensee]
 
 ## FText in HUDs
 
@@ -158,13 +131,4 @@ You must call the DrawItem() function within your HUD class's **DrawHUD()** func
 
 Blueprints support Text literals. Any input pins for Text values will support being assigned a value in the node and these strings will be gathered up for localization.
 
-**FText::Format** is also available in Blueprints. The Format text can be a literal or connected to a pin. When used as a literal, it will auto-generate arguments for formatting. So by entering the 
-literal "There are {number} days until {day}.", a **Number** pin and a **Day** pin will be automatically added to the node. These pins can be filled in as a literal or from a variable and the output 
-Text will replace the arguments with whatever values are hooked to the pins. 
-
-When the Format pin is connected, the Format can not be deciphered because it is variable. In this case, you can add pins yourself and name them after the possible arguments. In this way, you can 
-set up the Format Text node to be ready for multiple possibilities. Whenever an argument matches an argument in the Format variable, it will utilize the provided information.
-
-Pin names can be changed, added, and removed from the **Details** Panel.
 
- 
