#UE4 Docs: Making first check in of the Render Targets docs.

diff --git a/Engine/Documentation/Source/Engine/Rendering/RenderTargets/BlueprintRenderTargets/BestPractices/BlueprintRenderTargetsBestPractices.INT.udn b/Engine/Documentation/Source/Engine/Rendering/RenderTargets/BlueprintRenderTargets/BestPractices/BlueprintRenderTargetsBestPractices.INT.udn
new file mode 100644
--- /dev/null
+++ b/Engine/Documentation/Source/Engine/Rendering/RenderTargets/BlueprintRenderTargets/BestPractices/BlueprintRenderTargetsBestPractices.INT.udn
@@ -0,0 +1,49 @@
+availability:docs
+title:Blueprint & Render Target Best Practices
+description: A collection of Best Practices to follow when using Blueprints and Render Targets.
+crumbs:
+version:4.13
+Type:topic
+SkillLevel:Intermediate
+parent:Engine
+order:2
+Navigation:Reference
+
+
+
+
+
+
+
+
+
+
+
+
+
+
+
+
+
+
+
+
+
+
+
+
+
+
+
+
+
+
+
+
+
+
+
+
+
+
+
diff --git a/Engine/Documentation/Source/Engine/Rendering/RenderTargets/BlueprintRenderTargets/BlueprintRenderTargets.INT.udn b/Engine/Documentation/Source/Engine/Rendering/RenderTargets/BlueprintRenderTargets/BlueprintRenderTargets.INT.udn
new file mode 100644
--- /dev/null
+++ b/Engine/Documentation/Source/Engine/Rendering/RenderTargets/BlueprintRenderTargets/BlueprintRenderTargets.INT.udn
@@ -0,0 +1,112 @@
+availability:docs
+title: Blueprint & Render Targets
+description: Using Render Targets with Blueprints in Unreal Engine 4
+crumbs:
+version:4.13
+Type:landing
+SkillLevel:Intermediate
+parent:Engine
+order:1
+Navigation:topic
+
+![](Engine\Rendering\RenderTargets\BlueprintRenderTarget.png)
+
+Having the ability to use Blueprints to draw Materials into Render Targets and then use or save those results for later use can enable a huge variety of game-specific rendering effects. On the following page you will find several links to documentation pages that cover how to use Blueprints and Render Targets. 
+
+
+[REGION:topics third]
+[OBJECT:Topic]
+	[PARAM:image]
+		![Engine\Rendering\RenderTargets\BlueprintRenderTargets/Overview:title%](Engine/Rendering/RenderTargets/Render_Targets_Overview.png)
+	[/PARAM]
+	[PARAM:icon]
+		![](%ROOT%/blueprint_icon.png)(convert:false)
+	[/PARAM]
+	[PARAM:title]
+		%Engine/Rendering/RenderTargets/BlueprintRenderTargets/Overview:title%
+	[/PARAM]
+	[PARAM:description]
+		%Engine/Rendering/RenderTargets/BlueprintRenderTargets/Overview:description%
+	[/PARAM]
+	[PARAM:path]
+		[RELATIVE:Engine/Rendering/RenderTargets/BlueprintRenderTargets/Overview]
+	[/PARAM]
+[/OBJECT]
+
+[OBJECT:Topic]
+	[PARAM:image]
+		![Engine\Rendering\RenderTargets\BlueprintRenderTargets/Reference:title%](Engine/Rendering/RenderTargets/Render_Targets_Reference.png)
+	[/PARAM]
+	[PARAM:icon]
+		![](%ROOT%/blueprint_icon.png)(convert:false)
+	[/PARAM]
+	[PARAM:title]
+		%Engine/Rendering/RenderTargets/BlueprintRenderTargets/Reference:title%
+	[/PARAM]
+	[PARAM:description]
+		%Engine/Rendering/RenderTargets/BlueprintRenderTargets/Reference:description%
+	[/PARAM]
+	[PARAM:path]
+		[RELATIVE:Engine/Rendering/RenderTargets/BlueprintRenderTargets/Reference]
+	[/PARAM]
+[/OBJECT]
+
+[OBJECT:Topic]
+	[PARAM:image]
+		![Engine\Rendering\RenderTargets\BlueprintRenderTargets/BestPractices:title%](Engine/Rendering/RenderTargets/Render_Targets_Best_Practies.png)
+	[/PARAM]
+	[PARAM:icon]
+		![](%ROOT%/blueprint_icon.png)(convert:false)
+	[/PARAM]
+	[PARAM:title]
+		%Engine/Rendering/RenderTargets/BlueprintRenderTargets/BestPractices:title%
+	[/PARAM]
+	[PARAM:description]
+		%Engine/Rendering/RenderTargets/BlueprintRenderTargets/BestPractices:description%
+	[/PARAM]
+	[PARAM:path]
+		[RELATIVE:Engine/Rendering/RenderTargets/BlueprintRenderTargets/BestPractices]
+	[/PARAM]
+[/OBJECT]
+
+[OBJECT:Topic]
+	[PARAM:image]
+		![Engine\Rendering\RenderTargets\BlueprintRenderTargets/BestPractices:title%](Engine/Rendering/RenderTargets/Render_Targets_Best_Practies.png)
+	[/PARAM]
+	[PARAM:icon]
+		![](%ROOT%/blueprint_icon.png)(convert:false)
+	[/PARAM]
+	[PARAM:title]
+		%Engine/Rendering/RenderTargets/BlueprintRenderTargets/HowTo:title%
+	[/PARAM]
+	[PARAM:description]
+		%Engine/Rendering/RenderTargets/BlueprintRenderTargets/HowTo:description%
+	[/PARAM]
+	[PARAM:path]
+		[RELATIVE:Engine/Rendering/RenderTargets/BlueprintRenderTargets/HowTo]
+	[/PARAM]
+[/OBJECT]
+
+[/REGION]
+
+
+
+
+
+
+
+
+
+
+
+
+
+
+
+
+
+
+
+
+
+
diff --git a/Engine/Documentation/Source/Engine/Rendering/RenderTargets/BlueprintRenderTargets/HowTo/BlueprintRenderTargetsHowTos.INT.udn b/Engine/Documentation/Source/Engine/Rendering/RenderTargets/BlueprintRenderTargets/HowTo/BlueprintRenderTargetsHowTos.INT.udn
new file mode 100644
--- /dev/null
+++ b/Engine/Documentation/Source/Engine/Rendering/RenderTargets/BlueprintRenderTargets/HowTo/BlueprintRenderTargetsHowTos.INT.udn
@@ -0,0 +1,56 @@
+availability:docs
+title: Blueprint and Render Targets - How To's
+description: A How To Guide for using Blueprints to manipulate render targets. 
+crumbs:
+version:4.13
+Type:topic
+SkillLevel:Intermediate
+parent:Engine
+order:2
+Navigation:topic
+
+The **Blueprint and Render Targets - How To…** pages provide step by step examples on how to use Blueprints to work with render targets. Whether you want to use Blueprints and render targets to create Textures for use in your projects Materials, or you want to create a surface that deforms when the player shoots at it, these pages will guide you on how to accomplish this. Reference links are also included to provi
+
+## Intermediate
+[REGION:wrapper transparent]
+[](FANCY:Engine\Rendering\RenderTargets\BlueprintRenderTargets\HowTo\CreatingTextures "%Engine\Rendering\RenderTargets\BlueprintRenderTargets\HowTo\CreatingTextures:description%")
+
+[/REGION]
+
+
+## Advanced
+[REGION:wrapper transparent]
+[](FANCY:Engine\Rendering\RenderTargets\BlueprintRenderTargets\HowTo\HeightFieldPainter "%Engine\Rendering\RenderTargets\BlueprintRenderTargets\HowTo\HeightFieldPainter:description%")
+[/REGION]
+
+
+
+
+
+
+
+
+
+
+
+
+
+
+
+
+
+
+
+
+
+
+
+
+
+
+
+
+
+
+
+
diff --git a/Engine/Documentation/Source/Engine/Rendering/RenderTargets/BlueprintRenderTargets/HowTo/CreatingTextures/CreatingTexturesHowTo.INT.udn b/Engine/Documentation/Source/Engine/Rendering/RenderTargets/BlueprintRenderTargets/HowTo/CreatingTextures/CreatingTexturesHowTo.INT.udn
new file mode 100644
--- /dev/null
+++ b/Engine/Documentation/Source/Engine/Rendering/RenderTargets/BlueprintRenderTargets/HowTo/CreatingTextures/CreatingTexturesHowTo.INT.udn
@@ -0,0 +1,156 @@
+availability:docs
+title:Creating Textures using Blueprints and Render Targets
+description:A demonstration of how to use Blueprints and Render Targets to create Textures in UE4.
+crumbs:%ROOT%, Engine, Engine\Rendering\RenderTargets
+version:4.13
+Type:topic
+SkillLevel:Intermediate
+parent:Engine
+order:2
+Navigation:topic
+
+[TOC(start:2)]
+
+![](Creating_Textures_Header_00.png)
+
+When creating Textures for your Unreal Engine 4 (UE4) project, you will generally need to use an additional 2D painting program like Adobe Photoshop or GIMP. While authorizing content in this manner can give you amazing results there are some instances where having the ability to create content like this inside of UE4 would be easier and faster to work on. In the following How To we will take a look at how you can use Blueprints and Render Targets to create Textures to use in your UE4 projects.
+
+## Required Asset Creation & Setup
+[region:note]
+Please note that the following is not specific to any UE4 project type and can be accomplished in any UE4 project that used UE4 version 4.13 or later. 
+[/region]
+
+Before we can create and use Textures, we will first need to create and set up the following three assets inside of the UE4 Content Browser.
+* **Render Target**
+* **Material**
+* **Blueprint** (Based of Actor)
+
+### Create a new Empty Level
+Creating a new **Empty Level** to place your Texture creation Blueprint in will make things easier to manage and also ensure that nothing else is affecting your Texture creation process. To create a new Empty Level inside of UE4 you will need to do the following:
+
+1. From the **Main Toolbar** go to **File** > **New Level**.
+
+	![](New_Level_01.png)
+
+1. From the New Level window that is displayed, click on the **Empty Level** option.
+
+	![](New_Level_02.png)
+
+1. After clicking on the Empty Level option UE4 will load a new level that has nothing inside of it. Once that has happened, save the level by clicking on the **Save** button on the Main Toolbar and giving it a name of **BlueprintTextureCreation**.
+
+	![](New_Level_03.png)
+
+### Render Target Creation & Setup
+To create a **Render Target** inside of the UE4 Content Browser you will need to do the following:
+
+1. Open up the UE4 Editor and then  **right - click** Inside of the **Content Browser** and from the menu that is displayed, locate the **Create Advanced Assets** section.
+
+	![](Create_RTT_01.png)
+
+1. From the Create Advanced Assets menu, move your mouse to the **Materials & Textures** section and from the asset list that is displayed, select the **Render Target** option.
+
+	![](Create_RTT_02.png)
+
+1. Name the newly created Render Target, **RT_00** and when completed your Content Browser should look like the following image.
+
+	![](Create_RTT_03.png)
+	
+### Material Creation & Setup
+To create and set up a Material to use a Render Target you will need to do the following inside of the UE4 Content Browser. 
+
+1. Right - Click in the Content Browser and from the **Create Basic Asset** section select the **Material** option to create a new Material.  
+
+	![](Create_RTT_Material_01.png)
+
+1. Name the Material **RT_Material_00** and then **double - click** on it to open it up in the Material Editor.
+
+	![](Create_RTT_Material_02.png)
+
+1. With the RT_Material_00 open in the Material Editor, search for and add the following Material Expression nodes to the Material Graph.
+
+	* **Texture Coordinate**
+	* **Constant**
+	* **Append**
+	* **Noise**
+
+	![](Create_RTT_Material_03.png)
+
+1. Once the required Material Expression nodes have been added to the Material Graph, connect the nodes so that they match the image below and press the **Apply** button when completed to compile the Material.
+
+	![](Create_RTT_Material_05.png)
+
+### Blueprint Creation & Setup
+To create and set up a Blueprint that will work with both Render Targets and Materials you will need to do the following.
+
+1. Right - Click in the Content Browser and from the menu that is displayed in the **Create Basic Asset** section, select the **Blueprint** option to create a new Blueprint.  
+
+	![](Create_Blueprint_01.png)
+
+1. From the **Pick Parent Class** window, select the **Actor** class and then name the Blueprint **BP_Texture_Creator**.
+
+	![](Create_Blueprint_02.png)
+
+1. Double - click on the BP_Texture_Creator Blueprint to open it up and once opened, click on the **Event Graph** tab to open up the Event Graph.
+
+	![](Create_Blueprint_03.png)
+
+1. Add the following Blueprint nodes to the Event Graph.
+	* **Custom Event**
+	* **Draw Material to Render Target** 
+
+	![](Create_Blueprint_04.png)
+
+1. Select the Custom Event node and press **F2** so you can re-name the Custom Event node to **Bake**. Then connect it's output to the input of the Draw Material to Render Target node.
+
+	![](Create_Blueprint_05.png)
+
+1. Next, on the Draw Material to Render Target node, load the Render Target that was created previously into the **Texture Render Target** slot by clicking on the small white triangle next to the text **Select Asset** and then from the list find the Render Target, RT_00, and select it by clicking on it. Once that is completed, do the same for the **Material** section, selecting the **RT_Material_00**.
+
+	![](Create_Blueprint_06.png)
+
+1. Now, click on the **Construction Script** tab and drag off the Construction Script node to display the **Executable actions** menu. From this menu locate the **Bake** event that was created in the Event Graph and click on it to add it. 
+
+	![](Create_Blueprint_07.png)
+
+1. Finally, press the **Compile** and **Save** buttons to compile and save your Blueprint.
+
+	![](Create_Blueprint_08.png)
+
+## Texture Creation
+Now that all of the required parts have been created and set up, in the following section we will take a look at how to create and save the Textures that are created.
+
+1. Locate the **BP_Texture_Creator** Blueprint in the Content Browser and then drag it into the level.
+
+	![](Create_Textures_01.png)
+
+1. Locate the **RT_Material_00** Material in the Content Browser and double - click on it to open it up. When opened, select the **Noise** Material Expression node.
+
+	![](Create_Textures_02.png)
+
+1. With the **Noise** Material Expression selected,  go to the **Details** panel and input the following values into the following parameters and when completed, press the **Apply** button to compile the Material.
+
+	![](Create_Textures_03.png)
+
+	* **Scale:** 2.0 
+	* **Levels:** 8.0
+	* **Output Min:** -0.01
+	* **Level Scale:** 3.0
+
+1. Select the BP_Texture_Creator Blueprint that was placed in the level and move it slightly in any direction to update the RT_00 Render Target that is in the Content Browser. When you do this you should see the RT_00 Render Target update and display exactly what was in the RT_Material_00. 
+
+	![](Create_Textures_04.png)
+
+	[region:note]
+	Anytime you want to see or update the changes that were made in the Material, you will need to ever so slightly move the Blueprint in order for the changes to propagate. 
+	[/region]
+
+1. Inside of the Content Browser, right-click on the RT_00 Render Target and in the **Render Targets Actions** section of the menu, select the **Create Static Texture** option to create a new Texture asset that will display the output of the RT_Material_00.
+
+	![](Create_Textures_05.png)
+
+1. When completed, you should now have a new 2D Texture that will display the contents of the RT_Material_00 and is now ready to be used in any other Material of your liking.
+
+	![](Create_Textures_06.png)
+
+## Conclusion
+Having the ability to create new Textures inside of UE4 opens up an entirely new world of content creation possibilities. Remember that in order to see your changes you will need to make sure that you compile your Material and slightly move your Blueprint so that your changes will be propagated to your Render Target.
\ No newline at end of file
diff --git a/Engine/Documentation/Source/Engine/Rendering/RenderTargets/BlueprintRenderTargets/HowTo/HeightFieldPainter/HeightFieldPainterHowTo.INT.udn b/Engine/Documentation/Source/Engine/Rendering/RenderTargets/BlueprintRenderTargets/HowTo/HeightFieldPainter/HeightFieldPainterHowTo.INT.udn
new file mode 100644
--- /dev/null
+++ b/Engine/Documentation/Source/Engine/Rendering/RenderTargets/BlueprintRenderTargets/HowTo/HeightFieldPainter/HeightFieldPainterHowTo.INT.udn
@@ -0,0 +1,49 @@
+availability:docs
+title:Creating Deformable Surfaces with Blueprints & Render Targets
+description:How-To use Blueprints and Render Targets to create realtime deformable surfaces.
+crumbs:
+version:4.13
+Type:topic
+SkillLevel:Intermediate
+parent:Engine
+order:2
+Navigation:topic
+
+
+
+
+
+
+
+
+
+
+
+
+
+
+
+
+
+
+
+
+
+
+
+
+
+
+
+
+
+
+
+
+
+
+
+
+
+
+
diff --git a/Engine/Documentation/Source/Engine/Rendering/RenderTargets/BlueprintRenderTargets/Overview/BlueprintRenderTargetsOverview.INT.udn b/Engine/Documentation/Source/Engine/Rendering/RenderTargets/BlueprintRenderTargets/Overview/BlueprintRenderTargetsOverview.INT.udn
new file mode 100644
--- /dev/null
+++ b/Engine/Documentation/Source/Engine/Rendering/RenderTargets/BlueprintRenderTargets/Overview/BlueprintRenderTargetsOverview.INT.udn
@@ -0,0 +1,64 @@
+availability:docs
+title:Blueprint & Render Targets Overview
+description:Taking a high level look at how to use Blueprints and Render Targets togther in UE4.
+crumbs:
+version:4.13
+Type:topic
+SkillLevel:Intermediate
+parent:Engine
+order:2
+Navigation:topic
+
+![](BP_RT_Overview.png)
+
+## Overview
+Using Blueprints to interact with render targets opens up a an entirely new way to create and use assets inside Unreal Engine 4 (UE4).  From creating fluid like surfaces to creating Textures that can be used with visual effects and Materials there is almost no limit to the amount of creative freedom that this will offer your UE4 projects. In the following document we will go over some of the things that you need to be aware of when using Blueprints and render targets.  
+
+## Know Limitations
+When using Blueprints and render targets there are few limitations and restrictions that you need to be aware of in order to get the most out of this tool. In the following section we will cover what these limitations are and what can be done to accommodate them in your workflow.
+
+* You cannot draw to a render target that is being sampled as a Texture by the Material that is input inot the Material input. If you want to accomplish this you will need to either use alpha blending to modify a render target in-place, or ping-pong between two different render targets.
+
+* Only the **Emissive Color** and **Opacity** outputs of the Material are valid when drawing to a render target. 
+
+* Lighting is not supported when using render targets and Material Expression nodes like **World Position** may return unexpected values.
+
+* Emissive Color is clamped to be positive by default, but you can output negative values by enabling the material property **AllowNegativeEmissiveColor**.
+
+	![](Allow_Neg_Emissive_Color.png)
+
+
+
+
+
+
+
+
+
+
+
+
+
+
+
+
+
+
+
+
+
+
+
+
+
+
+
+
+
+
+
+
+
+
+
+
diff --git a/Engine/Documentation/Source/Engine/Rendering/RenderTargets/BlueprintRenderTargets/Reference/BlueprintRenderTargetsReference.INT.udn b/Engine/Documentation/Source/Engine/Rendering/RenderTargets/BlueprintRenderTargets/Reference/BlueprintRenderTargetsReference.INT.udn
new file mode 100644
--- /dev/null
+++ b/Engine/Documentation/Source/Engine/Rendering/RenderTargets/BlueprintRenderTargets/Reference/BlueprintRenderTargetsReference.INT.udn
@@ -0,0 +1,89 @@
+availability:docs
+title: Blueprint & Render Target Reference
+description: A listing of properties and descriptions for using Blueprints & Render Targets.
+crumbs:
+version:4.13
+Type:Landing
+SkillLevel:Intermediate
+parent:Engine
+order:2
+Navigation:Reference
+
+## Overview
+This page contains a reference listing of properties you will find with the Render Target Blueprint nodes.
+
+
+## Begin Draw Canvas to Render Target
+
+![](Begin_Draw_Canvas_To_Render_Target.png)
+
+* Returns a Canvas object that can be used to draw to the specified render target. Be sure to call **EndDrawCanvasToRenderTarget** to complete the rendering!
+
+|Property| Description|
+|----------|-----------------|
+|Texture Render Target| Texture Render Target 2D reference.|
+|Canvas| Reference to canvas.|
+|Size| Size of the render target.|
+|Context| Draw to Render Target Context Structure.|
+
+## End Draw Canvas to Render Target
+
+![](End_Draw_Canvas_to_Render_Target.png)
+
+* Must be paired with a **BeginDrawCanvasToRenderTarget** to complete rendering to a Render Target.
+
+|Property| Description|
+|----------|-----------------|
+|Context| Draw to Render Target Context Structure.|
+
+
+## Clear Render Target 2D
+
+![](Clear_Render_Target_2D.png)
+
+* Clears the specified Render Target with the given Clear Color.
+
+|Property| Description|
+|----------|-----------------|
+|Texture Render Target| Texture Render Target 2D reference.|
+|Clear Color| Clears the specified Render Target with the given Clear Color.|
+
+
+## Create Render Target 2D
+
+![](Create_Render_Target_2D.png)
+
+* Creates a new Render Target and initializes it to the specified dimensions.
+
+|Property| Description|
+|----------|-----------------|
+|Width| Render Targets width size.|
+|Height| Render Targets height size.|
+|Return Value| Texture Render Target 2D reference.|
+
+## Draw Material to Render Target
+
+![](Draw_Material_To_Render_Target.png)
+
+* Renders a quad with the Material applied to the specified Render Target. This set the Render Target even if it is already set, which is an expensive operation. Use **BeginDrawCanvasToRenderTarget / EndDrawCanvasToRenderTarget** instead if rendering multiple primitives to the same same Render Target.
+
+|Property| Description|
+|----------|-----------------|
+|Texture Render Target| Texture Render Target 2D reference.|
+|Material| Material that the Render Target results will be saved to.|
+
+## Export Render Target
+
+![](Export_Render_Target.png)
+
+* Exports a Render Target as an HDR image onto the disk.
+
+|Property| Description|
+|----------|-----------------|
+|Texture Render Target| Texture Render Target 2D reference.|
+|File Path| Path to where the Render Target will be saved.|
+|File Name| Name to give the saved Render Target.|
+
+
+
+
diff --git a/Engine/Documentation/Source/Engine/Rendering/RenderTargets/RenderTargets.INT.udn b/Engine/Documentation/Source/Engine/Rendering/RenderTargets/RenderTargets.INT.udn
new file mode 100644
--- /dev/null
+++ b/Engine/Documentation/Source/Engine/Rendering/RenderTargets/RenderTargets.INT.udn
@@ -0,0 +1,57 @@
+availability:docs
+title: Render Targets
+description:Using Render Targets in Unreal Engine 4
+crumbs:
+version:4.13
+Type:landing
+SkillLevel:Intermediate
+parent:Engine
+order:1
+Navigation:topic
+
+![](BlueprintRenderTarget.png)
+
+A render target is a buffer where the video card draws pixels for a scene that is being rendered. render targets are used extensively throughout Unreal Engine 4 (UE4) and in the following set of documents we will take a look at how to use and manipulate render targets inside UE4.
+
+
+
+
+[REGION:topics third]
+[OBJECT:Topic]
+	[PARAM:image]
+		![Engine\Rendering\RenderTargets\BlueprintRenderTargets:description%](Engine/Rendering/RenderTargets/Render_Targets_Overview.png)
+	[/PARAM]
+	[PARAM:icon]
+		![](%ROOT%/blueprint_icon.png)(convert:false)
+	[/PARAM]
+	[PARAM:title]
+		%Engine\Rendering\RenderTargets\BlueprintRenderTargets/:description%
+	[/PARAM]
+	[PARAM:description]
+		%Engine\Rendering\RenderTargets\BlueprintRenderTargets:description%
+	[/PARAM]
+	[PARAM:path]
+		[RELATIVE:Engine\Rendering\RenderTargets\BlueprintRenderTargets/]
+	[/PARAM]
+[/OBJECT]
+[/REGION]
+
+
+
+
+
+
+
+
+
+
+
+
+
+
+
+
+
+
+
+
