#UE4 Docs: Adding new info for noise Mat Exp node.

diff --git a/Engine/Documentation/Source/Engine/Rendering/Materials/ExpressionReference/Utility/UtilityExpressions.INT.udn b/Engine/Documentation/Source/Engine/Rendering/Materials/ExpressionReference/Utility/UtilityExpressions.INT.udn
--- a/Engine/Documentation/Source/Engine/Rendering/Materials/ExpressionReference/Utility/UtilityExpressions.INT.udn
+++ b/Engine/Documentation/Source/Engine/Rendering/Materials/ExpressionReference/Utility/UtilityExpressions.INT.udn
@@ -234,14 +234,14 @@ The **Noise** expression creates a procedural noise field, giving you control ov
 | [REGION:tablesection]Properties[/REGION] ||
 | **Scale** | Changes the overall size of the noise cells. Lower numbers make the noise bigger. |
 | **Quality** | A look/performance setting. 0 is fast. Higher values are slower but may look better. |
-| **Function** | Controls the type of noise being calculated. Options include texture-based and purely computational versions of Simplex, Gradient, and Value noise. |
+| **Function** | [INCLUDE:#NoiseFunctions] |
 | **Turbulence** | Whether or not to calculate multiple levels of noise on top of one another. |
 | **Levels** | How many different levels of noise are combined. Only useful is Turbulence is active. |
 | **Output Min** | The lowest value output by the noise calculation. |
 | **Output Max** | The highest value output by the noise calculation. |
 | **Level Scale** | Controls the scale of individual levels when Turbulence is active. This defaults to 2, but higher values can result in more efficient use of low numbers of levels.
 | **Tiling** | For noise functions that support it, allows noise to tile. This is more expensive, but useful when baking noise into a seamless wrapping texture. |
-| **Repeat Size** | When tiling, how often should the noise repeat? |
+| **Repeat Size** | When tiling, how often should the noise repeat. |
 || [REGION:tablesection]Inputs[/REGION] ||
 | **Position** | Allows the texture size to be adjusted via a 3D vector. |
 | **FilterWidth** | In effect, controls how much blur will be applied to the noise texture. |
@@ -400,4 +400,21 @@ The image below shows what data the DistanceFieldGradient will use when the vari
 | -------- | ----------- |
 | **1**  | Enabling the R channel and disabling all other channels. |
 | **2**  | Enabling the G channel and disabling all other channels. |
-| **3**  | Enabling the B channel and disabling all other channels. |
\ No newline at end of file
+| **3**  | Enabling the B channel and disabling all other channels. |
+
+
+
+
+<!--
+[EXCERPT:NoiseFunctions]
+
+|Function| Description|
+|---------|-----------|
+|Simplex - Texture Based|High quality for direct use and bumps ~77 instructions per level, 4 Texture lookups, Can not tile.|
+|Gradient - Texture Based| High quality for direct use and bumps. Non-tiled: ~61 instructions per level, 8 Texture lookups. Tiling: ~74 instructions per level, 8 Texture lookups. Even "non tiled' mode has a repeat of 128. Useful repeat size range <= 128. Formerly labeled as Perlin Noise.|
+|Fast Gradient - 3D Texture | High quality for direct use, **BAD** for bumps. ~16 instructions per level, 1 Texture lookup. Always tiles with a repeat of 16, "Tiling" mode is not an option for Fast Gradient noise.|
+| Gradient - Computational | High quality for direct use and bumps. Non-tiled: ~80 instructions per level, no Textures. Tiling: ~143 instructions per level, no Textures. |
+| Value - Computational | Low quality, but pure computation. Non-tiled: ~53 instructions per level, no Textures. Tiling: ~118 instructions per level, no Textures. Formerly mislabeled as Gradient noise. |
+| Voronoi | Also know as Worley or Cellular noise. Quality=1 searches 8 cells, Quality=2 searches 16 cells, Quality=3 searches 27 cells, Quality=4 searches 32 cells. All are about ~20 instructions per cell searched. |
+[/EXCERPT]
+-->
