Fix some out of date information on the plugin doc page.

diff --git a/Engine/Documentation/Source/Programming/Plugins/Plugins.INT.udn b/Engine/Documentation/Source/Programming/Plugins/Plugins.INT.udn
--- a/Engine/Documentation/Source/Programming/Plugins/Plugins.INT.udn
+++ b/Engine/Documentation/Source/Programming/Plugins/Plugins.INT.udn
@@ -160,24 +160,21 @@ for engine modules. This is a philosophical choice that allows the engine to fun
 
 Here are the steps required to distribute a plugin.
 
-1. **Edit your plugin's descriptor file** (.uplugin file) and make sure the plugin name, modules, version and other settings are configured as desired.
-1. **Delete** your plugin's **Binaries** and **Intermediate** folders. This is just to make sure you are starting with a clean slate.
-1. If your plugin contains a **Source** directory, compile plugin binaries for **Win64 Development** configuration. This configuration will always be needed by an editor on PC platform to load your plugin. 
-1. For plugins with **Source**, if you intend to support additional configurations such as Mac, you should build those too.
-1. Make a copy of your plugin to a new "staging" folder outside of your project directory. It does not matter where, but you are going to be making changes to the plugin here.
-1. In the staging directory for your plugin, delete the **Intermediate** directory. These are temporary files that should never be distributed.
-1. If your plugin includes source code and you do not want to distribute it, delete any **Source** directories in the staging folder for your plugin.
-1. In the staging directory for your plugin, delete all files that you do not want to distribute.
-1. Your plugin is now ready to distribute!
+1. Make sure the modules list in your .uplugin file has a correct list of whitelisted and blacklisted platforms.
+1. Find your plugin in the editor's plugin browser, and click the '**Edit...**' link to update the metadata associated with your plugin
+   (description, link to documentation, and so on).
+1. Click the '**Package...*' link to package your plugin into a folder for distribution.
 
+By default, only source code, binaries and content will be packaged with your plugin. If you have additional files you need to include,
+create a file called Config/FilterPlugin.ini in your plugin folder which lists additional paths to include, like this:
+
+	[FilterPlugin]
+	/ThirdParty/...
+	/MyOtherFolder/...
 
 [REGION:warning]
 Important notes:
 
-* Including only Public headers with a distributed plugin is not yet supported. For now, you must either include the entire Source 
-directory, or no Source directory at all. We intend to add support for distributing Public source (as well as the *.Build.cs files) 
-at a later date.
-* We are investigating tools to automate most of this process. More details will be coming at a later date.
 * For plugins with Content, you may want to generate Derived Data to include with your plugin, to avoid having to require it 
 to be generated on demand by an end-user. This feature is not yet available, but will be added at a later date. Stay tuned.
 * Including an EULA or other documentation with distributed plugins is not yet supported. We may provide support for this at a later date.
@@ -222,28 +219,17 @@ This example plugin descriptor is from the **UObjectPlugin** [downloadable examp
 | **FileVersion**					| _Required_	| Version of this plugin descriptor file itself. It is used for backwards compatibility as new features are added to the plugin system. You should usually set this to the latest version that is allowed by the version of the engine you are using. The latest version is currently **3**, and is the version of format that is documented here. We do not expect this version to change very frequently. In source code, you can look at **EProjectDescriptorVersion** to see the actual value. If you require maximum compatibility with older versions of the engine, then you can use an older version of the format, but it is not recommended.
 | **Version**						| _Optional_    | Current version number of this build of your plugin. This value should always increase with future versions. This version number is not usually displayed to end-users. |
 | **VersionName**					| _Optional_    | Version of the plugin displayed in the editor UI. This is never used for any versioning checks and can be in whatever format you would like, however we suggest a simple Major.Minor format. You should always update the VersionName whenever the Version number has increased. |
-| **PackageFileUE4Version**			| _Optional_    | Minimum package file version necessary to load content in this plugin. The plugin will fail to load when installed to projects using a package file version less than the value specified here. |
-| **PackageFileLicenseeUE4Version**	| _Optional_    | A secondary minimum package file version necessary to load content in this plugin. The plugin will fail to load when installed to projects using a licensee package file version less than the value specified here. |
 | **FriendlyName**					| _Optional_	| Name of the plugin displayed in the editor UI. If not specified, the name will default to the .uplugin file name. |
 | **Description**					| _Optional_	| A paragraph of text that describes what this plugin is used for. This will be displayed in the editor's plugin window. |
 | **Category**						| _Optional_	| This is a special dot-separated path string that allows you to assign your plugin to a category in the editor UI. It is purely for organizational purposes. An example of a category path is **"Editor Features.Level Editing.Mesh Painting"**. Each category is separated by a period character and represents a deeper level in the tree. |
 | **CreatedBy**						| _Optional_	| The individual or company name that created this plugin. This may be displayed in the plugin UI or in other locations. |
 | **CreatedByURL**					| _Optional_	| A web link to the individual or company that created this plugin. If specified, the editor UI may display a hyperlink that allows the user to browse to this web page. |
+| **DocsURL**						| _Optional_	| A web link to the plugin's documentation. If specified, the editor's plugin browser will display a hyperlink that allows the user to browse to this web page. |
+| **SupportURL**					| _Optional_	| Link to the support page for the plugin. Will be shown on the editor's plugin browser if specified. |
 | **CanContainContent**				| _Optional_	| When specified and set to **true**, enables Content support for this plugin. The default setting is **false**. See the section on [Content in Plugins](#ContentInPlugins) for more info. | 
 | **Modules**						| _Optional_	| For plugins that contain source code (and binaries), this is the list of modules that should be loaded at startup. See below for more info. |
 
 
-[REGION:warning]
-This .uplugin file format is still in development and highly likely to change. For example, we expect the way that 
-versioning is handled will be modified in a future release.
-[/REGION]
-
-[REGION:warning]
-Localization of plugin name and description strings are not yet supported, or localization of text displayed in GUIs 
-implemented in the plugin, is not yet supported.
-[/REGION]
-
-
 #### Module Descriptors
 
 For plugins that contain code, the descriptor file will contain at least one module descriptor. 
@@ -253,13 +239,14 @@ For plugins that contain code, the descriptor file will contain at least one mod
 		"Type" : "Developer"
 	}
 
-| Field name		| Info			| Description |
-| -----------------	| -------------	| ----------- |
-| **Name**			| _Required_	| Unique name of this plugin module that will be loaded with the plugin. At runtime, the engine will expect appropriate plugin binaries to exist in the plugin's Binaries folder with the module name as specified here. For modules that have a Source directory, a matching *.Build.cs file is expected to exist within the module's subfolder tree. |
-| **Type**			| _Required_	| Sets the type of module. Valid options are **Runtime**, **RuntimeNoCommandlet**, **Developer**, **Editor**, **EditorNoCommandlet**, and **Program**. This type determines which types of applications this plugin's module is suitable for loading in. For example, some plugins may include modules that are only designed to be loaded when the editor is running. Runtime modules will be loaded in all cases, even in shipped games. Developer modules will only be loaded in development runtime or editor builds, but never in shipping builds. Editor modules will only be loaded when the editor is starting up. Your plugin can use a combination of modules of different types. |
-| **LoadingPhase**	| _Optional_	| If specified, controls when the plugin is loaded at startup. This is an advanced option that should not normally be required. The valid options are **Default** (which is used when no LoadingPhase is specified), **PreDefault**, and **PostConfigInit**. **PostConfigInit** allows the module to be loaded before the engine has finished starting up key subsystems. PreDefault loads just before the normal phase. Typically, this is only needed if you expect game modules to depend directly on content within your plugin, or types declared within your plugin's code. |
-
-		 
+| Field name				| Info			| Description |
+| -------------------------	| -------------	| ----------- |
+| **Name**					| _Required_	| Unique name of this plugin module that will be loaded with the plugin. At runtime, the engine will expect appropriate plugin binaries to exist in the plugin's Binaries folder with the module name as specified here. For modules that have a Source directory, a matching *.Build.cs file is expected to exist within the module's subfolder tree. |
+| **Type**					| _Required_	| Sets the type of module. Valid options are **Runtime**, **RuntimeNoCommandlet**, **Developer**, **Editor**, **EditorNoCommandlet**, and **Program**. This type determines which types of applications this plugin's module is suitable for loading in. For example, some plugins may include modules that are only designed to be loaded when the editor is running. Runtime modules will be loaded in all cases, even in shipped games. Developer modules will only be loaded in development runtime or editor builds, but never in shipping builds. Editor modules will only be loaded when the editor is starting up. Your plugin can use a combination of modules of different types. |
+| **LoadingPhase**			| _Optional_	| If specified, controls when the plugin is loaded at startup. This is an advanced option that should not normally be required. The valid options are **Default** (which is used when no LoadingPhase is specified), **PreDefault**, and **PostConfigInit**. **PostConfigInit** allows the module to be loaded before the engine has finished starting up key subsystems. PreDefault loads just before the normal phase. Typically, this is only needed if you expect game modules to depend directly on content within your plugin, or types declared within your plugin's code. |
+| **WhitelistPlatforms**	| _Optional_	| If specified, gives a list of platforms which this module will be compiled for. If not specified, the module will be compiled for all platforms. |
+| **BlacklistPlatforms**	| _Optional_	| If specified, gives a list of platforms which this this module will not be compiled for. If not specified, the module will be compiled for all platforms. |
+	 
 
 ## Icon Files
 
