#UE4 Docs: fixed topic object link

diff --git a/Engine/Documentation/Source/Engine/Landscape/Materials/Materials.INT.udn b/Engine/Documentation/Source/Engine/Landscape/Materials/Materials.INT.udn
--- a/Engine/Documentation/Source/Engine/Landscape/Materials/Materials.INT.udn
+++ b/Engine/Documentation/Source/Engine/Landscape/Materials/Materials.INT.udn
@@ -5,45 +5,7 @@ Description:Setting up materials to work with Landscape terrains.
 version: 4.9
 parent:Engine/Landscape
 order:5
-
-[VAR:Topic]
-[OBJECT:Topic]
-	[PARAM:image]
-		![%Engine/Landscape/Materials:title%](Engine/Landscape/Materials/landscape_materials_topic.png)
-	[/PARAM]
-	[PARAM:icon]
-		![](%ROOT%/landscape_icon.png)(convert:false)
-	[/PARAM]
-	[PARAM:title]
-		%Engine/Landscape/Materials:title%
-	[/PARAM]
-	[PARAM:description]
-		%Engine/Landscape/Materials:description%
-	[/PARAM]
-	[PARAMLITERAL:path]
-		[RELATIVE:Engine/Landscape/Materials]
-	[/PARAMLITERAL]
-[/OBJECT]
-[/VAR]
-[VAR:TopicCompact]
-[OBJECT:TopicCompact]
-	[PARAM:image]
-		![%Engine/Landscape/Materials:title%](Engine/Landscape/Materials/landscape_materials_topic.png)
-	[/PARAM]
-	[PARAM:icon]
-		![](%ROOT%/landscape_icon.png)(convert:false)
-	[/PARAM]
-	[PARAM:title]
-		%Engine/Landscape/Materials:title%
-	[/PARAM]
-	[PARAM:description]
-		%Engine/Landscape/Materials:description%
-	[/PARAM]
-	[PARAMLITERAL:path]
-		[RELATIVE:Engine/Landscape/Materials]
-	[/PARAMLITERAL]
-[/OBJECT]
-[/VAR]
+topic-image:landscape_materials_topic.png
 
 [TOC(start:2 end:3)]
 
diff --git a/Engine/Documentation/Source/Engine/Landscape/TechnicalGuide/LansdscapeTechnicalGuide.INT.udn b/Engine/Documentation/Source/Engine/Landscape/TechnicalGuide/LansdscapeTechnicalGuide.INT.udn
--- a/Engine/Documentation/Source/Engine/Landscape/TechnicalGuide/LansdscapeTechnicalGuide.INT.udn
+++ b/Engine/Documentation/Source/Engine/Landscape/TechnicalGuide/LansdscapeTechnicalGuide.INT.udn
@@ -5,45 +5,7 @@ Description:Technical Settings for Landscape.
 version: 4.9
 parent:Engine/Landscape
 order:6
-
-[VAR:Topic]
-[OBJECT:Topic]
-	[PARAM:image]
-		![%Engine/Landscape/TechnicalGuide:title%](Engine/Landscape/Materials/LS_Technical.png)
-	[/PARAM]
-	[PARAM:icon]
-		![](%ROOT%/landscape_icon.png)(convert:false)
-	[/PARAM]
-	[PARAM:title]
-		%Engine/Landscape/TechnicalGuide:title%
-	[/PARAM]
-	[PARAM:description]
-		%Engine/Landscape/TechnicalGuide:description%
-	[/PARAM]
-	[PARAMLITERAL:path]
-		[RELATIVE:Engine/Landscape/TechnicalGuide]
-	[/PARAMLITERAL]
-[/OBJECT]
-[/VAR]
-[VAR:TopicCompact]
-[OBJECT:TopicCompact]
-	[PARAM:image]
-		![%Engine/Landscape/TechnicalGuide:title%](Engine/Landscape/Materials/LS_Technical.png)
-	[/PARAM]
-	[PARAM:icon]
-		![](%ROOT%/landscape_icon.png)(convert:false)
-	[/PARAM]
-	[PARAM:title]
-		%Engine/Landscape/TechnicalGuide:title%
-	[/PARAM]
-	[PARAM:description]
-		%Engine/Landscape/TechnicalGuide:description%
-	[/PARAM]
-	[PARAMLITERAL:path]
-		[RELATIVE:Engine/Landscape/TechnicalGuide]
-	[/PARAMLITERAL]
-[/OBJECT]
-[/VAR]
+topic-image:Engine/Landscape/Materials/LS_Technical.png
 
 [TOC(start:2 end:3)]
 
