#UE4 Docs: Adding new info to the Landscape tessellation docs.

diff --git a/Engine/Documentation/Source/Engine/Landscape/Materials/Materials.INT.udn b/Engine/Documentation/Source/Engine/Landscape/Materials/Materials.INT.udn
--- a/Engine/Documentation/Source/Engine/Landscape/Materials/Materials.INT.udn
+++ b/Engine/Documentation/Source/Engine/Landscape/Materials/Materials.INT.udn
@@ -217,9 +217,11 @@ As an alternative to manually blending layers together with the **LandscapeLayer
 **To use a LandscapeLayerBlend node:**
 
 1. In the Material Editor, add a LandscapeLayerBlend node.
+	
 	![](Landscape_LayerBlend_Material_Example_00.png)
 
-1. In the **Details** panel, next to **Layers**,  click the plus icon (![Plus Sign](Shared/Icons/icon_PlusSymbol_12x.png)) to add a layer.
+1. In the **Details** panel, next to **Layers**, click the plus icon (![Plus Sign](Shared/Icons/icon_PlusSymbol_12x.png)) to add a layer.
+	
 	![](Landscape_LayerBlend_Material_Example_01.png)
 
 1. Expand the layer to view its properties.
@@ -364,10 +366,16 @@ For information about using Landscape Hole Materials, see [Using Opacity Masking
 
 ## Tessellation and Displacement
 
-The **Tessellation Multipliers** and **World Displacement** channels of the Material can be used in conjunction with Landscape terrains to provide additional detail and deformation. These work exactly as they would with any other standard Material, and the layer blending can be used in conjunction with these just as with the **Base Color**, **Normal**, or any other channel.
+The **Tessellation Multipliers** and **World Displacement** channels of the Material can be used in conjunction with Landscape terrains to provide additional detail and deformation when the Landscape is close to the camera. These work exactly as they would with any other standard Material, and the layer blending can be used in conjunction with these just as with the **Base Color**, **Normal**, or any other channel. Landscape will also only render hardware tessellation on the highest level-of-detail(LOD) which will fade out as the Landscape approaches the second LOD. Subsequent LODs will no longer have tessellation enabled which will significantly improve performance when enabling tessellation for displacement on your Landscape Terrain. 
 
 [REGION:note]
-These features are limited to use with DirectX 11 enabled. As such, they should be used as secondary, polish effects unless you are targeting DirectX 11 hardware specifically.
+These features are limited to use with DirectX 11 enabled. As such, they should be used as secondary, polish effects unless you are targeting DirectX 11 hardware specifically. You can read more about how to enable and use tessallation in your project by checking out the [Tessellation Multipler](Resources/ContentExamples/MaterialNodes/1_12) page. 
+[/REGION]
+
+![](Landscape_HW_Tess.png)
+
+[REGION:caption]
+In the visualization above, the highest the LOD (white) is rendered with tessellation enabled, while the other LODs (colors) are rendered without tessellation.  
 [/REGION]
 
 
