#UE4 docs SME feedback integration (early submit for computer swap)

diff --git a/Engine/Documentation/Source/Engine/HLOD/HLODLanding.INT.udn b/Engine/Documentation/Source/Engine/HLOD/HLODLanding.INT.udn
--- a/Engine/Documentation/Source/Engine/HLOD/HLODLanding.INT.udn
+++ b/Engine/Documentation/Source/Engine/HLOD/HLODLanding.INT.udn
@@ -2,14 +2,10 @@ title:Hierarchical Level of Detail
 description:Landing page for the Hierarchical Level of Detail system in Unreal Engine 4
 crumbs:
 availability:docs
-version:4.12
+version:4.13
 
 **Hierarchical Level of Detail** meshes (HLODs) have been developed as a way to increase performance and keep the quality bar high for projects using Unreal Engine 4.11 or later. HLODs can replace multiple **Static Mesh Actors** with single, combined **Static Mesh Actor** at long view distances. This helps reduce the number of **Actors** that need to be rendered for the scene, increasing performance by lowering the number of draw calls per frame. 
 
-[REGION:warning]
-HLOD doesn't currently work well with dynamic shadows (casting or recieving them), and as such, you should use static lighting and shadowing for the meshes that the HLOD system generates.
-[/REGION]
-
 [OBJECT:TopicButtonList]
 	[PARAM:icon]
 		![](%ROOT%/tech_icon.png)(convert:false)
diff --git a/Engine/Documentation/Source/Engine/HLOD/HowTo/BuildingHLODs/BuildingHLODs.INT.udn b/Engine/Documentation/Source/Engine/HLOD/HowTo/BuildingHLODs/BuildingHLODs.INT.udn
--- a/Engine/Documentation/Source/Engine/HLOD/HowTo/BuildingHLODs/BuildingHLODs.INT.udn
+++ b/Engine/Documentation/Source/Engine/HLOD/HowTo/BuildingHLODs/BuildingHLODs.INT.udn
@@ -2,16 +2,12 @@ title:How-To Build Hierarchical Level of Detail Meshes
 description:This how to covers how to generate the HLOD meshes for your HLOD enabled Unreal Engine 4 project.
 crumbs:
 availability:docs
-version:4.12
+version:4.13
 related:Engine/HLOD/Overview
 related:Engine/HLOD/Reference
 
 There are two major steps in building your **Hierarchical Level of Detail** (HLOD) Meshes. Firstly, you must **Generate Clusters**, which groups **Actors** in your level together based on the setting you specify in the **Cluster Generation Settings**. Then you have to **Generate Proxy Meshes**, which can take a very long time based on scene complexity and settings specified in the **Mesh Generation Settings**.
 
-[REGION:warning]
-HLOD doesn't currently work well with dynamic shadows (casting or receiving them), and as such, you should use static lighting and shadowing for the meshes that the HLOD system generates.
-[/REGION]
-
 ## Steps
 
 Below are the steps to generate your **HLOD Meshes** broken into two major sections: **Generate Clusters** and **Generate Proxy Meshes**.
@@ -48,6 +44,8 @@ Should you run into any issues you don’t like with the generated cluster you c
 
 **Actors** can also be excluded on a per-instance basis by selecting them in the level and in their **Details Panel** set the option for **Enable Auto LODGeneration** to false.
 
+Further, if you wish to add an **Actor** to a **Cluster** you can click and drag from the **World Outliner** to the **Cluster** you wish it to be included in.
+
 ![image alt text](Engine/HLOD/Overview/image_11.png)(h:236 w:269 convert:false)
 
 Repeat this process while adjusting the values in the HLOD's level **Cluster Generation Settings** until you are happy with the clusters being generated, then continue on to the next section: [**Generate Proxy Meshes**](#generateproxymeshes).
diff --git a/Engine/Documentation/Source/Engine/HLOD/HowTo/Enabling/EnablingHLODs.INT.udn b/Engine/Documentation/Source/Engine/HLOD/HowTo/Enabling/EnablingHLODs.INT.udn
--- a/Engine/Documentation/Source/Engine/HLOD/HowTo/Enabling/EnablingHLODs.INT.udn
+++ b/Engine/Documentation/Source/Engine/HLOD/HowTo/Enabling/EnablingHLODs.INT.udn
@@ -2,16 +2,12 @@ title:How-To Enable the HLOD System
 description:This how-to covers how to enable the HLOD system in your Unreal Engine 4 project.
 crumbs:
 availability:docs
-version:4.12
+version:4.13
 related:Engine/HLOD/Overview
 related:Engine/HLOD/Reference
 
 The **Hierarchical Level of Detail System** (HLOD) must be enabled on each Level that wants to make use of it. Follow the steps below to enable the Hierarchical LOD System.
 
-[REGION:warning]
-HLOD doesn't currently work well with dynamic shadows (casting or receiving them), and as such, you should use static lighting and shadowing for the meshes that the HLOD system generates.
-[/REGION]
-
 ## Steps
 
 1. Enable by going to **World Settings -> LODSystem -> Enable Hierarchical LODSystem**
diff --git a/Engine/Documentation/Source/Engine/HLOD/HowTo/HLODHowTos.INT.udn b/Engine/Documentation/Source/Engine/HLOD/HowTo/HLODHowTos.INT.udn
--- a/Engine/Documentation/Source/Engine/HLOD/HowTo/HLODHowTos.INT.udn
+++ b/Engine/Documentation/Source/Engine/HLOD/HowTo/HLODHowTos.INT.udn
@@ -2,6 +2,7 @@ title:Hierarchical Level of Detail How-To's
 description:Words
 crumbs:
 availability:docs
+version:4.13
 
 This page contains links to the **Hierarchical Level of Detail** (HLOD) mesh's How-To's. They will walk you through enabling the HLOD system and setting up HLOD in your game.
 
diff --git a/Engine/Documentation/Source/Engine/HLOD/Overview/HLODOverview.INT.udn b/Engine/Documentation/Source/Engine/HLOD/Overview/HLODOverview.INT.udn
--- a/Engine/Documentation/Source/Engine/HLOD/Overview/HLODOverview.INT.udn
+++ b/Engine/Documentation/Source/Engine/HLOD/Overview/HLODOverview.INT.udn
@@ -2,20 +2,17 @@ title:Hierarchical Level of Detail Overview
 description:Overview of the Hierarchical Level of Detail system in Unreal Engine 4
 crumbs:%ROOT%
 availability:docs
+version:4.13
 
 HLODs, in their simplest form, combine pre-existing **Static Mesh Actors** into a single **Static Mesh Actor** and material, with atlased textures, to reduce draw calls which in-turn helps with your performance of your game. There is a robust system that enables you to set what needs to be generated when you create these **HLOD Actors**.
 
-[REGION:warning]
-HLOD doesn't currently work well with dynamic shadows (casting or recieving them), and as such, you should use static lighting and shadowing for the meshes that the HLOD system generates.
-[/REGION]
-
 ## Proxy Meshes
 
 Once your proxy meshes have been generated, you can find them located under the `HLOD` folder generated in your project. This will include the **LODActors** referenced in the **HLOD Outliner**. 
 
 ![image alt text](image_15.png)(h:320 w:373 convert:false)
 
-These proxy meshes can be opened independently and settings adjusted. If you decided to merge the materials, you’ll have access to them to adjust the texture source and any material settings. 
+These proxy meshes can be opened independently and settings adjusted.
 
 ![image alt text](image_16.png)(h:357 w:404 convert:false)
 
diff --git a/Engine/Documentation/Source/Engine/HLOD/Reference/HLODReference.INT.udn b/Engine/Documentation/Source/Engine/HLOD/Reference/HLODReference.INT.udn
--- a/Engine/Documentation/Source/Engine/HLOD/Reference/HLODReference.INT.udn
+++ b/Engine/Documentation/Source/Engine/HLOD/Reference/HLODReference.INT.udn
@@ -2,13 +2,10 @@ title:Hierarchical LOD Reference
 description:A reference of the interface elements and properties found in the Hierarchical LOD Outliner.
 crumbs:
 availability:docs
+version:4.13
 
 This page details the buttons and properties that exist on the **Hierarchical LOD Outliner**.
 
-[REGION:warning]
-HLOD doesn't currently work well with dynamic shadows (casting or recieving them), and as such, you should use static lighting and shadowing for the meshes that the HLOD system generates.
-[/REGION]
-
 ## Property and Interface Reference
 
 ![HLOD Outliner Fully Expanded](image_2.png)
@@ -18,6 +15,7 @@ HLOD doesn't currently work well with dynamic shadows (casting or recieving them
 
 Below are the descriptions for the 3 buttons and slider that exist in the **Hierarchical LOD Outliner**.
 
+[REGION:raw]
 | --- | :---: | :--- |
 | **Generate Clusters** || This will generate spherical clusters of **Static Mesh Actors** in the scene that can be grouped. The way this is handled is by using the settings for each LOD settings specified under the LODSystem section of the HLOD Outliner. |
 | **Delete Clusters** || This will delete all HLODs that have been previously setup in this level. This includes all generated proxy meshes as well as the clusters that have been setup. This will not wipe the previously set values or properties for the HLOD Levels. |
@@ -26,12 +24,44 @@ Below are the descriptions for the 3 buttons and slider that exist in the **Hier
 | | %spacer%  | |
 [/REGION]
 
+### Context Menus
+
+The Hierarchical LOD Outliner has a few right-click context menus, below are how to access them and the options they present.
+
+* The first is accessed by right-clicking on a **Cluster Entry** in the HLOD Outliner:
+
+	[REGION:raw]
+	![](ContextMenu1.png)
+	[/REGION]
+
+	[REGION:raw]
+	| Option | %spacer% | Description |
+	| --- | :---: | :--- |
+	| **Select LOD Actor** || Selects the LOD Actor in the viewport (same effect as left-clicking the cluster). |
+	| **Select Contained Actors** || Selects all the **Actors** contained in the **LOD Cluster**. |
+	| **ForceView** || Toggles the view of the **HLOD Actor** in the viewport. |
+	| **Rebuild Proxy Mesh** || Forces the right-clicked cluster to rebuild its proxy mesh. |
+	| **Create Containing Hierarchical Volume** || Creates a **Hierarchical Volume** around the **Actors** contained in the **Cluster** that can then be modified to encompass greater or fewer **Actors**. |
+	| **Delete Cluster** || Deletes the right-clicked Cluster. |
+	[/REGION]
+
+* The next is accessed by right-clicking on an **Actor** entry contained within an **Cluster Entry** in the HLOD Outliner:
+
+	[REGION:raw]
+	![](ContextMenu2.png)
+	[/REGION]
+
+	[REGION:raw]
+	| Option | %spacer% | Description |
+	| --- | :---: | :--- |
+	| **Remove From Cluster** || Removes the right-clicked **Actor** from the **Cluster**. |
+	| **Exclude From Cluster Generation** || Excludes the right-clicked **Actor** from **Cluster** generation. If you wish to re-add the excluded **Actor** click and drag it from the **World Outliner** into the **Cluster** in the **Hierarchical LOD Outliner**. |
+	[/REGION]
+
 ### Properties
 
 Below, broken out by major section, are the properties that can be found in a **Hierarchical LODSetup** in the **LODSystem** panel located in the lower portion of the **Hierarchical LOD Outliner** interface.
 
-[REGION:raw]
-
 ### Cluster Generation Settings
 
 | Property | %spacer% | Description |
@@ -54,8 +84,8 @@ Below, broken out by major section, are the properties that can be found in a **
 | **Merge Materials** || Whether to merge source materials into one flat material. |
 | [REGION:tablesection]**[Material Settings See Below](#materialsettings)**[/REGION] |||
 | **Bake Vertex Data ** || Incorporates painted vertex colours into model/baked out material. |
-| **Calculate Correct LODModel** || Calculates the correct LOD model based upon the screen size value (similar to how regular LODs are determined for static meshes in the editor/game). |
-| **Export Specific LOD** || Forces the system to pick the given LOD level for all models (if the LOD level exists, otherwise picks the maximum closest value). |
+| **Use Vertex Data for Baking Material** || Whether or not vertex data such as vertex colours should be used when baking out materials. |
+| **Use Landscape Culling** || Whether or note to use available landscape geometry to cull away invisible triangles. |
 
 ### Material Settings
 
