#UE4 docs Blueprint Splines Docs initial submission (early submit due to computer swap)

diff --git a/Engine/Documentation/Source/Engine/BlueprintSplines/HowTo/AddSplineComponent/AddSplineComponent.INT.udn b/Engine/Documentation/Source/Engine/BlueprintSplines/HowTo/AddSplineComponent/AddSplineComponent.INT.udn
new file mode 100644
--- /dev/null
+++ b/Engine/Documentation/Source/Engine/BlueprintSplines/HowTo/AddSplineComponent/AddSplineComponent.INT.udn
@@ -0,0 +1,73 @@
+Availability:Docs
+Title:How-To Add a Spline Component to a Blueprint
+Crumbs: %ROOT%, Engine
+Description:This how-to will show you how to add a Spline or Spline Mesh Component to a Blueprint.
+Navigation:topic
+parent:Engine/BlueprintSplines/HowTo
+order:1
+social-image:Splines_social.png
+related:Engine/BlueprintSplines/HowTo/EditSplineComponentInEditor
+related:Engine/BlueprintSplines/HowTo/EditSplineInConstScript
+
+[VAR:Topic]
+[OBJECT:Topic]
+	[PARAM:image]
+		![%Engine/BlueprintSplines/HowTo/AddSplineComponent:title%](Engine/BlueprintSplines/Splines_topic.png)
+	[/PARAM]
+	[PARAM:icon]
+		![](%ROOT%/editor_icon.png)(convert:false)
+	[/PARAM]
+	[PARAM:title]
+		%Engine/BlueprintSplines/HowTo/AddSplineComponent:title%
+	[/PARAM]
+	[PARAM:description]
+		%Engine/BlueprintSplines/HowTo/AddSplineComponent:description%
+	[/PARAM]
+	[PARAM:path]
+		[RELATIVE:Engine/BlueprintSplines/HowTo/AddSplineComponent]
+	[/PARAM]
+[/OBJECT]
+[/VAR]
+[VAR:TopicCompact]
+[OBJECT:TopicCompact]
+	[PARAM:image]
+		![%Engine/BlueprintSplines/HowTo/AddSplineComponent:title%](Engine/BlueprintSplines/Splines_topic.png)
+	[/PARAM]
+	[PARAM:icon]
+		![](%ROOT%/editor_icon.png)(convert:false)
+	[/PARAM]
+	[PARAM:title]
+		%Engine/BlueprintSplines/HowTo/AddSplineComponent:title%
+	[/PARAM]
+	[PARAM:description]
+		%Engine/BlueprintSplines/HowTo/AddSplineComponent:description%
+	[/PARAM]
+	[PARAM:path]
+		[RELATIVE:Engine/BlueprintSplines/HowTo/AddSplineComponent]
+	[/PARAM]
+[/OBJECT]
+[/VAR]
+
+
+This page will take you through the steps necessary to add a **Blueprint Spline Component** or **Blueprint Spline Mesh Component** to any Blueprint you have available. Both of these components require that your Blueprint inherits from any class that can trace its linage back to **Actor** (Pawn, Character, etc...).
+
+## Steps
+
+1. Create a new **Blueprint** or Open a **Blueprint** in the **Blueprint Editor**.
+
+1. **Add Component** ![](AddComp_Button.png) -> Spline **or** Add Component -> Spline Mesh.
+
+	[REGION:raw]
+	![](AddSplineComp.png)
+	[/REGION]
+	
+1. Adjust any settings for the added component. See [](Engine/BlueprintSplines/Reference/ComponentProperties) or [](Engine/BlueprintSplines/Reference/MeshComponentProperties) for more information on the properties available.
+
+	You won't be able to directly edit the Spline's or Spline Mesh's points in the **Blueprint Editor Viewport**. You'll either need define the curve in the **Construction Script** of your Blueprint or place your Blueprint in your level and edit the points there.
+	
+## Result
+
+You will have added a **Spline Component** or **Spline Mesh Component** to your Blueprint, see the next how-to's for more information on how to manipulate them.
+
+
+
diff --git a/Engine/Documentation/Source/Engine/BlueprintSplines/HowTo/EditSplineComponentInEditor/EditSplineCompInEditor.INT.udn b/Engine/Documentation/Source/Engine/BlueprintSplines/HowTo/EditSplineComponentInEditor/EditSplineCompInEditor.INT.udn
new file mode 100644
--- /dev/null
+++ b/Engine/Documentation/Source/Engine/BlueprintSplines/HowTo/EditSplineComponentInEditor/EditSplineCompInEditor.INT.udn
@@ -0,0 +1,87 @@
+Availability:Docs
+Title:How-To Edit a Spline Component in the Level Editor
+Crumbs: %ROOT%, Engine
+Description:Null
+Navigation:topic
+parent:Engine/BlueprintSplines/HowTo
+order:2
+social-image:Splines_social.png
+related:Engine/BlueprintSplines/HowTo/EditSplineInConstScript
+
+[VAR:Topic]
+[OBJECT:Topic]
+	[PARAM:image]
+		![%Engine/BlueprintSplines/HowTo/EditSplineComponentInEditor:title%](Engine/BlueprintSplines/Splines_topic.png)
+	[/PARAM]
+	[PARAM:icon]
+		![](%ROOT%/editor_icon.png)(convert:false)
+	[/PARAM]
+	[PARAM:title]
+		%Engine/BlueprintSplines/HowTo/EditSplineComponentInEditor:title%
+	[/PARAM]
+	[PARAM:description]
+		%Engine/BlueprintSplines/HowTo/EditSplineComponentInEditor:description%
+	[/PARAM]
+	[PARAM:path]
+		[RELATIVE:Engine/BlueprintSplines/HowTo/EditSplineComponentInEditor]
+	[/PARAM]
+[/OBJECT]
+[/VAR]
+[VAR:TopicCompact]
+[OBJECT:TopicCompact]
+	[PARAM:image]
+		![%Engine/BlueprintSplines/HowTo/EditSplineComponentInEditor:title%](Engine/BlueprintSplines/Splines_topic.png)
+	[/PARAM]
+	[PARAM:icon]
+		![](%ROOT%/editor_icon.png)(convert:false)
+	[/PARAM]
+	[PARAM:title]
+		%Engine/BlueprintSplines/HowTo/EditSplineComponentInEditor:title%
+	[/PARAM]
+	[PARAM:description]
+		%Engine/BlueprintSplines/HowTo/EditSplineComponentInEditor:description%
+	[/PARAM]
+	[PARAM:path]
+		[RELATIVE:Engine/BlueprintSplines/HowTo/EditSplineComponentInEditor]
+	[/PARAM]
+[/OBJECT]
+[/VAR]
+
+This how-to will step through editing **Blueprint Spline Component** curves while they are in the Level Editor. **Spline Mesh Components** can also be set so they can be edited in the Level Editor, but only have two points to use, and you cannot add any more points.
+
+## Steps
+
+1. Place Blueprint with a **Spline Component** into your Level.
+
+	![](PlaceActor.png)
+
+1. Select Actor with Spline Component in your level. If you just placed it into the level, it will already be selected.
+
+	If your Blueprint Actor has no other visible Components (Meshes, billboards, etc...) it may be difficult to select, however, you can always select it from the **World Outliner**.
+
+	[REGION:note]
+	If you are using a **Spline Mesh Component** you'll need to enable **Allow Spline Editing Per Instance** on the Blueprint in the Blueprint Editor to see its **Curve Points**.
+	[/REGION]
+	
+1. Select a **Curve Point** on the spline, and using the Move, Rotate, and Scale tools, manipulate the Curve Points and Tangent Points into the desired shape.
+
+	![](FullEdit.gif)(convert:false)
+
+	You can also use the following features to edit the curve:
+	
+	* To Add a point, right-click on a **Curve Section** and select **Add Spline Point Here**.
+	* To Duplicate a point, right-click the **Curve Point** you want to duplicate and select **Duplicate Spline Point**
+	* To Remove a point, right-click on a **Curve Point** and select **Delete Spline Point**
+	
+	For more information on all the tools available to edit Spline Components see [](Engine/BlueprintSplines/Reference/SplineEditorTool).
+	
+	[REGION:note]
+	Once you start editing the Spline Component, the **Override Construction Script** property will be set to **True** to prevent it from overwriting your work. 
+	[/REGION]
+	
+## Result
+
+While above, we are just editing a spline that doesn't affect anything, there are a number of uses for this data. See [](Resources/ContentExamples/Blueprint_Splines) documentation and the **Content Examples** project for additional information.
+
+
+
diff --git a/Engine/Documentation/Source/Engine/BlueprintSplines/HowTo/EditSplineInConstScript/EditSplineInConstScript.INT.udn b/Engine/Documentation/Source/Engine/BlueprintSplines/HowTo/EditSplineInConstScript/EditSplineInConstScript.INT.udn
new file mode 100644
--- /dev/null
+++ b/Engine/Documentation/Source/Engine/BlueprintSplines/HowTo/EditSplineInConstScript/EditSplineInConstScript.INT.udn
@@ -0,0 +1,51 @@
+Availability:Docs
+Title:How-To Edit a Spline Component in the Construction Script
+Crumbs: %ROOT%, Engine
+Description:Null
+Navigation:topic
+parent:Engine/BlueprintSplines/HowTo
+order:3
+social-image:Splines_social.png
+
+[VAR:Topic]
+[OBJECT:Topic]
+	[PARAM:image]
+		![%Engine/BlueprintSplines/HowTo/EditSplineInConstScript:title%](Engine/BlueprintSplines/Splines_topic.png)
+	[/PARAM]
+	[PARAM:icon]
+		![](%ROOT%/editor_icon.png)(convert:false)
+	[/PARAM]
+	[PARAM:title]
+		%Engine/BlueprintSplines/HowTo/EditSplineInConstScript:title%
+	[/PARAM]
+	[PARAM:description]
+		%Engine/BlueprintSplines/HowTo/EditSplineInConstScript:description%
+	[/PARAM]
+	[PARAM:path]
+		[RELATIVE:Engine/BlueprintSplines/HowTo/EditSplineInConstScript]
+	[/PARAM]
+[/OBJECT]
+[/VAR]
+[VAR:TopicCompact]
+[OBJECT:TopicCompact]
+	[PARAM:image]
+		![%Engine/BlueprintSplines/HowTo/EditSplineInConstScript:title%](Engine/BlueprintSplines/Splines_topic.png)
+	[/PARAM]
+	[PARAM:icon]
+		![](%ROOT%/editor_icon.png)(convert:false)
+	[/PARAM]
+	[PARAM:title]
+		%Engine/BlueprintSplines/HowTo/EditSplineInConstScript:title%
+	[/PARAM]
+	[PARAM:description]
+		%Engine/BlueprintSplines/HowTo/EditSplineInConstScript:description%
+	[/PARAM]
+	[PARAM:path]
+		[RELATIVE:Engine/BlueprintSplines/HowTo/EditSplineInConstScript]
+	[/PARAM]
+[/OBJECT]
+[/VAR]
+
+[VAR:spacer]![](%ROOT%/spacer.png)(w:150 convert:false)[/VAR]
+
+[VAR:spacer2]![](%ROOT%/spacer.png)(w:40 convert:false)[/VAR]
diff --git a/Engine/Documentation/Source/Engine/BlueprintSplines/HowTo/HowToLanding.INT.udn b/Engine/Documentation/Source/Engine/BlueprintSplines/HowTo/HowToLanding.INT.udn
new file mode 100644
--- /dev/null
+++ b/Engine/Documentation/Source/Engine/BlueprintSplines/HowTo/HowToLanding.INT.udn
@@ -0,0 +1,58 @@
+Availability:Docs
+Title:Blueprint Splines How-To's
+Crumbs: %ROOT%, Engine
+Description:Various How-To's for Blueprint Spline Components and Blueprint Spline Mesh Components.
+Navigation:topic
+parent:Engine/BlueprintSplines
+order:3
+social-image:Splines_social.png
+
+[VAR:Topic]
+[OBJECT:Topic]
+	[PARAM:image]
+		![%Engine/Splines:title%](Engine/Splines/Splines_topic.png)
+	[/PARAM]
+	[PARAM:icon]
+		![](%ROOT%/editor_icon.png)(convert:false)
+	[/PARAM]
+	[PARAM:title]
+		%Engine/Splines:title%
+	[/PARAM]
+	[PARAM:description]
+		%Engine/Splines:description%
+	[/PARAM]
+	[PARAM:path]
+		[RELATIVE:Engine/Splines]
+	[/PARAM]
+[/OBJECT]
+[/VAR]
+[VAR:TopicCompact]
+[OBJECT:TopicCompact]
+	[PARAM:image]
+		![%Engine/Splines:title%](Engine/Splines/Splines_topic.png)
+	[/PARAM]
+	[PARAM:icon]
+		![](%ROOT%/editor_icon.png)(convert:false)
+	[/PARAM]
+	[PARAM:title]
+		%Engine/Splines:title%
+	[/PARAM]
+	[PARAM:description]
+		%Engine/Splines:description%
+	[/PARAM]
+	[PARAM:path]
+		[RELATIVE:Engine/Splines]
+	[/PARAM]
+[/OBJECT]
+[/VAR]
+
+Below are the how-to's available for Blueprint Spline Components and Blueprint Spline Mesh Components.
+
+## How-To's
+
+[DIR(output:"listbutton" parent:"Engine/BlueprintSplines/HowTo" org:"hierarchy" end:"1")]
+
+
+
+
+
diff --git a/Engine/Documentation/Source/Engine/BlueprintSplines/Overview/BPSplineOverview.INT.udn b/Engine/Documentation/Source/Engine/BlueprintSplines/Overview/BPSplineOverview.INT.udn
new file mode 100644
--- /dev/null
+++ b/Engine/Documentation/Source/Engine/BlueprintSplines/Overview/BPSplineOverview.INT.udn
@@ -0,0 +1,65 @@
+Availability:Docs
+Title:Blueprint Spline Components Overview
+Crumbs: %ROOT%, Engine
+Description:A brief overview of Blueprint Spline Components and Blueprint Spline Mesh Components.
+Navigation:topic
+parent:Engine/BlueprintSplines
+order:1
+social-image:Splines_social.png
+related:Engine/BlueprintSplines/
+related:Engine/BlueprintSplines/HowTo
+related:Engine/BlueprintSplines/Reference/ComponentProperties
+related:Engine/BlueprintSplines/Reference/MeshComponentProperties
+related:Engine/BlueprintSplines/Reference/SplineEditorTool
+
+[VAR:Topic]
+[OBJECT:Topic]
+	[PARAM:image]
+		![%Engine/BlueprintSplines/Overview:title%](Engine/BlueprintSplines/Splines_topic.png)
+	[/PARAM]
+	[PARAM:icon]
+		![](%ROOT%/editor_icon.png)(convert:false)
+	[/PARAM]
+	[PARAM:title]
+		%Engine/BlueprintSplines/Overview:title%
+	[/PARAM]
+	[PARAM:description]
+		%Engine/BlueprintSplines/Overview:description%
+	[/PARAM]
+	[PARAM:path]
+		[RELATIVE:Engine/BlueprintSplines/Overview]
+	[/PARAM]
+[/OBJECT]
+[/VAR]
+[VAR:TopicCompact]
+[OBJECT:TopicCompact]
+	[PARAM:image]
+		![%Engine/BlueprintSplines/Overview:title%](Engine/BlueprintSplines/Splines_topic.png)
+	[/PARAM]
+	[PARAM:icon]
+		![](%ROOT%/editor_icon.png)(convert:false)
+	[/PARAM]
+	[PARAM:title]
+		%Engine/BlueprintSplines/Overview:title%
+	[/PARAM]
+	[PARAM:description]
+		%Engine/BlueprintSplines/Overview:description%
+	[/PARAM]
+	[PARAM:path]
+		[RELATIVE:Engine/BlueprintSplines/Overview]
+	[/PARAM]
+[/OBJECT]
+[/VAR]
+
+
+At their core, **Blueprint Spline Components** are just a path for you to use for positional data. You can use it to move **Actors** (or other **Components**) around the world or place a series of **Actors** (or other **Components**) along the spline. They are fully editable in Blueprints with the ability to add/remove/duplicate Spline Points, change their tangent types, even animate them on tick. Further, they are also editable in the Level Editor with the capability to pass back those edits to their **Blueprint Construction Script**.
+
+![](BPSC_1.png)
+
+**Blueprint Spline Mesh Components** have a completely different use case. These deform a single **Static Mesh** along a two point spline. You cannot add more Spline Points to a Blueprint Spline Mesh Component, but the two points are completely controllable through Blueprints.
+
+![](BPSMC_1.png)
+
+While their use cases are very different, they share the same procedures for adding them to a Blueprint, and editing tools.
+
+
diff --git a/Engine/Documentation/Source/Engine/BlueprintSplines/Reference/ComponentProperties/SplineComponentReference.INT.udn b/Engine/Documentation/Source/Engine/BlueprintSplines/Reference/ComponentProperties/SplineComponentReference.INT.udn
new file mode 100644
--- /dev/null
+++ b/Engine/Documentation/Source/Engine/BlueprintSplines/Reference/ComponentProperties/SplineComponentReference.INT.udn
@@ -0,0 +1,473 @@
+Availability:Docs
+Title:Blueprint Spline Component Property Reference
+Crumbs: %ROOT%, Engine
+Description:Contains a reference of all properties available in the editor for Blueprint Spline Components.
+Navigation:topic
+parent:Engine/BlueprintSplines
+order:1
+social-image:Splines_social.png
+
+[VAR:Topic]
+[OBJECT:Topic]
+	[PARAM:image]
+		![%Engine/BlueprintSplines/Reference/ComponentProperties:title%](Engine/BlueprintSplines/Splines_topic.png)
+	[/PARAM]
+	[PARAM:icon]
+		![](%ROOT%/editor_icon.png)(convert:false)
+	[/PARAM]
+	[PARAM:title]
+		%Engine/BlueprintSplines/Reference/ComponentProperties:title%
+	[/PARAM]
+	[PARAM:description]
+		%Engine/BlueprintSplines/Reference/ComponentProperties:description%
+	[/PARAM]
+	[PARAM:path]
+		[RELATIVE:Engine/BlueprintSplines/Reference/ComponentProperties]
+	[/PARAM]
+[/OBJECT]
+[/VAR]
+[VAR:TopicCompact]
+[OBJECT:TopicCompact]
+	[PARAM:image]
+		![%Engine/BlueprintSplines/Reference/ComponentProperties:title%](Engine/BlueprintSplines/Splines_topic.png)
+	[/PARAM]
+	[PARAM:icon]
+		![](%ROOT%/editor_icon.png)(convert:false)
+	[/PARAM]
+	[PARAM:title]
+		%Engine/BlueprintSplines/Reference/ComponentProperties:title%
+	[/PARAM]
+	[PARAM:description]
+		%Engine/BlueprintSplines/Reference/ComponentProperties:description%
+	[/PARAM]
+	[PARAM:path]
+		[RELATIVE:Engine/BlueprintSplines/Reference/ComponentProperties]
+	[/PARAM]
+[/OBJECT]
+[/VAR]
+
+[VAR:spacer]![](%ROOT%/spacer.png)(w:150 convert:false)[/VAR]
+
+[VAR:spacer2]![](%ROOT%/spacer.png)(w:40 convert:false)[/VAR]
+
+This page contains a reference listing of properties available on **Spline Components**. The properties displayed will be slightly different if the **Blueprint Spline Component** is selected in the **Blueprint Editor**, is a **Root Component**, or if the component is selected in the **Level Editor**.
+
+## Properties
+
+[EXCERPT:PropTransformCat]
+[OBJECT:modulesection]
+	[PARAMLITERAL:id]
+		TransformCat
+	[/PARAMLITERAL]
+	[PARAM:heading]
+		Transform
+	[/PARAM]
+	[PARAM:content]
+		%c1%
+	[/PARAM]
+[/OBJECT]
+
+[VAR:c1]
+
+![](TransformCatRef.png)
+
+[REGION:raw]
+| Property | %spacer% | Description |
+| --- | :---: | :--- |
+| **Location** || The location of the **Actor** or **Component** in **World Space** or **Relative** to its parent. |
+| **Rotation** || The rotation of the **Actor** or **Component** in **World Space** or **Relative** to its parent. |
+| **Scale** || The Scale of the **Actor** or **Component** in **World Space** or **Relative** to its parent. |
+[/REGION]
+
+[/VAR]
+[/EXCERPT]
+
+
+
+[EXCERPT:PropSocketsCat]
+[OBJECT:modulesection]
+	[PARAMLITERAL:id]
+		SocketsCat
+	[/PARAMLITERAL]
+	[PARAM:heading]
+		Sockets
+	[/PARAM]
+	[PARAM:content]
+		%c2%
+	[/PARAM]
+[/OBJECT]
+
+[VAR:c2]
+
+![](SocketsCatRef.png)
+
+[REGION:raw]
+| Property | %spacer% | Description |
+| --- | :---: | :--- |
+| **Parent Socket** || When a this component is the **Child** of a **Skeletal Mesh Component** (or **Static Mesh Component** with a **Socket**) you can specify a **Socket** or **Joint** to attach this component to. |
+[/REGION]
+
+[/VAR]
+[/EXCERPT]
+
+
+
+[EXCERPT:PropSplineCat]
+[OBJECT:modulesectionexpanded]
+	[PARAMLITERAL:id]
+		SplineCatd
+	[/PARAMLITERAL]
+	[PARAM:heading]
+		Spline
+	[/PARAM]
+	[PARAM:content]
+		%sp1%
+	[/PARAM]
+[/OBJECT]
+
+[VAR:sp1]
+
+![](BPS_SplineCatRef.png)
+
+[REGION:raw]
+| Property | %spacer% | Description |
+| --- | :---: | :--- |
+| **Duration** || %doxygen:USplineComponent::Duration% |
+| **Override Construction Script** || %doxygen:USplineComponent::bSplineHasBeenEdited% |
+| **Input Spline Points to Construction Script** || %doxygen:USplineComponent::bInputSplinePointsToConstructionScript% |
+| **Draw Debug** || %doxygen:USplineComponent::bDrawDebug% |
+| **Closed Loop** || %doxygen:USplineComponent::bClosedLoop% |
+| **Default Up Vector** || %doxygen:USplineComponent::DefaultUpVector% |
+[/REGION]
+
+### Advanced
+
+![](BPS_SplineCatRefAdv.png)
+
+[REGION:raw]
+| Property | %spacer% | Description |
+| --- | :---: | :--- |
+| **Reparam Steps Per Segment** || %doxygen:USplineComponent::ReparamStepsPerSegment% |
+| **Stationary Endpoints** || %doxygen:USplineComponent::bStationaryEndpoints% |
+[/REGION]
+
+[/VAR]
+[/EXCERPT]
+
+
+
+[EXCERPT:PropEditorCat]
+[OBJECT:modulesectionexpanded]
+	[PARAMLITERAL:id]
+		EditorCat
+	[/PARAMLITERAL]
+	[PARAM:heading]
+		Editor
+	[/PARAM]
+	[PARAM:content]
+		%sp2%
+	[/PARAM]
+[/OBJECT]
+
+[VAR:sp2]
+
+![](BPS_EditorCatRef.png)
+
+[REGION:raw]
+| Property | %spacer% | Description |
+| --- | :---: | :--- |
+| **Editor Unselected Spline Segment Color** || Color of an unselected spline component segment in the editor. |
+| **Editor Selected Spline Segment Color** || Color of a selected spline component segment in the editor. |
+| **Should Visualize Scale** || Whether or not scale visualization should be displayed in the editor. |
+| **Scale Visualization Width** || Width of the spline when scale visualization is enabled in the editor. |
+[/REGION]
+
+[/VAR]
+[/EXCERPT]
+
+
+
+[OBJECT:modulesection]
+	[PARAMLITERAL:id]
+		PhysicsCat
+	[/PARAMLITERAL]
+	[PARAM:heading]
+		Physics
+	[/PARAM]
+	[PARAM:content]
+		%Engine/Physics/PhysicsBodies/Reference:c3%
+	[/PARAM]
+[/OBJECT]
+
+
+
+
+[OBJECT:modulesection]
+	[PARAMLITERAL:id]
+		CollisionCat
+	[/PARAMLITERAL]
+	[PARAM:heading]
+		Collision
+	[/PARAM]
+	[PARAM:content]
+		%Engine/Physics/PhysicsBodies/Reference:p2%
+	[/PARAM]
+[/OBJECT]
+
+
+
+
+
+[EXCERPT:PropLightingCat]
+[OBJECT:modulesection]
+	[PARAMLITERAL:id]
+		LightingCat
+	[/PARAMLITERAL]
+	[PARAM:heading]
+		Lighting
+	[/PARAM]
+	[PARAM:content]
+		%c5%
+	[/PARAM]
+[/OBJECT]
+
+[VAR:c5]
+
+![](LightingCatRef.png)
+
+[REGION:raw]
+| Property | %spacer% | Description |
+| --- | :---: | :--- |
+| **Cast Shadow** || %doxygen:UPrimitiveComponent::CastShadow% |
+[/REGION]
+
+###Advanced
+
+![](LightingCatRefAdv.png)
+
+[REGION:raw]
+| Property | %spacer% | Description |
+| --- | :---: | :--- |
+| **Affect Dynamic Indirect Lighting** || %doxygen:UPrimitiveComponent::bAffectDynamicIndirectLighting% |
+| **Affect Distance Field Lighting** || %doxygen:UPrimitiveComponent::bAffectDistanceFieldLighting% |
+| **Cast Dynamic Shadow** || %doxygen:UPrimitiveComponent::bCastDynamicShadow% |
+| **Cast Static Shadow** || %doxygen:UPrimitiveComponent::bCastStaticShadow% |
+| **Cast Volumetric Translucent Shadow** || %doxygen:UPrimitiveComponent::bCastVolumetricTranslucentShadow% |
+| **Self Shadow Only** || %doxygen:UPrimitiveComponent::bSelfShadowOnly% |
+| **Cast Far Shadow** || %doxygen:UPrimitiveComponent::bCastFarShadow% |
+| **Cast Inset Shadow** || %doxygen:UPrimitiveComponent::bCastInsetShadow% |
+| **Cast Cinematic Shadow** || %doxygen:UPrimitiveComponent::bCastCinematicShadow% |
+| **Cast Hidden Shadow** || %doxygen:UPrimitiveComponent::bCastHiddenShadow% |
+| **Cast Shadow as Two Sided** || %doxygen:UPrimitiveComponent::bCastShadowAsTwoSided% |
+| **Light as if Static** || %doxygen:UPrimitiveComponent::bLightAsIfStatic% |
+| **Light Attachments as Group** || %doxygen:UPrimitiveComponent::bLightAttachmentsAsGroup% |
+| **Indirect Lighting Cache Quality** || %doxygen:UPrimitiveComponent::IndirectLightingCacheQuality% |
+| **Single Sample Shadow From Stationary Lights** || %doxygen:UPrimitiveComponent::bSingleSampleShadowFromStationaryLights% |
+| **Lighting Channles** || %doxygen:UPrimitiveComponent::LightingChannels% |
+[/REGION]
+
+[/VAR]
+[/EXCERPT]
+
+
+
+[EXCERPT:PropRenderingCat]
+[OBJECT:modulesection]
+	[PARAMLITERAL:id]
+		RenderingCat
+	[/PARAMLITERAL]
+	[PARAM:heading]
+		Rendering
+	[/PARAM]
+	[PARAM:content]
+		%c6%
+	[/PARAM]
+[/OBJECT]
+
+[VAR:c6]
+
+![](RenderingCatRef.png)
+
+[REGION:raw]
+| Property | %spacer% | Description |
+| --- | :---: | :--- |
+| **Visible** || %doxygen:USceneComponent::bVisible% |
+| **Hidden in Game** ||  %doxygen:USceneComponent::bHiddenInGame%|
+[/REGION]
+
+###Advanced
+
+![](RenderingCatRefAdv.png)
+
+[REGION:raw]
+| Property | %spacer% | Description |
+| --- | :---: | :--- |
+| **Texture Streaming** || [INCLUDE:#TextureStreamingProp] |
+| **LOD** || [INCLUDE:#LODProp] |
+| **Render In Main Pass** || %doxygen:UPrimitiveComponent::bRenderInMainPass% |
+| **Receives Decals** || %doxygen:UPrimitiveComponent::bReceivesDecals% |
+| **Owner No See** || %doxygen:UPrimitiveComponent::bOwnerNoSee% |
+| **Only Owner See** || %doxygen:UPrimitiveComponent::bOnlyOwnerSee% |
+| **Treat As Background for Occlusion** || %doxygen:UPrimitiveComponent::bTreatAsBackgroundForOcclusion% |
+| **Use As Occluder** || %doxygen:UPrimitiveComponent::bUseAsOccluder% |
+| **Render CustomDepth Pass** || %doxygen:UPrimitiveComponent::bRenderCustomDepth% |
+| **CustomDepth Stencil Value** || %doxygen:UPrimitiveComponent::CustomDepthStencilValue% |
+| **Translucency Sort Priority** || Translucent objects with a lower sort priority draw behind objects with a higher priority. Translucent objects with the same priority are rendered from back-to-front based on their bounds origin. |
+| **Lpv Bias Multiplier** || %doxygen:UPrimitiveComponent::LpvBiasMultiplier% |
+| **Bounds Scale** || %doxygen:UPrimitiveComponent::BoundsScale% |
+| **Use Attach Parent Bound** || %doxygen:USceneComponent::bUseAttachParentBound% |
+[/REGION]
+
+[/VAR]
+[/EXCERPT]
+
+<!--
+[EXCERPT:TextureStreamingProp]
+
+[REGION:raw]
+| Property | %spacer2% | Description |
+| --- | :---: | :--- |
+| **Force Mip Streaming** || %doxygen:UPrimitiveComponent::bForceMipStreaming% |
+[/REGION]
+
+[/EXCERPT]
+
+[EXCERPT:LODProp]
+
+[REGION:raw]
+| Property | %spacer2% | Description |
+| --- | :---: | :--- |
+| **Min Draw Distance** || %doxygen:UPrimitiveComponent::MinDrawDistance% |
+| **Desired Max Draw Distance** || %doxygen:UPrimitiveComponent::LDMaxDrawDistance% |
+| **Current Max Draw Distance** || %doxygen:UPrimitiveComponent::CachedMaxDrawDistance% |
+| **Allow Cull Distance Volume** || %doxygen:UPrimitiveComponent::bAllowCullDistanceVolume% |
+| **Detail Mode** || %doxygen:USceneComponent::DetailMode% |
+[/REGION]
+
+[/EXCERPT]
+-->
+
+[EXCERPT:PropTagsCat]
+[OBJECT:modulesection]
+	[PARAMLITERAL:id]
+		TagsCat
+	[/PARAMLITERAL]
+	[PARAM:heading]
+		Tags
+	[/PARAM]
+	[PARAM:content]
+		%c7%
+	[/PARAM]
+[/OBJECT]
+
+[VAR:c7]
+
+![](TagsCatRef.png)
+
+[REGION:raw]
+| Property | %spacer% | Description |
+| --- | :---: | :--- |
+| **Component Tags** || %doxygen:UActorComponent::ComponentTags% |
+[/REGION]
+
+[/VAR]
+[/EXCERPT]
+
+
+
+[EXCERPT:PropActivationCat]
+[OBJECT:modulesection]
+	[PARAMLITERAL:id]
+		ActivationCat
+	[/PARAMLITERAL]
+	[PARAM:heading]
+		Activation
+	[/PARAM]
+	[PARAM:content]
+		%c8%
+	[/PARAM]
+[/OBJECT]
+
+[VAR:c8]
+
+![](ActivationCatRef.png)
+
+[REGION:raw]
+| Property | %spacer% | Description |
+| --- | :---: | :--- |
+| **Auto Activate** || %doxygen:UActorComponent::bAutoActivate% |
+[/REGION]
+
+[/VAR]
+[/EXCERPT]
+
+
+
+[EXCERPT:PropEventsCat]
+[OBJECT:modulesection]
+	[PARAMLITERAL:id]
+		EventsCat
+	[/PARAMLITERAL]
+	[PARAM:heading]
+		Events
+	[/PARAM]
+	[PARAM:content]
+		%c9%
+	[/PARAM]
+[/OBJECT]
+
+[VAR:c9]
+
+![](EventsCatRef.png)
+
+[REGION:raw]
+| Property | %spacer% | Description |
+| --- | :---: | :--- |
+| **On Component Hit** || %doxygen:UPrimitiveComponent::OnComponentHit% |
+| **On Component Begin Overlap** || %doxygen:UPrimitiveComponent::OnComponentBeginOverlap% |
+| **On Component End Overlap** || %doxygen:UPrimitiveComponent::OnComponentEndOverlap% |
+| **On Component Wake** || %doxygen:UPrimitiveComponent::OnComponentWake% |
+| **On Component Sleep** || %doxygen:UPrimitiveComponent::OnComponentSleep% |
+| **On Begin Cursor Over** || %doxygen:UPrimitiveComponent::OnBeginCursorOver% |
+| **On End Cursor Over** || %doxygen:UPrimitiveComponent::OnEndCursorOver% |
+| **On Clicked** || %doxygen:UPrimitiveComponent::OnClicked% |
+| **On Released** || %doxygen:UPrimitiveComponent::OnReleased% |
+| **On Input Touch Begin** || %doxygen:UPrimitiveComponent::OnInputTouchBegin% |
+| **On Input Touch End** || %doxygen:UPrimitiveComponent::OnInputTouchEnd% |
+| **On Input Touch Enter** || %doxygen:UPrimitiveComponent::OnInputTouchEnter% |
+| **On Input Touch Leave** || %doxygen:UPrimitiveComponent::OnInputTouchLeave% |
+| **Physics Volume Changed** || Delegate that will be called when PhysicsVolume has been changed. |
+[/REGION]
+
+[/VAR]
+[/EXCERPT]
+
+
+
+[EXCERPT:PropMobileCat]
+[OBJECT:modulesection]
+	[PARAMLITERAL:id]
+		MobileCat
+	[/PARAMLITERAL]
+	[PARAM:heading]
+		Mobile
+	[/PARAM]
+	[PARAM:content]
+		%c10%
+	[/PARAM]
+[/OBJECT]
+
+[VAR:c10]
+
+![](MobileCatRef.png)
+
+[REGION:raw]
+| Property | %spacer% | Description |
+| --- | :---: | :--- |
+| **Receive Combined Static and CSM Shadows from Stationary Lights** || %doxygen:UPrimitiveComponent::bReceiveCombinedCSMAndStaticShadowsFromStationaryLights%  |
+[/REGION]
+
+[/VAR]
+[/EXCERPT]
+
+
diff --git a/Engine/Documentation/Source/Engine/BlueprintSplines/Reference/MeshComponentProperties/SplineMeshComponentReference.INT.udn b/Engine/Documentation/Source/Engine/BlueprintSplines/Reference/MeshComponentProperties/SplineMeshComponentReference.INT.udn
new file mode 100644
--- /dev/null
+++ b/Engine/Documentation/Source/Engine/BlueprintSplines/Reference/MeshComponentProperties/SplineMeshComponentReference.INT.udn
@@ -0,0 +1,215 @@
+Availability:Docs
+Title:Blueprint Spline Mesh Component Property Reference
+Crumbs: %ROOT%, Engine
+Description:Contains a reference of all properties available in the editor for Blueprint Spline Mesh Components.
+Navigation:topic
+parent:Engine/BlueprintSplines
+order:1
+social-image:Splines_social.png
+
+[VAR:Topic]
+[OBJECT:Topic]
+	[PARAM:image]
+		![%Engine/BlueprintSplines/Reference/MeshComponent:title%](Engine/BlueprintSplines/Splines_topic.png)
+	[/PARAM]
+	[PARAM:icon]
+		![](%ROOT%/editor_icon.png)(convert:false)
+	[/PARAM]
+	[PARAM:title]
+		%Engine/BlueprintSplines/Reference/MeshComponent:title%
+	[/PARAM]
+	[PARAM:description]
+		%Engine/BlueprintSplines/Reference/MeshComponent:description%
+	[/PARAM]
+	[PARAM:path]
+		[RELATIVE:Engine/BlueprintSplines/Reference/MeshComponent]
+	[/PARAM]
+[/OBJECT]
+[/VAR]
+[VAR:TopicCompact]
+[OBJECT:TopicCompact]
+	[PARAM:image]
+		![%Engine/BlueprintSplines/Reference/MeshComponent:title%](Engine/BlueprintSplines/Splines_topic.png)
+	[/PARAM]
+	[PARAM:icon]
+		![](%ROOT%/editor_icon.png)(convert:false)
+	[/PARAM]
+	[PARAM:title]
+		%Engine/BlueprintSplines/Reference/MeshComponent:title%
+	[/PARAM]
+	[PARAM:description]
+		%Engine/BlueprintSplines/Reference/MeshComponent:description%
+	[/PARAM]
+	[PARAM:path]
+		[RELATIVE:Engine/BlueprintSplines/Reference/MeshComponent]
+	[/PARAM]
+[/OBJECT]
+[/VAR]
+
+[VAR:spacer]![](%ROOT%/spacer.png)(w:150 convert:false)[/VAR]
+
+[VAR:spacer2]![](%ROOT%/spacer.png)(w:40 convert:false)[/VAR]
+
+This page contains a reference listing of properties available on **Spline Components**. The properties displayed will be slightly different if the **Blueprint Spline Component** is selected in the **Blueprint Editor**, is a **Root Component**, or if the component is selected in the **Level Editor**.
+
+## Properties
+
+[INCLUDE:Engine/BlueprintSplines/Reference/ComponentProperties#PropTransformCat]
+
+[INCLUDE:Engine/BlueprintSplines/Reference/ComponentProperties#PropSocketsCat]
+
+
+
+[EXCERPT:PropStaticMeshCat]
+[OBJECT:modulesectionexpanded]
+	[PARAMLITERAL:id]
+		StaticMeshCat
+	[/PARAMLITERAL]
+	[PARAM:heading]
+		Static Mesh
+	[/PARAM]
+	[PARAM:content]
+		%spm1%
+	[/PARAM]
+[/OBJECT]
+
+[VAR:spm1]
+
+![](StaticMeshCatRef.png)
+
+[REGION:raw]
+| Property | %spacer% | Description |
+| --- | :---: | :--- |
+| **Static Mesh** || Specifies the **Static Mesh** to render for this component. |
+
+[/REGION]
+
+###Materials
+
+![](MaterialsCatRef.png)
+
+[REGION:raw]
+| Property | %spacer% | Description |
+| --- | :---: | :--- |
+| **Element #** || Once a **Static Mesh** has been specified in the **Static Mesh Property**, a number of additional **Material Properties** will appear. These will be named based on the **Material IDs** applied to the **Static Mesh**. |
+[/REGION]
+
+[/VAR]
+[/EXCERPT]
+
+
+
+[EXCERPT:PropSplineMeshCat]
+[OBJECT:modulesectionexpanded]
+	[PARAMLITERAL:id]
+		SplineMeshCat
+	[/PARAMLITERAL]
+	[PARAM:heading]
+		Spline Mesh
+	[/PARAM]
+	[PARAM:content]
+		%spm2%
+	[/PARAM]
+[/OBJECT]
+
+[VAR:spm2]
+
+![](SplineMeshCat.png)
+
+[REGION:raw]
+| Property | %spacer% | Description |
+| --- | :---: | :--- |
+| **Start Pos** || %doxygen:FSplineMeshParams::StartPos% |
+| **Start Tangent** || %doxygen:FSplineMeshParams::StartTangent% |
+| **End Pos** || %doxygen:FSplineMeshParams::EndPos% |
+| **End Tangent** || %doxygen:FSplineMeshParams::EndTangent% |
+| **Spline Up Dir** || %doxygen:USplineMeshComponent::SplineUpDir% |
+| **Forward Axis** || %doxygen:USplineMeshComponent::ForwardAxis% |
+[/REGION]
+
+###Advanced
+
+![](SplineMeshCatAdv.png)
+
+[REGION:raw]
+| Property | %spacer% | Description |
+| --- | :---: | :--- |
+| **Start Scale** || %doxygen:FSplineMeshParams::StartScale% |
+| **Start Roll** || %doxygen:FSplineMeshParams::StartRoll% |
+| **Start Offset** || %doxygen:FSplineMeshParams::StartOffset% |
+| **End Scale** || %doxygen:FSplineMeshParams::EndScale% |
+| **End Roll** || %doxygen:FSplineMeshParams::EndRoll% |
+| **End Offset** || %doxygen:FSplineMeshParams::EndOffset% |
+| **Smooth Interp Roll Scale** || %doxygen:USplineMeshComponent::bSmoothInterpRollScale% |
+| **Spline Boundary Min** || %doxygen:USplineMeshComponent::SplineBoundaryMin% |
+| **Spline Boundary Max** || %doxygen:USplineMeshComponent::SplineBoundaryMax% |
+[/REGION]
+
+[/VAR]
+[/EXCERPT]
+
+
+
+[EXCERPT:PropSplineCat]
+[OBJECT:modulesectionexpanded]
+	[PARAMLITERAL:id]
+		SplineCat
+	[/PARAMLITERAL]
+	[PARAM:heading]
+		Spline
+	[/PARAM]
+	[PARAM:content]
+		%spm3%
+	[/PARAM]
+[/OBJECT]
+
+[VAR:spm3]
+
+![](BPSM_SplineCatRef.png)
+
+[REGION:raw]
+| Property | %spacer% | Description |
+| --- | :---: | :--- |
+| **Allow Spline Editing Per Instance** || %doxygen:USplineMeshComponent::bAllowSplineEditingPerInstance% |
+[/REGION]
+
+[/VAR]
+[/EXCERPT]
+
+
+
+[OBJECT:modulesection]
+	[PARAMLITERAL:id]
+		PhysicsCat
+	[/PARAMLITERAL]
+	[PARAM:heading]
+		Physics
+	[/PARAM]
+	[PARAM:content]
+		%Engine/Physics/PhysicsBodies/Reference:c3%
+	[/PARAM]
+[/OBJECT]
+
+[OBJECT:modulesection]
+	[PARAMLITERAL:id]
+		CollisionCat
+	[/PARAMLITERAL]
+	[PARAM:heading]
+		Collision
+	[/PARAM]
+	[PARAM:content]
+		%Engine/Physics/PhysicsBodies/Reference:p2%
+	[/PARAM]
+[/OBJECT]
+
+[INCLUDE:Engine/BlueprintSplines/Reference/ComponentProperties#PropLightingCat]
+
+[INCLUDE:Engine/BlueprintSplines/Reference/ComponentProperties#PropRenderingCat]
+
+[INCLUDE:Engine/BlueprintSplines/Reference/ComponentProperties#PropTagsCat]
+
+[INCLUDE:Engine/BlueprintSplines/Reference/ComponentProperties#PropActivationCat]
+
+[INCLUDE:Engine/BlueprintSplines/Reference/ComponentProperties#PropEventsCat]
+
+[INCLUDE:Engine/BlueprintSplines/Reference/ComponentProperties#PropMobileCat]
diff --git a/Engine/Documentation/Source/Engine/BlueprintSplines/Reference/SplineEditorTool/SplineEditorToolReference.INT.udn b/Engine/Documentation/Source/Engine/BlueprintSplines/Reference/SplineEditorTool/SplineEditorToolReference.INT.udn
new file mode 100644
--- /dev/null
+++ b/Engine/Documentation/Source/Engine/BlueprintSplines/Reference/SplineEditorTool/SplineEditorToolReference.INT.udn
@@ -0,0 +1,128 @@
+Availability:Docs
+Title:Blueprint Spline Editor Tool Reference
+Crumbs: %ROOT%, Engine
+Description:A quick reference of the tools and procedures for editing Blueprint Spline Components while in the Level Editor.
+Navigation:topic
+parent:Engine/BlueprintSplines
+order:1
+social-image:Splines_social.png
+
+[VAR:Topic]
+[OBJECT:Topic]
+	[PARAM:image]
+		![%Engine/BlueprintSplines/Reference/SplineEditorTool:title%](Engine/BlueprintSplines/Splines_topic.png)
+	[/PARAM]
+	[PARAM:icon]
+		![](%ROOT%/editor_icon.png)(convert:false)
+	[/PARAM]
+	[PARAM:title]
+		%Engine/BlueprintSplines/Reference/SplineEditorTool:title%
+	[/PARAM]
+	[PARAM:description]
+		%Engine/BlueprintSplines/Reference/SplineEditorTool:description%
+	[/PARAM]
+	[PARAM:path]
+		[RELATIVE:Engine/BlueprintSplines/Reference/SplineEditorTool]
+	[/PARAM]
+[/OBJECT]
+[/VAR]
+[VAR:TopicCompact]
+[OBJECT:TopicCompact]
+	[PARAM:image]
+		![%Engine/BlueprintSplines/Reference/SplineEditorTool:title%](Engine/BlueprintSplines/Splines_topic.png)
+	[/PARAM]
+	[PARAM:icon]
+		![](%ROOT%/editor_icon.png)(convert:false)
+	[/PARAM]
+	[PARAM:title]
+		%Engine/BlueprintSplines/Reference/SplineEditorTool:title%
+	[/PARAM]
+	[PARAM:description]
+		%Engine/BlueprintSplines/Reference/SplineEditorTool:description%
+	[/PARAM]
+	[PARAM:path]
+		[RELATIVE:Engine/BlueprintSplines/Reference/SplineEditorTool]
+	[/PARAM]
+[/OBJECT]
+[/VAR]
+
+[VAR:spacer]![](%ROOT%/spacer.png)(w:150 convert:false)[/VAR]
+
+[VAR:spacer2]![](%ROOT%/spacer.png)(w:40 convert:false)[/VAR]
+
+This page contains information on tools and procedures for editing **Blueprint Spline Components** and **Blueprint Spline Mesh Components** in the Level Editor.
+
+## Tools
+
+With a Spline Point (selecting a Spline Curve selects the last point on it) selected in the editor, you can use the standard manipulation tools (move, rotate, and scale) to adjust positions and tangents of the points.
+
+| Tool | Result |
+| --- | :--- |
+| **Move Selected Point** | ![](MovePoint.gif)(convert:false) |
+| **Move Selected Tangent** | ![](MoveTangent.gif)(convert:false) |
+| **Rotate** | ![](RotatePoint.gif)(convert:false) |
+| **Scale** | ![](ScalePoint.gif)(convert:false) |
+
+
+## Procedures
+
+There are a few menus accessable by right-clicking on the spline either on a point, on the curve inbetween points.
+
+### Right-Click on a Spline Point
+
+Right-clicking on a **point** will give you the following menu and options.
+
+![](Spline_RC_Point.png)
+
+[REGION:raw]
+| Option | %spacer% | Description |
+| --- | :---: | :--- |
+| **Delete Spline Point** || Removes the selected Spline Point. |
+| **Duplicate Spline Point** || Duplicates the currently selected Spline Point. |
+| **Spline Point Type** || [INCLUDE:#SplinePointType] |
+| **Reset to Automatic Tangent** || [INCLUDE:#ResetToAutomaticTanget] |
+| **Reset to Default** || Resets the Spline to its archetype default. |
+| **Visualize Roll and Scale** || Whether or not the spline should visualize roll and scale in the editor. |
+[/REGION]
+
+<!--
+[EXCERPT:SplinePointType]
+
+Changes the interpolation style for the selected Spline Point. The Spline Curve associated with the Spline Curve point is displayed in red (by default). It will be the curve affected by changing this. The default **Point Type** is **Curve**. Below are examples of what 
+
+![](Spline_RC_Point_Type.png)
+
+| Option | %spacer2% | Description |
+| --- | :---: | :--- |
+| **Curve** || ![](PointTypeCurve.png)(convert:false) |
+| **Linear** || ![](PointTypeLinear.png)(convert:false) |
+| **Constant** || ![](PointTypeConstant.png)(convert:false) |
+[/EXCERPT]
+
+[EXCERPT:ResetToAutomaticTanget]
+
+![](Spline_RC_Point_Reset.png)
+
+| Option | %spacer2% | Description |
+| --- | :---: | :--- |
+| **Unclamped Tangent** || Resets the tangents of the selected **Curve Point** to their **Unclamped Tangent** values. |
+| **Clamped Tangent** || Resets the tangents of the selected **Curve Point** to their **Clamped Tangent** values. |
+
+[/EXCERPT]
+-->
+
+### Right-Click on a Spline Curve
+
+Right-clicking on a **curve** will give you the following menu and options.
+
+![](Spline_RC_Curve.png)
+
+[REGION:raw]
+| Option | %spacer% | Description |
+| --- | :---: | :--- |
+| **Add Spline Point Here** || Adds a new Spline Point at the cursor's location. |
+| **Reset to Default** || Resets the Spline to its archetype default. |
+| **Visualize Roll and Scale** || Whether or not the spline should visualize roll and scale in the editor. |
+[/REGION]
+
+
diff --git a/Engine/Documentation/Source/Engine/BlueprintSplines/SplinesLanding.INT.udn b/Engine/Documentation/Source/Engine/BlueprintSplines/SplinesLanding.INT.udn
new file mode 100644
--- /dev/null
+++ b/Engine/Documentation/Source/Engine/BlueprintSplines/SplinesLanding.INT.udn
@@ -0,0 +1,64 @@
+Availability:Docs
+Title:Blueprint Splines
+Crumbs: %ROOT%, Engine
+Description:Landing page for Blueprint Spline and Spline Mesh Components.
+Navigation:topic
+parent:Engine
+order:3
+social-image:Splines_social.png
+related: Engine/Landscape/Editing/Splines/
+
+[VAR:Topic]
+[OBJECT:Topic]
+	[PARAM:image]
+		![%Engine/Splines:title%](Engine/Splines/Splines_topic.png)
+	[/PARAM]
+	[PARAM:icon]
+		![](%ROOT%/editor_icon.png)(convert:false)
+	[/PARAM]
+	[PARAM:title]
+		%Engine/Splines:title%
+	[/PARAM]
+	[PARAM:description]
+		%Engine/Splines:description%
+	[/PARAM]
+	[PARAM:path]
+		[RELATIVE:Engine/Splines]
+	[/PARAM]
+[/OBJECT]
+[/VAR]
+[VAR:TopicCompact]
+[OBJECT:TopicCompact]
+	[PARAM:image]
+		![%Engine/Splines:title%](Engine/Splines/Splines_topic.png)
+	[/PARAM]
+	[PARAM:icon]
+		![](%ROOT%/editor_icon.png)(convert:false)
+	[/PARAM]
+	[PARAM:title]
+		%Engine/Splines:title%
+	[/PARAM]
+	[PARAM:description]
+		%Engine/Splines:description%
+	[/PARAM]
+	[PARAM:path]
+		[RELATIVE:Engine/Splines]
+	[/PARAM]
+[/OBJECT]
+[/VAR]
+
+![](Engine/BlueprintSplines/Overview/BPSC_1.png)
+
+Below are links to the documentation pages that cover **Blueprint Spline Components** and **Blueprint Spline Mesh Components**.
+
+## Essentials
+
+[DIR(output:"listbutton" parent:"Engine/BlueprintSplines" org:"hierarchy" end:"1")]
+
+
+## Samples
+
+[DIR(output:"topiccompact" parent:"Resources/ContentExamples" tags:"Splines")]
+
+
+
diff --git a/Engine/Documentation/Source/Engine/Physics/PhysicsBodies/Reference/PhysicsBodiesRef.INT.udn b/Engine/Documentation/Source/Engine/Physics/PhysicsBodies/Reference/PhysicsBodiesRef.INT.udn
--- a/Engine/Documentation/Source/Engine/Physics/PhysicsBodies/Reference/PhysicsBodiesRef.INT.udn
+++ b/Engine/Documentation/Source/Engine/Physics/PhysicsBodies/Reference/PhysicsBodiesRef.INT.udn
@@ -49,20 +49,20 @@ This page contains a reference listing of properties in the Physics and Collisio
 
 Below are the properties for Physics Bodies (Body Instances) separated by major category.
 
-
+[EXCERPT:PropPhysicsCat]
 [OBJECT:modulesectionexpanded]
 	[PARAMLITERAL:id]
-		phys
+		PhysicsCat
 	[/PARAMLITERAL]
 	[PARAM:heading]
 		Physics
 	[/PARAM]
 	[PARAM:content]
-		%p1%
+		%c3%
 	[/PARAM]
 [/OBJECT]
 
-[VAR:p1]
+[VAR:c3]
 
 ![](physicsProps.png)
 
@@ -107,7 +107,7 @@ Below are the properties for Physics Bodies (Body Instances) separated by major
 [/REGION]
 
 [/VAR]
-
+[/EXCERPT]
 
 [OBJECT:modulesectionexpanded]
 	[PARAMLITERAL:id]
diff --git a/Engine/Documentation/Source/Resources/ContentExamples/Blueprint_Splines/Blueprint_Splines.INT.udn b/Engine/Documentation/Source/Resources/ContentExamples/Blueprint_Splines/Blueprint_Splines.INT.udn
--- a/Engine/Documentation/Source/Resources/ContentExamples/Blueprint_Splines/Blueprint_Splines.INT.udn
+++ b/Engine/Documentation/Source/Resources/ContentExamples/Blueprint_Splines/Blueprint_Splines.INT.udn
@@ -4,6 +4,7 @@ Crumbs:%ROOT%, Resources, Resources/ContentExamples
 Description:Overview of the samples provided in the Blueprint Splines Content Examples.
 Related: Engine/Rendering/Materials
 Version: 4.9
+tags:Splines
 
 [REGION:banner]
 ![](blueprintSplines_ce_header.png)(convert:false)
