- applying peer review feedback.

diff --git a/Engine/Documentation/Source/Engine/Sequencer/Workflow/SequenceRecorder/SequenceRecorder.INT.udn b/Engine/Documentation/Source/Engine/Sequencer/Workflow/SequenceRecorder/SequenceRecorder.INT.udn
--- a/Engine/Documentation/Source/Engine/Sequencer/Workflow/SequenceRecorder/SequenceRecorder.INT.udn
+++ b/Engine/Documentation/Source/Engine/Sequencer/Workflow/SequenceRecorder/SequenceRecorder.INT.udn
@@ -9,7 +9,8 @@ Version: 4.13
 Related: Engine\Sequencer\Overview
 Related: Engine/Sequencer/HowTo
 
-The **Sequence Recorder** allows you to record Actor(s) from gameplay or level editing into Skeletal Animation, Transform, Event and Audio Tracks into a Level Sequence which can then be edited inside the **Sequencer Editor**. 
+The **Sequence Recorder** allows you to record Actor(s) from gameplay or level editing into Skeletal Animation, Transform, Event and Audio Tracks inside a Level Sequence. 
+You can then edit the Level Sequence inside the **Sequencer Editor** to create a cinematic based around your recording. 
 
 [TOC(start:2 end:2)]
 
@@ -26,21 +27,22 @@ After selecting Sequence Recorder, the Sequence Recorder window will open:
 
 ## Tracking Actors for Recording
 
-When you first open the Sequence Recorder, you will notice three buttons across the top of the window: 
+Across the top of the Sequence Recorder are three buttons: Record, Add and Remove. 
 
 ![](RecorderToolbar1.png)
 
 These are used to start recording, add Actors to track for recording or remove Actors from the recording list. 
 After clicking the **Add** button, both the **Record** option as well as the **Remove** option will become accessible. 
-This will also add an **Actor Recording** field to the Actor Tracking window with the Actor to track listed as **None**. 
+This will also add an **Actor Recording** field to the Actor Tracking window with the Actor to track listed as **None** (shown below). 
 
 ![](RecorderToolbar2.png)
 
-Selecting the Actor Recording in the window will populate the Sequence Recorder with additional options related to that recording.
+Selecting the Actor Recording will populate the Sequence Recorder with additional options (highlighted below) related to that recording.
 
 ![](SequenceRecorderWindow.png)
 
-Above we have set the **Actor to Track** to the **ThirdPersonCharacter** Actor (which is our playable character). 
+With these options you can define what Actor to track during the recording along with its animation settings. 
+You can add multiple Actor Tracking fields to the Tracking Window, with each **Actor to Record** set to track different Actors. 
 
 [REGION:note]
 Please see the [](Engine\Sequencer\HowTo\SequenceRecorder) page for a step-by-step guide to recording Actors with the Sequence Recorder. 
@@ -111,8 +113,8 @@ Please note that this will initiate Editor Recording and not [Gameplay Recording
 
 ### Subscene Recording Gameplay
 
-While it is perfectly okay to use the Sequence Recorder to assign an Actor and initiate gameplay recording, there is another method that you can take to record your gameplay using Sequence Recorder. 
-If you create a [Subscene Track](Engine\Sequencer\HowTo\TracksSub) inside a Level Sequence, you can add a primed **Recorded Sequence** on the Subscene Track during gameplay set to record your playable character.
+Using the Sequence Recorder to record gameplay will (by default) automatically create a new Level Sequence to record the content to. 
+If you already have a Level Sequence that you want to add recorded content to, you can do so with the use of a [Subscene Track](Engine\Sequencer\HowTo\TracksSub) and recording the current player pawn. 
 
 [OBJECT:EmbeddedVideo]
 [PARAMLITERAL:width]
@@ -129,6 +131,10 @@ Wab8bB69TIo
 Above we create a new Level Sequence along with a Subscene Track, then during gameplay we hit **Shift + F1** to gain Mouse Control and add a primed Recorded Sequence based on our current player pawn. 
 When can then press the **Record** button in the viewport while playing our game to start recording and hit **Esc** to end recording. 
 
+[REGION:note]
+Recording in this fashion is the same as manually adding the **Actor to Record** through the Sequence Recorder window and starting to Record from within the Sequence Recorder.
+One thing to keep in mind is that the settings defined inside the Sequence Recorder will be used when using this method of Subscene Recording. 
+[/REGION]
 
 ### Recording Class Properties
 
@@ -138,7 +144,20 @@ You can use the **Classes and Properties to Record** setting in the Sequence Rec
 ![](PropertyRecording.png)
 
 If that property changes at all during recording it will automatically be keyframed reflecting the changes during the recording. 
-If the Actor you are tracking does not contain one of the defined classes it will just be ignored. There are also some default elements added to this option which you can use or delete. 
+If the Actor you are tracking does not contain one of the defined classes it will just be ignored. 
+
+By default, Sequence Recorder includes the following classes to track for recording: 
+
+* SkeletalMeshComponent
+* StaticMeshComponent
+* ParticleSystemComponent
+* LightComponent
+* CameraComponent
+* CineCameraComponent
+
+You can use these classes and track properties within them during your recording or you can delete them and use your own classes. 
+
+
 
 ### Recording to Possessable
 
@@ -162,3 +181,4 @@ The **Immersive Mode** option inside the Sequence Recorder window will automatic
 
 Once you stop recording, the viewport will return back to its normal size. 
 
+
