#UE4 Docs: Merging 4.13 release notes from depot

diff --git a/Engine/Documentation/Source/Support/Builds/ReleaseNotes/2016/4_13/ReleaseNotes_4_13.INT.udn b/Engine/Documentation/Source/Support/Builds/ReleaseNotes/2016/4_13/ReleaseNotes_4_13.INT.udn
new file mode 100644
--- /dev/null
+++ b/Engine/Documentation/Source/Support/Builds/ReleaseNotes/2016/4_13/ReleaseNotes_4_13.INT.udn
@@ -0,0 +1,4448 @@
+Availability:Public
+Title:Unreal Engine 4.13 Release Notes
+Crumbs:%ROOT%, Engine
+Description:
+order:13
+
+Unreal Engine 4.13 Release Notes
+
+This release brings hundreds of updates for Unreal Engine 4, including 145 improvements submitted by the community of Unreal Engine developers on GitHub! Thanks to all of these contributors to Unreal Engine 4.13:
+
+[REDACTED_CONTRIBUTORS]
+
+# What’s New
+
+**Unreal Engine 4.13** has arrived!  In this version you'll find numerous improvements across the board.
+
+Many new **rendering features** have been added, such as mesh decals, Blueprint drawing to render targets, GPU morph targets, refraction improvements and high quality, optimized noise functions are now available to materials.  Shadow map caching allows for more shadow-casting dynamic lights in a scene than ever before!
+
+**Sequencer**, our new non-linear cinematic editor, has been updated with a slew of new features for high-end cinematography.  **Live recording** from gameplay has been significantly improved.  Also, you can now transfer shots and animations back and forth from external applications.  You can see these features in our [Siggraph Real-time Live! 2016 demonstration](http://www.unrealengine.com/blog/unreal-engine-4-powers-real-time-cinematography-at-siggraph).
+
+**Alembic** support allows you to import complex and interesting vertex animations.  And the new **Physical Animation Component** lets your characters respond realistically to physical forces by driving their skeletal animation through motors.
+
+For **mobile developers**, dynamic shadows have been optimized, full-precision materials are supported, and custom post-processing is now possible.  OpenGL ES 3.1 can now be used on Android, and binary shader caching will improve your iteration times.
+
+**VR games** can now use capsule shadows, and stereo instancing has been optimized.  Oh, and check out the new **VR Template Project**!  It's a great example of how to use motion controllers to interact and navigate in your VR game.  
+
+Want to build your levels while in VR?  Unreal's **VR Editor** has been improved with support for mesh and foliage painting, a new transform gizmo, and VR color picking.  Finally, you now can instantly play your game right from VR!  You can turn on "Enable VR Editing" in the Experimental section of your Editor Preferences.
+
+# Major Features
+
+## New: Sequencer Live Recording
+
+Sequencer been updated with new Live Recording improvements as shown at [Siggraph Real-Time Live! 2016](https://www.unrealengine.com/blog/unreal-engine-4-powers-real-time-cinematography-at-siggraph).
+
+![image alt text](image_1.jpg)(h:674 w:929 convert:false)
+
+The Sequence Recorder allows you to capture live gameplay, including all animation, audio and effects into a standalone asset that you can edit using Sequencer!  New features in this release:
+
+* Quickly record selected actors and automatically create camera cuts track if camera is recorded.
+
+* Ability to specify arbitrary components and properties to record.
+
+* Optionally record data to actor possessed in level.
+
+* Record transforms in world space when actor is attached but not recorded.
+
+## New: Shadow Map Caching for Movable Lights
+
+When a point or spot light is not moving, we can store off the shadow map for that light and reuse it next frame.  This is now done automatically and makes shadow casting movable point and spot lights much more affordable in games where the environment is often not moving.
+
+![image alt text](image_2.jpg)(h:541 w:929 convert:false)
+
+The above image shows 33 dynamic shadow-casting point lights, with very minimal overhead:
+
+Performance results on a 970 GTX at 1920x1200
+
+* 33 shadow casting point lights without caching: **14.89ms to render Shadow Depths.**
+
+* With cached shadow maps: **.9ms** **(about 16 times faster!)**
+
+* Note that it still costs 2ms to render the contributions of the 33 point lights, which can be optimized in other ways but is not affected by this change.
+
+* Memory used by the cache can be seen under ‘Stat ShadowRendering’ and was 25.6Mb in this scene.
+
+* Max memory used by the cache can be controlled with ‘r.Shadow.WholeSceneShadowCacheMb’
+
+Limitations
+
+* By default, caching can only happen when:
+
+    * Primitives have mobility set to Static or Stationary 
+
+    * Materials used do not use World Position Offset
+
+    * Light is point or spot, shadow casting, and mobility set to Movable but not currently moving
+
+* Materials that use an animating Tessellation or Pixel Depth Offset can cause artifacts as their shadow depths are cached.
+
+## New: Voronoi Noise Materials
+
+We’ve added a new Voronoi noise option available for the Noise material node. Voronoi noise, also sometimes called Worley or Cellular noise, is can be useful for procedural material creation. 
+
+![image alt text](image_3.jpg)(h:333 w:624 convert:false)
+
+Voronoi noise can be used to generate patterns for familiar physical materials such as marble, as seen on the statue below.
+
+![image alt text](image_4.jpg)(h:377 w:624 convert:false)
+
+This example uses a Voronoi noise with a technique called ‘gradient mapping’ to achieve a marble look.
+
+![image alt text](image_5.jpg)(h:233 w:929 convert:false)
+
+From Left to Right:
+
+1) Standard Voronoi Noise, 1 octave
+
+2) Voronoi with "Gradient" noise added to input position, set to 0.05
+
+3) Gradient noise multiplied by 0.3 before adding to Voronoi input position
+
+4) Using result of step 3 as texture coordinates for a random tiling texture
+
+The Voronoi noise has four quality levels, with decreasing grid artifacts at the higher levels at the cost of significantly increased shading time.
+
+Also, performance has been improved for several of the Noise Material Node features, with more detailed description of the performance tradeoffs in the function selection tooltips. Most of these can be slow for runtime use, so baking the results into a texture is encouraged.
+
+## New: Blueprint Drawing to Render Targets
+
+Blueprint functions can now be used to draw materials into render targets.  This enables a huge variety of game-specific rendering effects to be implemented without having to modify source code.
+
+![image alt text](image_6.gif)(convert:false)
+
+This is a fluid surface simulation implemented entirely in Blueprint and Material graphs.  Characters and projectiles can push the fluid around!
+
+![image alt text](image_7.gif)(convert:false)
+
+Above shows a simple heightfield painter made entirely in Blueprint, by accumulating a height value at projectile impacts.
+
+The new Blueprint function **Draw Material to Render Target** draws a quad filling the destination render target with the **Emissive Color** input in the material.
+
+![image alt text](image_8.jpg)(h:149 w:661 convert:false)
+
+Upon starting the game, **Begin Play** is called and the render target is filled with blue.  You can then right-click on the render target and save it as a static Texture which can be compressed.  
+
+![image alt text](image_9.jpg)(h:164 w:500 convert:false)
+
+For more advanced drawing to a render target, use **Begin Draw Canvas to Render Target **and **End Draw Canvas to Render Target**.  These allow multiple draws to a subset of the render target efficiently, as well as font drawing methods through the Canvas object.
+
+![image alt text](image_10.jpg)(h:517 w:896 convert:false)
+
+Being able to render off-screen from blueprint enables a ton of rendering features to be implemented quickly without a graphics programmer.  It's also very easy to tank the GPU by doing a lot of pixel shader work with many passes to large render targets.  These draw calls show up after typing 'ProfileGPU' in the debug console (under the WorldTick event)
+
+Limitations
+
+* You cannot draw to a render target that is being sampled as a texture by the material you specify.  Either use alpha blending to modify a render target in-place, or ping-pong between two different render targets.
+
+* Only the Emissive Color and Opacity outputs of the material are valid when drawing to a render target, lighting is not supported and nodes like WorldPosition may have unexpected values.
+
+* Emissive Color is clamped to be positive by default, but you can output negative values by enabling the material property ‘AllowNegativeEmissiveColor’
+
+Check out the **BlueprintRenderToTarget **map in the ContentExamples project for working examples!
+
+## New: Alembic Importer for Vertex Animation (Experimental)
+
+Alembic animation import is now supported!  Alembic allows for complex animations to be authored offline, then rendered in real-time inside UE4!  This feature is still considered experimental, but please try it out and send us feedback.
+
+![image alt text](image_11.gif)(convert:false)
+
+We allow importing Alembic caches in different ways:
+
+* **Static Mesh**.  A single frame from the Alembic animation will be imported as a static mesh asset (no animation.)
+
+* **Geometry Cache.  **This is a new type of animation asset that allows playback of vertex-varying sequences.  The imported Alembic animation will be played back as a flipbook of frames.  Performance will scale with your mesh’s complexity and may not be optimal in all cases
+
+* **Skeletal Mesh**.  This is the most efficient way to play back an Alembic animation, as long as the vertex count doesn’t change.  During import, your animation sequence will be compressed using a PCA scheme, in which common poses (bases) are extracted and weighted to compose the original animation during playback time. The percentage or fixed number of bases used can be set during import to tweak the level of compression
+
+## New: Mesh Preview Scenes
+
+New functionality has been added to setup the scene used to preview Static and Skeletal meshes. 
+
+![image alt text](image_12.jpg)(h:232 w:929 convert:false)
+
+A new Preview Scene Settings panel has been added to the Static Mesh and Skeletal Mesh editors.  Here, you can setup multiple profiles (scenes) to preview your meshes, and the profiles allow for changing:
+
+* Directional light (colour, intensity, rotation)
+
+* Sky light (HDRI environment map, intensity, rotation)
+
+* Post processing settings (identical to post process volume)
+
+![image alt text](image_13.jpg)(h:494 w:766 convert:false)
+
+We’ve also added some showcase functionality:
+
+* Manually rotate environment (hold K key) and directional light (hold L key)
+
+* Automatically rotate lighting rig (directional light and environment)
+
+* Easily hide the floor (O key) and environment (I key)
+
+## New: Mesh Decals
+
+The new Mesh Decals feature lets you efficiently splat materials on top of your static meshes, effectively allowing you to smoothly layer different materials on top of one another.  You can think of it as a second mesh that sits above the underlying geometry’s profile with it’s own topology and materials. 
+
+![image alt text](image_14.jpg)(h:326 w:428 convert:false)
+
+The above pillars are each created as base mesh overlayed with a single detailed break mesh, as shown in the wireframe below. The left pillar is using a mesh decal to allow smoothly blended color, normals and roughness. The right-most pillar is a masked material, just for comparison.  
+
+![image alt text](image_15.jpg)(h:220 w:165 convert:false)![image alt text](image_16.jpg)(h:220 w:169 convert:false)
+
+Unlike deferred decals there is no projection involved, so a typical decal mesh which is tightly coupled to the underlying surface may need to include a surface offset in the material.  Also, be careful if you have LODs where the mesh decal geometry would interpenetrate the LOD mesh.
+
+## New: Widget Interaction Component
+
+Using the widget interaction component, you can now simulate hardware input events with widget components in the world.  
+
+![image alt text](image_17.gif)(convert:false)
+
+You can attach it like a laser pointer to any object in the world to interact with widgets; there are also some other options available for more customized behavior.  When standard input comes to the player controller, you’ll instruct the interaction component to simulate a particular hardware input such as Left Mouse Down/Up over whatever widget the user happens to be hovering at that time.
+
+For users who were previously relying on clicking directly on Widget Components in the world using the mouse, that path is no longer supported.  Users will need to attach an Interaction Component (In Mouse Mode) to their player, then forward input to the Interaction Component when the player receives it.
+
+## New: VR Project Template
+
+We’ve added a new project template designed for Virtual Reality on desktop and console.
+
+![image alt text](image_18.gif)(convert:false)
+
+To access this, simply choose the new VR option in the New Project template selection window.
+
+This Blueprint project has settings that are optimized to run in VR up to 90 frames per second.  To support different types of controllers the template includes two methods of locomotion, split into two distinct Pawn Blueprints. The first is designed for gamepads while the second supports motion controllers.  When using motion controllers, you can teleport to different locations and grab and throw objects.  A C++ version of this template will follow in a future update.  Also, mobile VR templates will be coming later too.
+
+## New: Custom Post-Process for Mobile
+
+Custom Post-Process materials now can be used on Mobile devices!  Here is a "TV Static" effect displayed using the mobile renderer.
+
+![image alt text](image_20.gif)(convert:false)
+
+* This feature requires the "**Mobile HDR**" option to be enabled in your Project Settings
+
+* Supports fetching from PostProcessInput0 (SceneColor) with blendable location ‘Before Tonemapping’ and ‘After Tonemapping’
+
+* This feature does not currently work on older Android devices which require the ‘mosaic’ mode for HDR rendering.
+
+* Pixel depth information is not yet supported.
+
+## New: Lighting Channels on Mobile
+
+Lighting channels now work in the mobile renderer!  These allows you to selectively influence objects with specific lights.  Great for cinematics of advanced lighting rigs.
+
+![image alt text](image_21.jpg)(h:224 w:793 convert:false)
+
+* Multiple directional lights are supported in different channels.
+
+* Each primitive can be affected by only one directional light. The first lighting channel set on a primitive determines which directional light will affect it.
+
+* CSM shadows from stationary or movable directional lights cast only on primitives with matching lighting channels.
+
+* Dynamic point lights fully support lighting channels.
+
+## New: Shader Model 5 Rendering for Mac
+
+Mac Metal now has initial Shader Model 5 support enabled by default. This exposes all the available features of Metal on Mac OS X 10.11.6 that are applicable to Unreal Engine 4. 
+
+![image alt text](image_22.jpg)(h:521 w:925 convert:false)
+
+* Implements the RHI thread & parallel translation features to parallelise render command dispatch.
+
+* Exposes support for Metal compute shaders.
+
+* Exposes asynchronous compute support on AMD GPUs.
+
+* Enables high-end rendering features previously unavailable on Mac, including:
+
+    * *High quality dynamic exposure (a.k.a. Eye Adaptation).*
+
+    * *Compute-shader reflection environments - only available on discrete GPUs for 4.13.*
+
+    * *DistanceField Ambient Occlusion - only available on discrete GPUs for 4.13.*
+
+    * *DistanceField Shadowing - only available on discrete GPUs for 4.13.*
+
+## New: Physical Animation Component (Experimental)
+
+We’ve added a physical animation component that allows you to easily drive skeletal mesh animation through physical motors! 
+
+![image alt text](image_23.gif)(convert:false)
+
+The component allows you to set motor strengths directly, as well as using pre-configured physical animation profiles which can be created and edited inside PhAT.  The new "**Physical Animation Profiles**" feature in PhAT provide a powerful way to customize character physics in different game contexts, as well as fine tune it for special animations. 
+
+You can create and edit different profiles within the PhAT tool, and then change between them easily at runtime.  Check out the new "Apply Physical Animation Profile" and “Apply Physical Animation Settings” functions in Blueprints, which allow you to to change physical animation behavior dynamically.
+
+## New: Procedural Mesh Slicing
+
+There is a new utility in Procedural Mesh Component which will ‘slice’ a Procedural Mesh at runtime using a plane.
+
+![image alt text](image_25.gif)(convert:false)
+
+After slicing, we support adding ‘capping’ geometry, and creating a second Procedural Mesh Component for the ‘other half’ if desired. 
+
+Also, Procedural Mesh now supports simple collision, so physics simulation can be enabled!  Finally, we added a utility to copy data from a Static Mesh to a Procedural Mesh (‘Allow CPU Access’ flag must be set on the Static Mesh for this to work in cooked builds.) 
+
+## New: Mesh Painting in VR
+
+You can now paint on textures and mesh vertices in using the VR Editor.
+
+![image alt text](image_26.gif)(convert:false)
+
+This allows you to use motion controllers to paint on static meshes while immersed in VR.    To use this feature, open the "Modes" window in VR, then click the “Mesh Paint” tab.  Now simply select an object in the world, then aim and pull your trigger to paint!  Pressure sensitivity is supported on your controller’s trigger, and you can hold the ‘Modifier’ button to erase instead of paint.
+
+## New: Foliage Painting in VR
+
+In this release, you can use motion controllers to spray down foliage instances while in VR.
+
+![image alt text](image_27.gif)(convert:false)
+
+Select a foliage type, aim your laser and pull the trigger to paint foliage! You can hold the modifier button to erase foliage. Pressure sensitivity on the trigger is supported. Additionally, the Foliage Editor has been improved to be a bit more functional while in VR, although some features, lasso tool and select tool, are still unavailable.
+
+## New: Color Picker in VR
+
+The Color Picker window is now available in VR, so you can change color properties on lights and other Actors in your level. You can also use the Color Picker to select colors for Vertex Painting and Texture Painting in VR. 
+
+![image alt text](image_28.gif)(convert:false)
+
+## New: Play from VR Editor
+
+To easily prototype your project it is now possible to play your game in VR from within the VR Editor. Press the "Play" button on the quick menu to start playing in VR!  To instantly go back to the VR Editor, *hold the Grip buttons on both controllers and squeeze both trigger buttons*.
+
+![image alt text](image_29.gif)(convert:false)
+
+## New: Improved VR Transform Gizmo
+
+The VR Editor’s transform gizmo has been improved with better usability and new features!
+
+![image alt text](image_30.jpg)(h:442 w:466 convert:false)
+
+Translating and rotating objects feels much more natural, and you can now uniformly scale objects or translate them along a single 2D plane in VR.  We’ll continue to make improvements to VR gizmos in future releases.
+
+## New: VR Editor Flashlight
+
+Using the Quick Menu, you can now add a flashlight to your controller, to light up dark parts of your scene or see how light interacts with different Materials.
+
+![image alt text](image_31.jpg)(h:333 w:466 convert:false)
+
+## New: Screenshots from VR Editor
+
+You can now take screenshots right from VR!
+
+![image alt text](image_32.jpg)(h:392 w:604 convert:false)
+
+## New: Automatic Entry to VR Editing Mode
+
+You can now enter and leave VR editing mode when the VR editor is enabled without having to use the VR button or escape manually! As long as the editor is in the foreground, when you wear the headset, you will automatically enter VR editing mode; when you remove the headset, you will leave it.
+
+![image alt text](image_33.gif)(convert:false)
+
+There is a setting under VR in the Experimental section of Editor Settings that will allow you to turn off auto-entry if you prefer.
+
+![image alt text](image_34.jpg)(h:84 w:384 convert:false)
+
+## New: Sequencer Import/Export
+
+Sequencer can now import and export CMX EDL files for interchange with non-linear editing packages. 
+
+![image alt text](image_35.jpg)(h:529 w:869 convert:false)
+
+The above shot shows a sequence that was exported to Adobe Premiere Pro.  Each shot in a sequence will be written to a separate movie file which is referenced by the EDL file.  Any sequencing changes made in Premiere can then be imported back into UE4’s Sequencer!
+
+Sequencer’s ability to export HDR data in OpenEXR files has been expanded to give the user a choice of the color gamut used to encode the HDR data. 
+![image alt text](image_36.jpg)(h:220 w:492 convert:false)
+
+Finally, Sequencer now supports importing FBX animation directly to an object or track. You can also export animated tracks to FBX!
+
+## New: Sequencer Burn-ins on Renders
+
+When rendering out your movie sequence, you can now configure a "Burn-in" for the exported images.  This is very useful in dailies for identifying and tracking shots.
+![image alt text](image_37.jpg)(h:458 w:865 convert:false)
+
+## New: Media Framework Overhaul
+
+Media Framework API has been completely overhauled, with many new features!  Media Framework allows you embed live video and audio into your projects, with numerous powerful features.
+
+![image alt text](image_38.jpg)(convert:false)
+
+* Playlist assets for playing multiple media sources in a row
+
+* Audio playback support has been added
+
+* Improved media file import workflow
+
+* Improved Blueprint integration
+
+* Performance improvements on several platforms
+
+* Pixel format conversion on the GPU
+
+* Support for dynamically changing video dimensions on some platforms
+
+**Android**
+
+* Support for multiple audio tracks
+
+* HTTP Live Streaming (HLS) on devices supporting it (m3u8)
+
+**Playstation 4**
+
+* HTTP Live Streaming (HLS)
+
+* Improved playback controls (Pause, SetRate, etc.)
+
+* Media files can be pre-cached to memory
+
+* Opening media from FArchive
+
+* Ability to play multiple videos at once (may require increased memory pool settings)
+
+**Windows**
+
+* H.264 is now supported
+
+* Better support for HTTP(S) and RTSP streams
+
+* Better error handling and logging
+
+* Stability and usability improvements
+
+* Graceful handling of non-standard & unsupported codecs
+
+**Notes**
+
+* Early experimental macOS/iOS support (*AvfMedia* plug-in)
+
+* Experimental Linux support (via *VlcMedia* plug-in on Github)
+
+* Experimental Video-Over-IP support (via *NdiMedia* plug-in on Github)
+
+* XboxOne (*MfMedia*) and HTML5 are not supported yet
+
+* Integration with Sequencer/Video Recording is scheduled for 4.14
+
+* *PlatformMediaSource* asset is not implemented yet
+
+## New: Platform SDK Updates
+
+In every release, we update the engine to support the latest SDK releases from platform partners.  Also in this release, you can remote compile iOS/tvOS projects from Windows in the binary version of UE4!
+
+![image alt text](image_39.jpg)(h:377 w:888 convert:false)
+
+* iOs/tvOS: Code projects are now supported in the Windows version of the binary version of UE4 with the Unreal Engine Launcher. (You do need a Mac somewhere to remote compile on)
+
+* macOS: Now supports Xcode8 and macOS 10.13 Sierra
+
+* iOS/tvOS: Now supports iOS10/tvOS10
+
+* Xbox One: Upgraded to August 2016 XDK
+
+* Playstation 4: Upgraded to PS4 SDK 3.508.201
+
+* Oculus Rift: Updated to the Oculus 1.6 Runtime
+
+* SteamVR:  Updated to OpenVR 1.0.2
+
+* Google VR: Added Google VR (Cardboard) support for iOS
+
+* OSVR:  Updated to v0.6-1194-g0c54f5e
+
+* Android: Google Play Games native C++ SDK updated to 2.1
+
+* Android: Google Play Services updated to  9.2.0
+
+* Android: Supports running on Nougat aka Android 7.0
+
+* Vulkan API: Updated SDK to 1.0.17.0 (for Android and Windows)
+
+## New: Improved Landscape Tessellation
+
+Hardware tessellation on landscape is now much faster!  Landscape will now only render hardware tessellation on the highest level-of-detail (LOD), fading out as the landscape approaches the second LOD. Subsequent LODs will no longer have tessellation enabled. This significantly improves performance of enabling tessellation for displacement or extra detail up close to the camera. 
+
+![image alt text](image_40.jpg)(h:429 w:593 convert:false)
+
+In the visualization above, the highest the LOD (white) is rendered with tessellation enabled, while the other LODs (colors) are rendered without tessellation.
+
+## New: Animation Pose Assets
+
+We have added a new type of animation asset called a Pose Asset. This contains a set of named bone poses, which you can blend additively, in a similar manner to blend shapes for vertices. 
+
+![image alt text](image_41.gif)(convert:false)
+
+One use for this is to support facial animation where either FACS (Facial Action Coding System) or viseme curves can drive poses. However you could use this system to create new animation by blending multiple poses.
+
+Currently you create a Pose Asset from an Anim Sequence using Content Browser context menu or the Create Asset menu in Persona. When you create a Pose Asset, the pose names will be automatically generated. After that you can rename each pose manually, or use clipboard paste to rename all of them at once. 
+
+Poses are driven by normal animation curves. As long as they exist in the animation you can see the curve. In Persona, in order to preview a pose from a curve, you need to set the current Preview Pose Asset. 
+
+![image alt text](image_42.jpg)(h:152 w:348 convert:false)
+
+In the AnimGraph, you can use a Pose Blender Node (or Pose By Name) to output the pose based on the incoming curves. 
+
+![image alt text](image_43.jpg)(h:146 w:422 convert:false)
+
+To support this system, we have improved how we handle scale with additive blending. In the future we would like to support curves on other assets (e.g. Sound Waves) that can be used to drive Pose Assets and Morph Targets.
+
+## New: Pose Driver Animation Node (Experimental)
+
+We have added a new Pose Driver node, which lets you drive a curve value (such as a morph target weight) based on a bone’s movement. 
+
+![image alt text](image_44.gif)(convert:false)
+
+This uses an RBF (Radial Basis Function) to interpolate driven values based on the orientation of a target bone. You use a PoseAsset to define the target poses for the bone, and the desired curve values at each pose. This node can be used as a Pose Space Deformer, to drive corrective blend shapes based on the orientation of a bone.
+
+## New: Animation Node Pose Watching 
+
+Anim graph nodes can now be "watched" in Persona. 
+
+![image alt text](image_45.gif)(convert:false)
+
+This allows you to see a representation of the pose being generated at any point in the anim graph dynamically. Multiple watches can be active at once allowing you to compare poses at different points and find the exact point at which any errors in your current pose are introduced. This can be very useful for debugging complex Animation Blueprints and previously would only have been achievable by connecting the node you wanted to view directly to the root node and recompiling the blueprint.
+
+## New: Improved Scene Capture
+
+Scene Captures have been improved to be more useful with the new ‘Blueprint Drawing to Render Targets’ feature!
+
+![image alt text](image_46.jpg)(h:173 w:929 convert:false)
+
+* Orthographic projections are now supported.
+
+* When blueprints update a scene capture through the **Capture Scene** function, it happens immediately, allowing subsequent image processing with **Draw Material to Render Target**
+
+* Opacity is now captured in alpha, which allows partial rendering in a scene capture and compositing into another scene later.
+
+* Various GBuffer attributes are now available to be captured, including depth.
+
+* Added ‘Hidden Actors’ and ‘Show Only Actors’ arrays which can be used to easily control what is rendered into a scene capture.
+
+![image alt text](image_47.jpg)(h:312 w:364 convert:false)
+
+Game specific effects like fog of war can be implemented by rendering meshes as visibility shapes into an orthographic scene capture, and then doing image processing with Draw Material To Render Target.
+
+## New: Improved Refraction Shaders
+
+There’s a new "Pixel Normal Offset" refraction mode which uses the vertex normal as a reference, and computes the refraction offset from how different the per-pixel normal is from the vertex normal.  This is non-physical but allows refraction to be used on flat water surfaces.
+
+![image alt text](image_48.jpg)(h:153 w:929 convert:false)
+
+Left: simple scene with no refraction
+
+Center: default refraction method which causes an undesired constant offset for a water surface
+
+Right: Pixel Normal Offset refraction method which distorts based on the normalmap difference
+
+![image alt text](image_49.jpg)(h:78 w:344 convert:false)
+
+## New: Texture Coordinates from Line Traces
+
+We have added a project setting to support texture coordinate (UV) info from line traces. 
+
+![image alt text](image_50.jpg)(h:330 w:297 convert:false)
+
+The option is under Project Settings -> Physics -> Optimizations. When this is enabled, you can use the ‘Find Collision UV’ function to take a Hit Result and find the UV info for any UV channel at the point of impact. Enabling this feature does use extra memory, as a copy of UV information must be stored in CPU memory.
+
+## New: Spline Editing Improvements
+
+**Editing Spline Component defaults in the Blueprint Editor**
+
+Now it's possible to set Spline Component points in the Blueprint Editor using the standard spline visualizer editing features.  New instances of the Blueprint will be created with these defaults, although these too can be overridden on a per-instance level once placed.
+
+![image alt text](image_51.jpg)(h:445 w:929 convert:false)
+
+The "Reset to Default" context action on the spline visualizer will set an instance back to the Blueprint default.  Any changes made to the Blueprint default will be propagated to any instance whose spline points have not been subsequently edited.
+
+![image alt text](image_52.gif)(convert:false)
+
+**Numerical editing of spline points in the Details panel**
+
+Before there was no way to precisely place spline points and assign scale, roll or tangents to them.  Now these properties are exposed in the Details panel for selected spline points:
+
+![image alt text](image_53.jpg)(h:190 w:449 convert:false)
+
+**New spline point properties**
+
+Spline points can now be defined with distinct arrive and leave tangents, and an arbitrary input key.  The former allows for splines with discontinuities, while the latter allows for greater control of interpolation speed between points.   This allows for greater versatility when designing spline paths.
+
+In the spline visualizer, in the Spline Component details, there is an option in the context menu to allow the arrive and leave tangents to be edited separately instead of locked:
+
+![image alt text](image_54.jpg)(h:220 w:190 convert:false)
+
+As a consequence of being able to set arbitrary input keys per point, there is now also a way to specify the input key of the Loop Position for closed splines:
+
+![image alt text](image_55.jpg)(h:192 w:454 convert:false)
+
+If input keys or a loop position are not specified, they will default to starting at 0.0 and incrementing by 1.0 for each point, as before.
+
+**Defer spline update in Blueprints**
+
+Sometimes it's desirable to build splines procedurally in a Blueprint construction script.  Previously every operation on a spline point would cause the spline to be rebuilt, but now - for optimization purposes - it's possible to specify whether the spline should be rebuilt following a spline point operation.  There's also an explicit Update Spline node.
+
+![image alt text](image_56.jpg)(h:413 w:918 convert:false)
+
+**Input Spline Points to Construction Script**
+
+Sometimes it's useful to be able to hand edit a spline with the spline visualizer, and then refine them with a Blueprint construction script.  An example might be a Blueprint which locks all edited points to the surface of a sphere, like this:
+
+![image alt text](image_57.jpg)(h:456 w:929 convert:false)
+
+This is now achievable by checking the "Input Spline Points to Construction Script" property:
+
+![image alt text](image_58.jpg)(h:196 w:452 convert:false)
+
+## New: Sub Animation Blueprints
+
+You can now share animation logic by using a ‘Sub Anim Instance’ node within your Animation Blueprint to reference another Sub Animation Blueprint. This also allows you to break up large Animation Blueprints into separate assets, for example into ‘locomotion’ and ‘physics’ parts. 
+
+![image alt text](image_59.jpg)(h:149 w:644 convert:false)
+
+Member variables of the Sub Blueprint can be exposed as input pins on the node. The Sub Animation Blueprint must use the same Skeleton as the outer Animation Blueprint.
+
+## New: Animation Curve Viewer
+
+We removed the Skeleton Curve tab from Persona, and moved that functionality into the improved Animation Curves tab. Here you can now rename and delete curves, as well as previewing curve data. 
+
+![image alt text](image_60.jpg)(h:381 w:622 convert:false)
+
+You can see all curves that belong to current skeleton or currently active curves from preview asset. And you could also filter by specific type of curves if you only want to see active curves. Please note that we named default curve to be called "Attribute", so any animation curves will be by default attribute curves. 
+
+## If you want to modify the curve value, you can either turn off Auto check box option or just type the value. 
+
+## New: Sprites in UMG Widgets
+
+You can now use Paper2D Sprites as Brush inputs for UMG and Slate Widgets.  In addition to allowing users to reference UI art that may have been developed as a sprite sheet, it permits users to more efficiently render widgets on platforms where the draw call count budget is tight.  Sprites that are part of the same texture atlas can be batched together in Slate, provided they all share the same layer when rendered.
+
+![image alt text](image_61.jpg)(h:257 w:460 convert:false)![image alt text](image_62.jpg)(h:258 w:460 convert:false)
+
+## New: Optimized Instanced Stereo Rendering for VR
+
+There have been a number of improvements to instanced stereo rendering, including moving the velocity pass to use instanced stereo rendering.  Multi-view support has also been enabled on the PS4, which leads to significant performance improvements when using the ISR path.
+
+## New: GPU Morph Targets
+
+Projects can now enable calculating morph targets on the GPU on Shader Model 5 level hardware. This frees the CPU from performing those calculations:
+
+![image alt text](image_63.jpg)(h:194 w:520 convert:false)
+
+## New: Optimized Landscape Shader Memory
+
+Landscape now compiles dramatically fewer shader combinations for landscape materials, improving first-load shader compile times, improving editor iteration and reducing memory usage and shader cache size.
+
+## New: Shadow Optimizations for Mobile
+
+Optimizations have been added to the Combined Static and CSM shadow mode added in 4.12. In this mode, a stationary directional lights cast static shadows from static objects and CSM shadows for dynamic objects. In 4.13, the appropriate shader is now automatically selected based on the bounds of the dynamic objects casting CSM shadows, and there is no longer any need to manually tag each primitive that will receive combined static and CSM shadows.
+
+## New: Landscape Import Plug-ins
+
+A new plugin API has been added for landscape file formats, allowing developers to create plugins that add support for different heightmap and weightmap file formats to landscape. The existing raw and png support has been converted to the new API. The png support in particular makes a good reference for implementing new landscape file format plugins.
+
+## New: Automation Testing for Android
+
+The Project Launcher is now able to package and launch your project onto multiple Android devices simultaneously. The app running on each device will communicate back to your host PC over the USB cable and the will appear in the Session Frontend window.
+
+You can then launch Automated Tests on all the devices and see the results in the Session Frontend.
+
+![image alt text](image_64.jpg)(h:542 w:838 convert:false)
+
+## New: OpenGL ES 3.1 on Android
+
+While UE4 has long supported many OpenGL ES 3.0 and 3.1 features on Android, you can now specifically target ES 3.1 on Android. This brings feature parity with Metal and Vulkan to higher-end Android devices and gives you access to 16 texture samplers as well as improved performance through the use of uniform buffers.
+
+You can choose to package bothe ES 2.0 and ES 3.1 shaders for the same project, and the device will chooses the best shader platform based on the device’s capabilities. 
+
+![image alt text](image_65.jpg)(h:220 w:582 convert:false)
+
+## New: Mobile Packaging Wizard 
+
+We have added a Mobile Packaging Wizard to help support packaging for mobile where a minimal app without any content is uploaded to an App Store and the rest of the content is downloaded from the Cloud.
+
+![image alt text](image_66.jpg)(h:484 w:800 convert:false)
+
+* This type of packaging is common with larger games that have regular DLC updates
+
+* It also allows the user to download only the specific content required for their device, such as the texture or shader format their device needs
+
+* Wizard can be accessed from Project Launcher window as an option for new profile creation
+
+## New: Full Precision Materials on Mobile
+
+Materials have an option to use full precision (default: medium precision) in pixel shader when used on Mobile devices. For example this helps in cases when material uses world coordinates in computations.
+
+![image alt text](image_67.jpg)(h:92 w:260 convert:false)
+
+## New: Binary Shader Caching for Android
+
+Compiled shaders will be stored on the disk on first use and then reused on subsequent application runs. 
+
+* Requires GL_OES_get_program_binary extension 
+
+* Disabled by default, can be enabled only on Android devices (r.UseProgramBinaryCache=1)
+
+## New: Localized Text Formatting Improvements
+
+We’ve improved our localized text formatting to allow your translations to be more accurate.
+
+**Plural Forms**:
+
+* Plural forms allow you to use different text based upon a numeric variable given to your text format. Plural forms may be cardinal, eg "There is 1 cat", “There are 4 cats”, or ordinal, eg) “You came 1st!”, “You came 2nd!”, etc.
+
+* Plural forms are specified as key->value pairs, and support any of the following keywords (as defined for your culture by the [CLDR data](http://www.unicode.org/cldr/charts/29/supplemental/language_plural_rules.html)): zero, one, two, few, many, other. Values are an optionally quoted string that may also contain format markers.
+
+* Cardinal Format Example: "There {NumCats}|plural(one=is,other=are) {NumCats} {NumCats}|plural(one=cat,other=cats)"
+
+* Ordinal Format Example: "You came {Place}{Place}|ordinal(one=st,two=nd,few=rd,other=th)!"
+
+**Gender Forms:**
+
+* Gender forms allow you to use different text based upon an ETextGender value given to your text format, eg) "Le guerrier est fort", “La guerrière est forte”.
+
+* Gender forms are specified as a list of values in the order of masculine, feminine, neuter (neuter is optional). Values are an optionally quoted string that may also contain format markers.
+
+* Format Example: "{Gender}|gender(Le,La) {Gender}|gender(guerrier,guerrière) est {Gender}|gender(fort,forte)"
+
+**Hangul Post-Positions:**
+
+* Hangul post-positions help you deal with the grammar rules present in Korean, and will insert the correct glyph(s) based upon whether the value being inserted ends in a consonant or a vowel, eg) "사람은", “사자는”.
+
+* Hangul post-positions are specified as a list of values in the order of consonant, vowel. Values are an optionally quoted string.
+
+* Format Example: "{Arg}|hpp(은,는)"
+
+To allow you to pass in the numeric/gender values needed for plural/gender form support, all of the FText::Format(...) family of functions now take their values as FFormatArgumentValue rather than FText. This can be implicitly constructed from any numeric type, ETextGender, or FText.
+
+The ability to set these value types in Blueprints has been exposed using wildcard pins on the "Format Text" node:
+
+![image alt text](image_68.jpg)(h:166 w:564 convert:false)
+
+![image alt text](image_69.jpg)(h:168 w:820 convert:false)
+
+You can also pre-compile your format pattern if you’re going to be re-using it for multiple calls to FText::Format(...). Simply create and store a FTextFormat instance and pass it as the pattern to FText::Format(...).
+
+## New: Multi-threaded Audio (Experimental)
+
+Outside of the editor, sound cue evaluation and active wave instance determination can now run independently of the game thread. Benefits will vary per game depending on the quantity and complexity of active sound cues, however for some context, in early testing, we’ve seen a shift of approximately 1ms per frame off of the Game Thread and on to the Audio thread.  
+
+This feature is disabled by default in 4.13, but you can try it out by changing **UseAudioThread** to **true** in the [Audio] section of BaseEngine.ini.
+
+## New: Network Replay Backwards Compatibility
+
+The replay recording feature now supports backwards compatibility. This means you can make modifications to a build, even add or remove replicated properties and then load replays on the new build that were recorded with an older build. 
+
+To test this feature out, you can simply record a replay with a certain build, change some replicated properties, and then load that same replay on a newer build!
+
+Most of the work is handled by the low level reflection information we have for each replicated property. For custom serialized network data (UObject::NetSerialize), you can now use two new functions added to FArchive (FArchive::EngineNetVer() and FArchive::GameNetVer()) which allow you to obtain the current network version of the stream and handle old data manually.
+
+## New: Build Graph Scripting (Experimental)
+
+BuildGraph scripts can be used to create custom sequences of UE4 build commands for build automation. Graphs are declared using an XML script using syntax similar to MSBuild, ANT or NAnt, and consist of a network of dependent *nodes* made up of *tasks*. 
+
+![image alt text](image_70.jpg)(h:382 w:677 convert:false)
+
+Some example BuildGraph scripts are stored under Engine/Build/Graph/Examples/…
+
+## New: Script for Making Installed Builds
+
+The build process for installed UE4 distributions has been re-written using the BuildGraph scripting language, with a goal to making the process simpler and more transparent for other development teams with engine customizations. The script can be found at Engine/Build/InstalledEngineBuild.xml, and can be run by invoking AutomationTool with one of the following command lines:
+
+**Windows:**
+
+BuildGraph -target="Make Installed Build Win64” -script=Engine/Build/InstalledEngineBuild.xml -clean
+
+**Mac:**
+
+BuildGraph -target="Make Installed Build Mac” -script=Engine/Build/InstalledEngineBuild.xml -clean
+
+Append the -listonly option to view a list of available options for customizing the build process.
+
+## New: TSets as Unreal Properties
+
+TSets can now be used as a UPROPERTY for the purposes of automatic serialization, garbage collection, ini settings, and editing via the details panel (limited to entering as a string (i.e. "(1,2,3)")). Similar to TMap, TSet properties do not yet work as replicated members, nor can they be used in Blueprints.
+
+# Release Notes
+
+### AI
+
+* New: Added two appropriately named functions to the AI Perception Component
+
+    * The functions for "Get Currently Perceived Actors" and “Get Known Perceived Actors” have replaced the ambiguously named “Get Perceived Actors”.
+
+* New: Made AI Sight sense's collision channel configurable via the AI System's settings.
+
+* Bugfix: There was a case where EQS item score normalization where all all items had the same score. We used to normalize the score to 1, but now it's based on the shared item score.
+
+* Bugfix: The AI Controller was overzealously clearing out its Cached Gameplay Tasks Component on UnPossess, even if the pawn was not the owner of Game Play Tasks Component.
+
+* Bugfix: Blackboard Key Selectors were not being tied to the blackboard entries properly when used in Blueprint.
+
+* Bugfix: EQS queries were losing part of the data on asset duplication when they had more than one option branch.
+
+* Bugfix: The Processing of AI Messages has been fixed when a new message appears during a notify loop.
+
+* Made Simple Move To Location/Actor reset velocity only if the agent is already at goal.
+
+* Removed infinite loop protection from Gameplay Task event processing, as it was hiding source of issues.
+
+* Turned Graphs off by default in the Visual Logger.
+
+#### Behavior Tree
+
+* New: Added v-logging of failures of adding dynamic subtrees to a Behavior Trees Component.
+
+* Bugfix: Crash fix when pasting comments over from Blueprint editor into Behavior Tree editor.
+
+    * This makes it a partial comment-support for BT and EQS editors. You still cannot place them by using context menu, but you can copy-paste it.
+
+* Bugfix: Fixed a Blackboard Key Select issue where you were not able to allow more than one UObject subclass.
+
+* Bugfix: Fixed Behavior Tree auxiliary nodes, parallel tasks, active tasks, and aborting tasks being ticked while the Behavior tree is paused.
+
+* Bugfix: Fixed Behavior Tree decorator remaining relevant and observing blackboard after its subtree was removed from the stack.
+
+* Bugfix: Fixed Behavior Tree getting stuck on instantly finished gameplay tasks.
+
+* Bugfix: Fixed Behavior Tree serialization spawning duplicates of task services nodes.
+
+* Bugfix: Fixed missing observers in behavior tree when dynamic subtree was changed and the tree was waiting for a full restart. The search ran out of nodes.
+
+* Bugfix: Fixed an issue where Stop Tree was not working if at the time of the call the Behavior Tree Component was waiting for a latent task to abort. 
+
+* Changed behavior tree node instancing to use object duplication instead of spawning new object. This fixes missing values of Blueprint properties.
+
+#### Debugging Tools
+
+* New: Added detailed config for gameplay debugger's categories in Project Settings. Each can now define category slot, default state and input bindings. This will be moved to user settings section in a future release.
+
+* New: Added logging of the current Move ID to Pathfollowing Component's vlog snapshot.
+
+* New: Added client to server replication option for gameplay debugger's input event handlers.
+
+* Bugfix: Fixed Gameplay debugger extensions to now activate during Simulate in Editor.
+
+* Bugfix: Fixed Gameplay debugger extensions and rendering to now correctly update after tool is closed.
+
+* Bugfix: Fixed Gameplay debugger to no longer get stuck with outdated data packs on Client.
+
+* Bugfix: Fixed an issue where Gameplay debugger was sometimes not activating on clients, since it was depending on the wrong initialization order. 
+
+#### Navigation
+
+* New: Added a "default navigation filter" property to AI Controller to be used in given controller's navigation queries.
+
+* New: Added a "meta navigation filter" that, when used, fetches the Default Navigation Filter from AI Controller.
+
+* New: Added more details to the information returned by path following request: 
+
+    * Result flags. 
+
+    * Fixed stacking and identifying move tasks started by Blueprint functions.
+
+* Bugfix: Crash fix for path update processing.
+
+* Bugfix: Fixed an issue where auto-repathing to world location where the path’s end was being erroneously used to perform pathfinding rather than the original move goal location. 
+
+* Bugfix: Fixed an issue with Meta NavMesh Path doing navmesh path update without checking navigation agent if the path finding should be postponed (eg. when falling). 
+
+* Bugfix: RVO avoidance updated for vehicles.
+
+* Bugfix: Fixed an issue where Navmesh obstacles were cutting corners too close when using string pulling in detour’s crowd simulation.
+
+* Bugfix: Fixed an issue where navigation mesh streaming was not correctly handling criss-crossed sub-levels.
+
+* Navigation generation conditions in Scene Component's code have been improved to avoid doing unnecessary work.
+
+* Removed Nav Collision from static mesh on property change and Post Load if static mesh is not supposed to have one.
+
+* Unified bounds tests for applying navmesh modifiers, always expanding bounds one cell height on Z axis to cover for voxelization roundings.
+
+* Changed the Nav Mesh projection code to use 2D distance when looking for the best point on navmesh.
+
+### Animation
+
+* New: Added a pin to the Montage Play node to allow choice over the return type of the node for either the length of the selected montage or its playback duration.
+
+* New: Added ability to override angular joint bias per-node for Anim Dynamics nodes.
+
+* New: Added blueprint node to allow the setting of an Animation Blueprints Root Motion Mode
+
+* New: Added dynamic simulation space switching to Anim Dynamics animation node, and exposed more properties as optional pins.
+
+* New: Added the number of morph targets to the Skeletal Mesh Content Browser info.
+
+* New: Added option for parent bone space to Get the Socket Transform.
+
+* New: Added the ability to use the spacebar as a shortcut to play/pause animation playback in Persona.
+
+* New: Added spherical constraints to the Anim Dynamics animation node.
+
+* New: Animation curve import options:
+
+    * Remove Redundant Keys: This removes redundant keys when importing a custom attribute as a curve.
+
+    * Do Not Import Curve With Only Zero Values: If all of values are 0, do not import. Not having the curve value is same as 0.f.
+
+* New: Changed animation compression settings so that they are edited in a dialog instead of directly in the Details panel. This stops animation compression settings from becoming out of sync with the actual compression used.
+
+* New: Import Animation:
+
+    * Renamed the Animation Name setting to Override Animation Name.
+
+    * Option to add suffix for importing material curves by name. By default _mat.
+
+*  New: Animation sequences can now be marked as only being required by dedicated servers when they have root motion.
+
+    * Animation sequences still need to be flagged in the project cooker settings as being stripped on servers for this to have any effect.
+
+* New: Pose Asset supports runtime Retargeting option in Skeleton for Animation, Animation Scaled, and Skeletal.
+
+* New: Retarget Window:
+
+    * Automap button automatically attempts to map source skeleton to target skeleton.
+
+    * This doesn't work well with fingers, so you'll have to make sure to set them correctly.
+
+* New: Updated retargeting Animation Blueprints to a new skeleton so that it now finds all references to animations (in native properties and blueprint variables) and not just on animation nodes.
+
+* New: 'Showdebug Animation' displays Save Cache Pose nodes and subgraphs only once to keep display short. The Save Cache Pose and Use Cache Pose names are displayed.
+
+* New: Added search box to Socket chooser.
+
+* New: Socket copying in Persona now preserves bone attachment by default.
+
+* New: Support Step interpolation between key frames per animation. Please check property of Interpolation in Anim Sequence.
+
+* New: Added a tick dependency for slave components.
+
+    * Ensures that slave components always get ticked after master components.
+
+    * Prevents potential main thread stall updating morph targets in slave components.
+
+* New: Getting a Relative Transform now supports negative scale.
+
+* New: Getting a Socket Transform will now prioritize the socket name over the bone name.
+
+* New: Getting a World From Context Object is safe to be used in another thread, however the behavior is slightly modified. It will not use GWorld as an alternative option.
+
+* New: Exposed source indices for transition to runtime code so we can better identify which transitions are currently running when trying to determine elapsed times for anim getter nodes.
+
+* Bugfix: Fixed a crash when opening a skeleton when its assigned preview mesh doesn’t use that skeleton.
+
+* Bugfix: Fixed a crash when a Skeletal Mesh component gets destroyed while triggering a notify.
+
+* Bugfix: Fixed a crash when adding a section to a zero length montage. 
+
+    * Fixed a nullptr dereference in the montage handling code and disabled the menu option to add section when zero length.
+
+* Bugfix: Fixed a crash when viewing an uncompressed animation.
+
+* Bugfix: Fixed a crash when debug drawing of some animation nodes while in PIE. 
+
+* Bugfix: Fixed a crash when using undo/redo while editing the Animation Blueprints defaults.
+
+* Bugfix: There was an edge case in random sequence player where there might be an extra loop when waiting on a blend to finish. This made it impossible for a shuffle list to start with the animation that is currently playing, seemingly duplicating the animation.
+
+* Bugfix: Prevented an animation from restarting each time a new section is selected/inspected in the montage editor. 
+
+* Bugfix: Updated animation initialization to now verify the mesh’s hierarchy matches with the skeleton.
+
+* Bugfix: Animation Retargeting:
+
+    * Fixed an issue with Additive Blend Space/Anim Offset.
+
+    * Fixed an issue for not making the package dirty.
+
+* Bugfix: Fixed a compression issue with remove linear key not working properly with additive animations.
+
+* Bugfix: Fixed an issue where Additive Blend Space thumbnails no longer generate single node instance warnings.
+
+* Bugfix: Fixed an issue when copying non-POD structs in the animation fast path has been fixed. 
+
+* Bugfix: Fixed an issue where Root Motion for Everything would not blend correctly when using layered blend per bone.
+
+* Bugfix: Fixed an issue resulting in a slowdown in Anim Blueprints while dragging spin controls on default values.
+
+* Bugfix: Fixed abstract notify state classes to no longer show up in the create notify state menu.
+
+* Bugfix: Fixed Additive Animation scale to apply correctly. 
+
+* Bugfix: Fixed Animation ‘fast path’ to now work with various corner cases (e.g. multiple inputs to array properties from structs).
+
+* Bugfix: Fixed Blend Spaces to no longer produce bad data when degenerate spaces are present.
+
+* Bugfix: Fixed an issue with animation compression caused by an uninitialized box variable.  
+
+* Bugfix: Fixed Cached Posed Weighting to now ensure that the cached poses update with the highest weight they have been updated with, instead of just the first weight in their Animation Blueprints. 
+
+* Bugfix: Fixed Custom Blend graphs to no longer ignore curve values.
+
+* Bugfix: Fixed an issue with saving recompressed animations that would cause the error Graph Linked External Object.
+
+* Bugfix: Fixed an issue when creating animations with additive curves that would apply the layer twice.
+
+* Bugfix: Fixed an issue where animation metadata was not being saved for the montage.
+
+* Bugfix: Fixed an issue with the Layered Bone Blend calling Update on branches with the wrong weight.
+
+* Bugfix: Fixed an issue where properties that were fed to a Skeletal Mesh component via the construction script were not being updated when edited.
+
+* Bugfix: Fixed an issue where assets were missed in Montage nodes due to the reference gathering for animation assets in an animation graph not searching the Ubergraph.
+
+* Bugfix: Fixed an issue with Root Motion calculation when the actor and the component transforms did not match. 
+
+* Bugfix: Fixed some old Animation Blueprint nodes that maintained non-transactional pin links.
+
+* Bugfix: Fixed an issue that resulted in split pins not working as Anim Graph node inputs.
+
+* Improved handling of LODs for Anim Dynamics chains. Now when a leaf bone is removed from a LOD the rest of the chain continues to simulate as expected.
+
+* Prevented stale anim asset references from persisting in wired pins in animation blueprints.
+
+* Removed option to clear the compression scheme on an animation sequence as it is not valid to have no compression set.
+
+* Removed option to clear the default compression scheme as it is not valid to have no compression set.
+
+* Removed redundant copy in the Modify Bone animation node.
+
+* Removed Vertex Animation, and related code as they are deprecated.
+
+    * We have procedural mesh component that can be used for that.
+
+    * Vertex Anim Base has been removed, and replaced with Morph Target.
+
+* Removed various unnecessary bone container copies.
+
+* Renamed the Trigger Event curve to the Drive Attribute curve:
+
+    * In the UI, it will display as Attribute Curve.
+
+* Sequencer Scrub will display the Morph Target correctly.
+
+* Standalone games now benefit from parallel anim update if possible
+
+* Support different Sample Rate for reimporting with Set Range of time.
+
+### Audio
+
+* New: Added a setting to control dialogue wave audio filenames.
+
+* New: Added the ability to override the subtitle used on a dialogue wave.
+
+    * This is useful for effort sounds, plus some other cases, such as characters speaking in a foreign language not known to the player.
+
+* Bugfix: Fixed a crash that would occur when cooking in editor for non-xaudio platforms if a cooked sound is played during the same session.
+
+* Bugfix: Fixed an issue that resulted in the subtitle manager updating the wrong list of subtitles.
+
+* Bugfix: Fixed an issue with the output volume to make it more consistent with volumes of other PS4 games.
+
+### Automation
+
+* New: Added Tcp Message Transport, a connection-based message bus transport currently used for communicating with sessions launched on Android.
+
+* New: Added several enhancements to the automation framework and improving the UI:
+
+    * Tests in the UI now have a link to the source and line where they originate.
+
+    * There is now a general purpose latent lambda command you can use to run arbitrary code latently.
+
+    * Added Inlined Add Command for regular and networked commands to the base automation class, to avoid the use of the macro, which prevents breakpoints from working in lambda code.
+
+    * Frontend now has better column displays offering more room to the test name.
+
+    * Changed several events to the automation controller to multicast delegates so that many could hook them.
+
+    * The UI now refreshes the selection after tests finish so that the output log updates.
+
+* New: Logging ensure fails as errors. Automated tests with ensure fails will be unsuccessful.
+
+* Bugfix: Fixed a crash in Render Output Validation automation test when running  with cooked content.
+
+* Bugfix: Fixed an issue that could prevent a static mesh socket from being imported with the correct rotation.
+
+* Bugfix: Fixed screenshots to now save locally when running an automation test and an Automation Worker Screen Message is received.
+
+* Bugfix: Fixed an issue where Test logic had a broken check that would cause infinite map reloads when trying to run one in game.
+
+* Removed Func Test Manager warnings about PIE when running on a standalone game binary.
+
+* UAT parameter -signedpak now no longer implies -pak.
+
+* Tweaking the padding to make it look more like other windows and to make everything not look so squished. Recursive expansion now works on tests.
+
+* Toolbar improvements: 
+
+    * Removed the search box from the toolbar. It is now inlined above the test tree.
+
+    * Moved the filter group drop-down out of the toolbar and onto the line with the search box above the treeview and general improvements to the feeling of it.
+
+### Blueprints
+
+* New: Added "empty" keyword to Clear Array Node which will show the Clear Array Node in the Blueprint Graph Context Menu when searching using the word "empty."
+
+* New: Added "Get All Actors with Tag" blueprint node which returns an array of all actors with the given tag.
+
+* New: Added "negate" keyword to Boolean NOT node which will show the Boolean NOT Node in the Blueprint Graph Context Menu when searching using the word "negate."
+
+* New: Added a "Generate abstract class" option to Blueprint class settings.
+
+* New: Added a Blueprint Node to multiply a Vector2D by another Vector2D
+
+* New: Added Editor Preferences options in General/Appearance/Graphs to disable grid in the Blueprint and Material editor, as well as change grid colors. 
+
+* New: Added new macro to help scale and position materials that are drawn to canvas.
+
+* New: Added support for emitting the Blueprint Description as tooltip metadata for the compiled Blueprint class (displayed in class pickers, etc...)
+
+* New: Exposed the web browser widget's ExecuteJavascript method to widget blueprints
+
+* New: Added Interface properties can be exposed on spawn.
+
+* New: Added 
+
+* Bugfix: Fixed a crash in Graph Node widgets when handling invalid pins.
+
+* Bugfix: Fixed a crash when compiling Blueprints caused by Set Boolean (by ref) node.
+
+* Bugfix: Fixed a crash when compiling Blueprints caused by Expand Enum As Execs metadata used with an enum containing a hidden element.
+
+* Bugfix: Fixed a crash while loading Blueprints that are circularly referenced by their interface.
+
+* Bugfix: Fixed a crash that would occur after editing a Blueprint that had been opened in the diff tool.
+
+* Bugfix: Fixed a crash when hovering over Create Widget node in Blueprints.
+
+* Bugfix: Fixed a crash when opening Kite Demo.
+
+* Bugfix: Fixed a crash when transforming an actor after applying instances changes.
+
+* Bugfix: Fixed a crash for the struct editor that could occur when editing Name member variables.
+
+* Bugfix: Fixed a crash when loading of a Blueprint class containing an integer property tagged as a bitmask with missing enum type metadata.
+
+* Bugfix: Fixed a crash that would occur when attempting to display a node that was already removed.
+
+* Bugfix: Fixed a crash that would occur when giving an event in a child Blueprint the same name as an event in its parent.
+
+* Bugfix: Fixed a crash that could occur when opening the Blueprint palette window.
+
+* Bugfix: Fixed a crash that could occur while updating Blueprint instances on compile.
+
+* Bugfix: Fixed a crash that was possible when removing a corrupted Blueprint component node if it’s found to have no parent attachment. 
+
+* Bugfix: Fixed a crash when using the Blueprint Documentation commandlet.
+
+* Bugfix: Fixed a crash when copying and pasting collapsed graphs.
+
+* Bugfix: Fixed a crash for the debug target picker when trying to access an invalid world.
+
+* Bugfix: Fixed a crash when breaking all links to a MakeArray node with split struct pins.
+
+* Bugfix: Fixed a crash for a needed component that was being garbage collected during Blueprint compilation. 
+
+* Bugfix: Fixed a crash for modifying and hot-reloading a Blueprint superclass.
+
+* Bugfix: Fixed a crash for an assert when loading a Blueprint with instanced sub-objects.
+
+* Bugfix: Fixed a crash when opening a Blueprint with a component that has all inherited properties hidden.
+
+* Bugfix: Fixed a crash when breaking a link to a pin on a MakeArray node that is any array type.
+
+* Bugfix: Fixed a crash for dragging a component into the Event Graph in a Blueprint that inherits from a C++ class.
+
+* Bugfix: Fixed a crash for connection to a split struct pin in a MakeArray node.
+
+* Bugfix: Fixed a crash for Blueprint node to Macros that would happen sometimes.
+
+* Bugfix: Fixed a crash that would sometimes happen when adding unsupported pin types to a Custom Event node.
+
+* Bugfix: Fixed a crash for an editor-only crash that could occur during instancing of an uncompiled Blueprint class with a modified array property in the native parent classes default object.
+
+* Bugfix: Fixed an issue where the Get Class node would fail to appear in the Context Menu when pulling from a self reference pin into an empty space in the Blueprint Graph.
+
+* Bugfix: Fixed an issue with Blueprint Function Libraries not working with class/function redirectors.
+
+* Bugfix: "Add Component" nodes will now be fixed up at load time if they are found to not be associated with a unique template object.
+
+* Bugfix: Associated component template(s) are now being cloned when duplicating Blueprint function graphs containing one or more Add Component nodes. 
+
+* Bugfix: Basic shape components (cube etc.) will now apply the correct override material to instances when they are added to the component tree in the Blueprint editor.
+
+* Bugfix: Functions can now be properly overridden that are nested in an Interface Inheritance Hierarchy.
+
+* Bugfix: Changes made to settings in the Blueprint Profiler will now persist.
+
+* Bugfix: Changes to string literals that only affect case are properly applied (e.g. 'Hello' can be changed to 'hello')
+
+* Bugfix: Component data cache template objects will now properly remain referenced during garbage collection.
+
+* Bugfix: Components instanced as default subobjects of and attached to components instanced by Blueprint class construction script (i.e. nested default subobjects) are no longer being included in the Level editor's components tree view, as property values on these instances currently cannot be edited properly.
+
+* Bugfix: Copy-and-paste of Actor instances from the current level to the components tree view in the Details panel will now add components that have been properly initialized.
+
+* Bugfix: Customized defaults are now fully propagated to new instances at construction time during non-Actor-based Blueprint class re-instancing.
+
+* Bugfix: Details panel in the Blueprint diff tool is correctly set to read only again so that the user cannot edit the values.
+
+* Bugfix: Compiled Blueprint class assets in the Content Browser would give an invalid ensure (warning) when force-deleting.
+
+* Bugfix: Phantom asset being created when ‘Accept Source’ was used in the Blueprint Merge Tool.
+
+* Bugfix: Child Blueprints inherited variables would not display the replication settings set in the parent Blueprint.
+
+* Bugfix: The Details panel Macro where the tooltip for the "Keywords" and “Compact Node Title” was displaying the same “Category” rather than their respective tooltips.
+
+* Bugfix: The Structure editor would collapse all values in the Default Value panel after undoing any changes.
+
+* Bugfix: Widget Blueprints would collapse values in the Details panel after changing the default value of any Blueprint variable.
+
+* Bugfix: Unsupported pin types are no longer supported for Blueprint Event Nodes.
+
+* Bugfix: Corrected issue that caused a compile error in CustomEvent nodes.
+
+* Bugfix: Corrected issue that caused a MakeArray node to assume the type of a split struct pin rather than the parent struct when pasting from a Make Array node with split pins.
+
+* Bugfix: The Make Array node no longer gives incorrect indices when removing elements with split pins.
+
+* Bugfix: Fixed an issue where copying and pasting return nodes would fail to copy links to other nodes.
+
+* Bugfix: Fixed an issue where an automatically generated Get Display Name Blueprint node would not be connected when connecting a self reference output pin to a string input pin.
+
+* Bugfix: Fixed an issue where a wildcard Output Pin of an Array type could not connect to a split pin of an array type in a Make Array Blueprint node. 
+
+* Bugfix: Fixed an issue where Blueprint compilation logs were being cut off and not output.
+
+* Bugfix: Fixed an issue where Blueprint Node Comment Bubbles would not save changes when losing focus.
+
+* Bugfix: Fixed an issue where copying and pasting a Blueprint Event Node along with one or more linked nodes would not also copy any links.
+
+* Bugfix: Fixed an issue where copying and pasting a Make Array node with added values would create a wildcard Make Array node rather than a copy of the original.
+
+* Bugfix: Fixed an issue where cross sub-object pointers could be invalid post load.
+
+* Bugfix: Fixed an issue where double-clicking on a component listed in the results of a "Find-in-Blueprints" search would not select the component in the Blueprint editor.
+
+* Bugfix: Fixed a cosmetic issue where dragging from an input pin would cause the link to appear from an output pin if the pins share the same name until the connection is completed or cancelled.
+
+* Bugfix: Fixed an issue where including a period in the name of a Blueprint function would cause double clicking the function name to do nothing instead of open its Blueprint Graph.
+
+* Bugfix: Fixed an issue where multiple Primitive types (eg. Vector, Rotator, Transform) would not automatically type cast to Text when trying to connect them to a Text Input Pin on a Blueprint node.
+
+* Bugfix: Fixed an issue where native interface functions were being miscategorized.
+
+* Bugfix: Fixed an issue where output nodes sometimes would not be generated when collapsing.
+
+* Bugfix: Fixed an issue where overridden component properties could be cleared on load.
+
+* Bugfix: Fixed an issue where split struct pins would disconnect on load.
+
+* Bugfix: Fixed an issue where the Blueprint editor's initial zoom focus transition could not be interrupted.
+
+* Bugfix: Fixed an issue where toggling the Comment Bubble button on Reroute Nodes would not properly toggle the Comment Bubble.
+
+* Bugfix: Fixed an issue where the functions filter when searching with Find-In-Blueprints would list components from SCS in addition to matching Blueprint functions.
+
+* Bugfix: Fixed several issues that were causing a number of issues after Blueprint compilation for bad property read/writes, infinite loops, and crashes.
+
+* Bugfix: Fixed a minor spelling error in the tooltip for "Save On Compile" in the Blueprint Editor.
+
+* Bugfix: Fixed a potential data loss issue when loading a serialized non-native component class instance that's owned by an Actor-based Blueprint class instance.
+
+* Bugfix: Fixed a potential infinite loop when replacing variable nodes in a Blueprint graph via drag-and-drop.
+
+* Bugfix: Fixed a potential loss of "Get Class Defaults" node output pin links on load. This was due to a dependency load order issue.
+
+* Bugfix: Fixed an ensure that would fail when dragging off a self pin when another actor was selected in the Content Browser.
+
+* Bugfix: Updated an error message on Create Event nodes to make it more clear when signatures no longer match the delegate they are bound to.
+
+* Bugfix: Fixed an issue with latent nodes making Blueprint bytecode generation more deterministic.
+
+* Bugfix: Fixed an ensure when dragging a level from the Levels window into Blueprint event graph.
+
+* Bugfix: Fixed an issue with the direction of Get Overlap Info parameters.
+
+* Bugfix: Fixed instances where "Asterisk" was misspelled as "Asterix" when searching in the Blueprint's Event Graph Context Menu.
+
+* Bugfix: Fixed an issue where macro Blueprints were triggering other Blueprints to recompile every time they were loaded.
+
+* Bugfix: Fixed an issue where Multi Line in User Defined Structs was not working with arrays of string or text properties.
+
+* Bugfix: Fixed a regression in My Blueprint panel as displayed in the Blueprint Diff Tool - was not displaying any values.
+
+* Bugfix: Fixed a stack overflow that could occur when compiling huge blueprints.
+
+* Bugfix: Fixed an issue where a wrong value of an Enum literal would be used with functions from Array Library.
+
+* Bugfix: Updated function return nodes to be more fail-safe. You can no longer remove them in overridden/inherited functions and the function's signature is maintained even when disconnected.
+
+* Bugfix: Fixed an issue to immediately apply transform updates to component template instances when changing the scene component hierarchy root in a Blueprint class.
+
+* Bugfix: Improved "variable not found" warning messages. You can now jump to the offending node.
+
+* Bugfix: Fixed an issue in the Blueprint Profiler so that the nested event node is now distinguished from its top-level event with the addition of "(Entry Node)".
+
+* Bugfix: Fixed an issue so that moving a child component in a child blueprint does not force the parent to become dirty anymore.
+
+* Bugfix: Fixed an issue so that collapsing Blueprint Nodes to a function when there are multiple input execution paths is no longer allowed. The resulting function would have ambiguous execution logic.
+
+* Bugfix: Collapsing Blueprint Nodes to a function will no longer create a return node if there are no return values.
+
+* Preserved relative scale on the root scene component when converting an Actor instance to a Blueprint Class.
+
+* Editor Performance Updates:
+
+    * Removed a redundant Blueprint class compilation pass after a variable type change.
+
+    * Significantly improved PIE startup time when automatically recompiling Blueprint classes that have multiple dependencies.
+
+    * The Ed Graph Pin is no longer a UObject, this will improve load times significantly on projects with many large blueprints. Note that content does need to be resaved in order to see the improvement in load time.
+
+* The "Delete Unused Variable" feature now considers the "Get Class Defaults" node.
+
+    *  Output pins representing unused properties will now be removed from the node if they're not linked.
+
+* The "Get Class Defaults" node will now update immediately in response to structural Blueprint Class changes.
+
+* The Minimum Area Rectangle function is now callable without a World Context parameter.
+
+* When pasting a component in to the blueprint component tree, relative position and rotation is now cleared and will not be offset to its previous world position.
+
+* Optimized Math Expression parsing.
+
+### Core
+
+* New: Improved "DumpParticleSystems" console command to now print totals for each of the categories that were already listed for each individual particle system.
+
+* New: Added a new "-checkpak" command line switch which will attempt to verify every chunk in a signed pak file is valid when the pak is opened.
+
+* New: Added more detailed message when TArray's Bulk Serialize fails
+
+* New: Asset registry memory optimisation:
+
+    *  Made sure all dependency node link arrays are sized exactly as they need to be.
+
+* New: Crash callstacks will now be output to log by default.
+
+* New: Crash handling on Linux saves the Crash Report Client logs in the same location as the main game/server log.
+
+    * Helpful for easily associating CRC logs with the session that crashed on Linux game servers.
+
+* New: Dedicated servers will no longer use memory logging by default.
+
+* New: Added minor logging and debugging enhancements for signed pak files to detect corruption.
+
+* New: Added config file based system for marking which classes aren't required by a dedicated server or client.
+
+    * Added class names to the new options in the project cooker settings in the editor.
+
+    * Overridden Needs Load For Server and Needs Load For Client functions on each class will still overrule the config file.
+
+* New: Optimized cooked asset registry dependency node data structures to reduce memory usage.
+
+* New: Added Option to show warnings and errors visually on the HUD in development builds. 
+
+    * Controlled by Duration Of Errors And Warnings On HUD in Engine.Engine section of DefaultEngine.ini (0 = disable).
+
+* New: Split Garbage Collection cluster index and internal object flags to allow more than 8 million UObjects in editor builds.
+
+* New: When graphics driver crashes the engine will now assert so that we can send a crash report. The Crash Report Client will now attempt to initialize Slate Renderer multiple times before it gives up in which case it will send the crash report in unattended mode.
+
+* Bugfix: Fixed a crash occurring in serialization when signed paks are used with compressed packages. 
+
+* Bugfix: Fixed a crash when editing UObject’s in a Map property box.
+
+* Bugfix: Fixed a crash for Automation Tool when deploying to a default PS4 devkit.
+
+* Bugfix: Fixed a crash in UHT where an "Ambiguous search" error when a delegate has an own class declared.
+
+* Bugfix: Fixed a crash when loading an empty cached asset registry.
+
+* Bugfix: Fixed a crash when loading two stat capture profiles simultaneously.
+
+* Bugfix: Fixed a crash when pasting T3D into the Content Browser from the clipboard.
+
+* Bugfix: Fixed a crash when selecting a different capture in the stats dump browser of the profiler when another file is still loading.
+
+* Bugfix: Fixed a crash when spawning an actor using a template object with instance components. 
+
+* Bugfix: Fixed a crash when saving a package which has an unresolved dependency. 
+
+* Bugfix: Fixed a crash for hot-reload when a USTRUCT used is used as a field has a UPROPERTY field pointing to a hot-reloaded CDO.
+
+* Bugfix: Fixed a crash after opening edit config data for remote build and then trying to package.
+
+* Bugfix: Updated Exclusive Internal Flags so that they are now respected when passing in a null Object Package in Static Find Object Fast Internal Thread Safe.
+
+* Bugfix: Fixed an issue with null reference exception when connecting to the Swarm Coordinator.
+
+* Bugfix: Fixed and improved log message when garbage collector cluster assumptions are violated.
+
+* Bugfix: Updated DLC pak file path mappings to be relative to Game Dir instead of Root Dir.
+
+* Bugfix: Fixed Enum bitflags not being hot-reloaded properly.
+
+* Bugfix: Fixed an infinite loop in the editor that would happen if a directory that is being watched is deleted.
+
+* Bugfix: Fixed an issue with leaked linkers created by blocking load requests during async loading. 
+
+* Bugfix: Fixed an issue with leaked log archive.
+
+* Bugfix: Fixed an issue where there was an incomplete/missing callstacks reported by thread heartbeat when a hang is detected.
+
+* Bugfix: Fixed an issue where there was potentially missing log output when serializing text of length equal to the async log writer buffer size.
+
+* Bugfix: Fixed an issue when serializing packages over 2GB when cooking. Usually this was caused by lightmaps.
+
+* Bugfix: Fixed an issue with Cook-on-the-fly server asset paths when the game name differs from its folder name.  
+
+* Bugfix: Fixed timestamp checks of uncooked files when starting up the cook-on-the-fly server.
+
+* Bugfix: Fixed Asset Pointer remapping issues for duplicated level packages (level Save As, level Duplicate).
+
+* Bugfix: Updated most Box functions to use FORCEINLINE to improve Debug Game performance.
+
+* Bugfix: Made Launcher Tasks have unique names so we don't end up with stat name collisions which spam the log.
+
+* Bugfix: Fixed FPS pop-up updates when loading new stat file.
+
+* Bugfix: Removed an obsolete error message about INI file case sensitivity in the Unreal Header Tool.
+
+* Updated exporting to silently skip creating exports from a package where the outer is also an export and has been filtered at runtime during loading.
+
+* Updated the engine to no longer attempt to flush log messages in Check Rendering Thread Health if a critical error has already occurred. Crash handlers flush log anyway.
+
+* Updated the thread heartbeat to be killed when reporting ensures so that it doesn't get triggered when the report takes a long time to generate.
+
+* Updated the thread heartbeat to be stopped when the engine crashes to prevent it from interfering with generating crash reports.
+
+* Updated the thread heartbeat to be suspended when a message box is being displayed to fix false positives in hang detection.
+
+* Updated Media Player objects to not be added to any Garbage Collection clusters because they can load additional assets after they had PostLoad called on them and that results in Disregard For Garbage Collection assumptions being violated.
+
+* Non-backup files will no longer be deleted when cleaning up the log folder.
+
+* You can once again specify more than one ini override on the command line.
+
+* Removed a few redundant Add Referenced Objects functions to speed up Garbage Collection.
+
+* Modified crash handling code (on Windows) to handle two threads crashing at the same time properly. Previously the second crash would force the process to exit before generating the crash report.
+
+### Editor and Tools
+
+* New: Updated the "Description" field for a blueprint function to be multi-line to encourage big descriptive tooltips.
+
+* New: Added a setting to allow recording of actors in the sequence recorder that are spawned by sequencer itself.
+
+* New: Added hotkey (ctrl-T) to toggle between showing frame numbers and time.
+
+* New: Added plugin support to the editor class wizard.
+
+* New: Added a "Console Variables" command under the main editor Help menu.
+
+    * This creates a local html page with browse and search functionality for all cvars.
+
+* New: Added a setting to invert the Y axis in editor viewports for mouse look and orbit in Project Settings.
+
+* New: Added rich tooltip support to enum combo-box entries.
+
+* New: Added support for creating a HLOD cluster from a single actor, as long as it has two or more static mesh components in it. This improves handling when including Blueprints.
+
+* New: Added track for string properties.
+
+* New: Added UI for loading and browsing multiple stats dumps from a folder in the session frontend profiler for easy comparing of perf stats.
+
+* New: Changed how materials are generated by the Hierarchical Level of Detail (HLOD) system. They are now created as material instances to reduce the overall shader count
+
+* New: During a "Play in Editor" session, automatically regain focus after user takes over mouse control. This allows the editor to continue to process hotkeys without having to click on the editor manually to regain focus.
+
+* New: Added the ability to set PNG and JPG files as a Splash Image from the Project Settings.
+
+* New: Locking the mouse to a viewport is now determined by an enumerated value (Don't Lock, Lock on Capture, Lock Always) rather than a boolean.
+
+* New: Changed sRGB Preview to be the default in the Color Picker.
+
+* New: Changed Memory Profiler 2 to deselect current bar when clicking off the histogram. This allows the user to see the top level group data again.
+
+* New: Updated Interface classes to now have their own entry in the New C++ Class wizard.
+
+* New: Updated Obj exporter to now export all mesh components contained in Blueprint Level instance.
+
+* New: Refactored the way debugger commands are executed in PIE. Now they are controlled using the Editor Key Bindings, for Play/Pause/Stop...etc, rather than the Debugger Commands ini settings.
+
+    * You can now do things like bind stop to Shift+Escape, if you wanted to test showing a menu when Escape is pressed, instead of closing PIE.
+
+* New: Added support to Scene Importer for rigid mesh animation.
+
+* New: Added substring matching for console command suggestions, with space as a wildcard.
+
+* New: Added support to the Scene Importer for .OBJ files.
+
+* New: Added a new flag to the Merge Actor settings 'Use Vertex Data for Material Merging' and renamed 'Bake Vertex Data' to 'Bake Vertex Data to Mesh'.
+
+* New: Updated UnrealGameSync to hide *.uprojectdirs files by default in the project file selector, to discourage their selection by accident.
+
+* New: Added the ability for Resave Packages commandlet to now also rebuild asset registry paths with consistent case.
+
+    * This produces much cleaner results than using the text export method (which will use the internal names for user defined structs).
+
+* New: You can now package a plugin from the plugin browser. 
+
+    * Select the plugin and there's a small link at the bottom of the plugin tile letting you execute the package operation for sharing your plugin.
+
+* Bugfix: Fixed a crash when renaming the current level in the Content Browser.
+
+* Bugfix: Fixed a crash with the curve editor widget if you box selected outside of the visible area.
+
+* Bugfix: Fixed a crash for exporting the current level to an OBJ file and exporting materials as images.
+
+* Bugfix: Fixed a crash for selecting a Grid Snap size with a corresponding Decimal Grid Interval Value of 0 units.
+
+* Bugfix: Fixed a crash when exporting actors with non-scene components to FBX.
+
+* Bugfix: Fixed a crash that could potentially happen when selecting a spline point in the editor after having deleted another point. 
+
+* Bugfix: Fixed a crash Importing a brush, so that the brush is always validated (relinked), whether it is a static or dynamic brush. 
+
+    * This fixes an occasional crash when copy/pasting BSP brushes from one level to another.
+
+* Bugfix: Fixed a crash when hovering over a submenu while entering a Blueprint node comment in the Context Menu. In general, issues with opening submenus while a parent menu has focus.
+
+* Bugfix: Fixed a crash caused by repeatedly undoing and redoing a change in a level.
+
+* Bugfix: Fixed a crash for clicking the Install button multiple times when creating a new code project without a compiler installed.
+
+* Bugfix: Made a number of fixes to spline mesh collision generation:
+
+    * A serious issue with DDC ID generation was fixed, in that the static mesh itself wasn't forming part of the key, hence any two splines with identical properties but different meshes would yield the same cache entry.
+
+    * Fixed how different collision boxes are transformed when rebuilding physics meshes. Convex collision transforms are now correctly taken into account, and spherical and capsule collision now gets correctly translated when a scale is applied to the start or the end of the spline mesh.
+
+    * Optimized physics rebuilding. A new Body Setup object is now only created when needed, otherwise it is reused.
+
+    * Fixed Spline Mesh collision building so that geometry does not default to being auto-inflated in PhysX.
+
+* Bugfix: Updated Actors placed in the current viewport to no longer snap to surfaces if surface snapping is disabled. 
+
+* Bugfix: Fixed an issue where compiling a parent Blueprint would sometimes result in all of its child Blueprint assets adding toast notifications that they require check out.
+
+* Bugfix: Updated the Rotator UI showing Roll/Pitch/Yaw in some places. It will always show as X/Y/Z now.
+
+* Bugfix: Improved the "Repaired Painted Vertex Colors" log message. This is now done once on levels at load in the editor, instead of for every component (massive log spam). Total load time caused by the fixup is reported.
+
+* Bugfix: Updated auto-reimport to now be disabled inside an editor running in unattended mode.
+
+* Bugfix: Fixed an issue where editing condition property value changes would not propagate to instances of template objects.
+
+* Bugfix: Updated the Fbx importer to now import static mesh socket rotation correctly.
+
+* Bugfix: Updated the Fbx importer to now prevent importing a mesh with a size of zero and log a message to the user.
+
+* Bugfix: Fixed an issue where audio was being disabled after recording a movie.
+
+* Bugfix: Fixed an issue with Decal Visibility so that its behavior matches other visible actors with regard to the "Actor hidden in game" flag. This also makes it possible to animate it's visibility correctly with sequencer.
+
+* Bugfix: Updated nodes in folders to be filtered so that other unfiltered children aren’t visible.
+
+* Bugfix: Fixed an issue when parsing LOD in a Static Mesh component when importing custom properties.
+
+* Bugfix: Fixed an issue where selecting a color with the eyedropper tool would sometimes prevent the user from being able to Play in Editor.
+
+* Bugfix: Fixed an issue with a localization warning that would happen when you opened the GameMode menu.
+
+* Bugfix: Fixed a longstanding issue where sometimes BSP geometry was not correctly rebuilt after editing vertices in geometry mode.
+
+    * Levels with incorrectly built geometry (due to invalid poly normals) will now issue a warning in the log upon load, and can be fixed by resaving the level.
+
+* Bugfix: Fixed a thread metadata warning in the output log when using the project launcher in the editor.
+
+* Bugfix: Fixed an issue with Actor position being set to origin after using undo.
+
+* Bugfix: Fixed an issue where the output log was prevented from remaining anchored to the bottom of the window.
+
+* Bugfix: Fixed an issue where sometimes spline components were debug rendered incorrectly when not selected in the editor. 
+
+* Bugfix: Fixed an issue where the "All Classes" tab of the mode panel would not refresh when a new class was added to the project.
+
+* Bugfix: Fixed an issue where the surface properties of a BSP would erroneously persist in the details panel after the BSP was deleted.
+
+* Bugfix: Fixed an issue where the transform gizmo would prevent a user from dragging a material onto a mesh.
+
+* Bugfix: Fixed an issue where stopping ‘Cancel’ from correctly ‘Cancelled’ resulted in P4 asynchronous ops.
+
+* Bugfix: Fixed an issue importing an array inside a JSON Data Table.
+
+    * The JSON importer for a Data Table was appending the new data to the array rather than replacing it. It now clears the array prior to importing.
+
+* Bugfix: Fixed an issue when moving translation of the child when attached to the mirrored parent in viewport.
+
+* Bugfix: Fixed an issue when multi-selecting scaling when objects are rotated.
+
+* Bugfix: Fixed an issue where Post Process settings for blendable picker were not readable in the Details panel.
+
+* Bugfix: Fixed an issue with single player PIE so that the window position is correctly fetched and saved, even when running a dedicated server. This does not interfere with stored positions for multiple client PIE.
+
+* Bugfix: Fixed an issue where geometry operations were not undoable (eg. Material Placement).
+
+* Bugfix: Fixed an issue where foliage instances were not being included when exporting a scene to FBX.
+
+* Bugfix: Fixed an issue where the Launcher Profile settings window was not always showing all available maps for cooking.
+
+* Bugfix: Fixed Blueprint widgets "is Enabled" property to now animate properly.
+
+* Bugfix: Fixed the Editor Mode plugin sample to no longer provide a bad starting example for where to create your widgets.
+
+* Bugfix: Fixed the File menu option for "Import into level" so that it now re-imports the original blueprint instead of overwriting it.
+
+* Bugfix: Fixed an issue where multiple text cursors would appear when editing Keyboard Shortcuts in Editor Preferences.
+
+* Bugfix: Fixed a formatting issue for the error of Error_TooManyMaterials message.
+
+* Bugfix: Fixed Datatable select next row.
+
+* Bugfix: Fixed a small typo in GraphEditor.h.
+
+* Bugfix: Fixed missing END_OPTIMIZATION macro to SOutputLog.
+
+* Bugfix: Fixed issue preventing generation of JSON for nested structs in DataTable rows.
+
+* Bugfix: Reverted the change to the Merge Actor tool where you could not merge a single mesh.
+
+    * This is needed if you want to have a single Actor merge multiple materials but not be merged with another Actor in the scene.
+
+* Bugfix: Fixed the scale widget axes to render correctly while the current viewport is in an orthographic view mode.
+
+* Bugfix: Fixed the Texture Editor to properly display Mip Level Count and Format for 2D Dynamic Textures.
+
+* Bugfix: Fixed the "Game Gets Mouse Control" setting in the preferences for the Level Editor to be correctly applied when Playing in Editor.
+
+* Bugfix: Fixed the names of Automation tests within the standalone Session Frontend to be displayed correctly rather than as solid white blocks.
+
+* Bugfix: Fixed UnrealGameSync to now exclude deleted files when searching for the editor target name, which was sometimes preventing the editor from launching.
+
+* Bugfix: Fixed UnrealGameSync so that syncing no longer fails if the current workspace has files open for move or add.
+
+* Bugfix: Updated structs to export as JSON when displaying them in the Data Table editor.
+
+* Bugfix: Updated the F4 debug key binding to match what is in ShowFlags.cpp.
+
+* Bugfix: Fixed importing a morph target so that it no longer imports textures and materials for a base mesh that has already been imported. 
+
+* Changed the 'Merged Physics Data' flag in the Merge Actor tool to be true by default in order to match user expectation.
+
+* Disabled the Paste context menu action if a property in the Details panel is marked as "EditConst".
+
+* Updated importing assets into a level so that assets that do not come from FBX or OBJ files are no longer supported.
+
+* Improved the performance of the details panel when thousands of objects are selected.
+
+* Improved moving parent & grouped child actors so that it now does not result in deltas being applied twice.
+
+* Updated the Editor so that new windows will no longer open in between monitors when the primary monitor is not set to monitor 1.
+
+* Improved DataTable import/export.
+
+* Improved mesh vertex color painting:
+
+    * Lower LODs are now automatically fixed up for instances which were created in a previous bugged version of the engine.
+
+    * Since lower LODs cannot currently have their vertex colors edited, their vertex colors are always derived from LOD0.
+
+    * A bug has been fixed when building lower LODs so that vertices in neighboring octree nodes are considered when looking for the nearest vertex from LOD0 which corresponds. This improves the appearance of lower LODs' vertex colors.
+
+    * An issue has been fixed where static meshes with imported LODs would not have the lower LODs' override colors set when "Copy instance vertex colors to source mesh" was used.
+
+    * Various crashes when using mesh paint mode have been fixed, e.g. when selecting a previously selected mesh, when deleting an asset which has an instance currently selected.
+
+* Improved MemoryAnalyser2:
+
+    * Upgraded to Visual Studio 2015.
+
+    * Clean up solution configurations. Only leave "Any CPU".
+
+    * Switch project to build with "Any CPU" rather than "x64".
+
+#### Content Browser
+
+* New: Added the ability to select a base Physical Material class when creating a new Physical Material.
+
+* Bugfix: Fixed a crash caused by the Content Browser allowing drag and drop operations to import assets while the File Picker dialog was opened.
+
+* Bugfix: Fixed a crash for an issue that could occur when a class and a folder have the same name.
+
+* Bugfix: Fixed an issue where Blueprint Actors would sometimes appear in the preview thumbnail of other assets.
+
+* Bugfix: Fixed an issue where a Skeletal Mesh asset was prevented from editing thumbnails in Thumbnail Edit Mode.
+
+* Bugfix: Fixed an issue where Material thumbnails would display an extreme close-up of the material rather than the entire preview asset.
+
+* Bugfix: Fixed an issue where there was a duplicate "Font" entry in the asset picker menu.
+
+* Bugfix: Fixed an issue with the selection of objects in the Content browser from WorldSettings.
+
+* Redirector objects are now deleted after selecting the "Fix Up" option from their right-click context menu.
+
+* Updated World thumbnails no longer cull primitives. 
+
+    * This is because the camera is very far away and if terrain pieces are culled, the level is not visible.
+
+#### Landscape
+
+* New: Added right-click options to fill and clear landscape layers
+
+* New: Added support for "Max Pitch Angle" and "Random Yaw" foliage options in procedural foliage.
+
+* New: Automatically filled the first weight-blended landscape layer (if a layerinfo is specified at landscape creation time).
+
+    * This stops the landscape being black when first created, and fixes problems with the first paint strokes not having proper falloff.
+
+    * The old workaround of painting a landscape layer over the whole landscape to make painting work right is no longer necessary!
+
+* New: Added a Generate Overlap Events flag to Landscape Actors.
+
+* Bugfix: Fixed a crash when cooking landscape for mobile.
+
+* Bugfix: Undoing the deletion of a landscape or changes to its properties will no longer cause a crash.
+
+* Bugfix: PIE no longer crashes when using Landscape Splines with bPlaceSplineMeshesInStreamingLevels.
+
+* Bugfix: Fixed a crash that could happen when switching to the Landscape Pattern brush.
+
+* Bugfix: Fixed a crash with the Landscape Circle brush.
+
+* Bugfix: Removed all use of deferred exec commands from the landscape code.
+
+    * This solves numerous crashes during commandlets, etc. with landscape infos not being initialized because they do not tick the engine.
+
+* Bugfix: Fixed an issue where after resampling or changing landscape component size, any new components that correspond to previously deleted components will remain deleted in the new landscape.
+
+* Bugfix: Fixed an "Object in another map" error when packaging a project containing Landscape for Android.
+
+* Bugfix: Fixed an "Object in another map" error when saving a Landscape streaming level after using Mobile Preview and the landscape Move to Level tool.
+
+* Bugfix: Fixed Global Illumination textures being visible in the content browser when r.GenerateLandscapeGIData is enabled.
+
+* Bugfix: Fixed that the "Rebuilding Hierarchical Instanced mesh component, please resave map" error wouldn’t go away after a resave.
+
+* Bugfix: Fixed an issue with foliage being placed on top of Blocking Volumes.
+
+* Bugfix: Fixed a hang that could happen when changing a material which is used on a Landscape.
+
+* Bugfix: Fixed an issue causing the incorrect UV coordinates when tessellation is enabled for Landscape.
+
+* Bugfix: Fixed an issue where the Landscape flatten target grid preview is displayed on the wrong Landscape when switching landscape target. 
+
+* Bugfix: Landscape "continuous" sculpting now works in multiple viewports. 
+
+* Bugfix: Fixed an issue with Landscape "Layer Alphamap Type" being visible when in "New Landscape" mode even though it only affects "Import".
+
+* Bugfix: Landscape layer import can now be undone.
+
+* Bugfix: Fixed an issue where Landscape spline control point meshes were being offset if "Place Spline Meshes in Streaming Levels is ticked. 
+
+* Bugfix: Fixed an issue where Navmesh generation would use Landscape simple collision issues with Landscape Holes.
+
+* Bugfix: Fixed an issue where new Landscapes would render with the material’s layer preview weights instead of using the actual (all 0) weights.
+
+* Bugfix: Fixed an issue where Landscape Splines were not able to be extended if the current level was different to the selected spline.
+
+* Bugfix: Fixed an issue that would prevent selecting Landscapes from the World Outliner while in the "new Landscape" tool.
+
+    * It will now switch to editing whenever Landscape you select.
+
+* Bugfix: Added the ability for Landscape Grass Type to exclude decals by using the Receives Decals member.
+
+* If a landscape material contains both a "Layer Blend" node and other landscape nodes referencing some of the same layers (e.g. sampling a layer for grass), the landscape editor will now use the ordering from the blend node for the target layers list.
+
+* Stopped Landscapes compiling so many shader combinations.
+
+#### Material Editor
+
+* Bugfix: Fixed a crash when editing text in "Libraries Category Text" field of a material function. 
+
+* Bugfix: Fixed a crash when pasting a material function call node from one project to another in which it is not defined. 
+
+* Bugfix: Fixed an initial hitch when dragging around in a Color Picker opened from a Material Expression node.
+
+* Bugfix: Fixed undo and redo to now work correctly when making changes to Material Functions. 
+
+#### Matinee
+
+* Bugfix: Fixed an issue preventing keyframes from being added when an asset browser was opened in the same window as Matinee.
+
+#### Cascade
+
+* New: Updated the Particle Event Receiver Spawn module to now filter collision events based on Physical Material.
+
+    * The Physical Material of the collision is now also accessible in Blueprint.
+
+* Bugfix: Fixed a crash while in Cascade when enabling/disabling a module.
+
+* Bugfix: Fixed an issue where Instance parameters for Spawn, Required, and Typedata modules in Cascade were not being auto populated.
+
+* Renamed the two "Collision" choices in the Collision submenu to "Actor Collision" and "Scene Depth Collision" in the Cascade editor..
+
+#### Persona
+
+* Bugfix: Crash fix with crash or data corruption when reverting changes to the "Event" structure in the animation notify editor for montage assets.
+
+* Bugfix: Fixed a crash when deleting curve keys after animation keys had also been removed. 
+
+* Bugfix: Fixed the "Isolate" option on Skeletal Mesh materials that was not working correctly with Clothing asset sections. 
+
+#### PhAT
+
+* Bugfix: Fixed scaling widget in PhAT in some viewport modes.
+
+#### World Outliner
+
+* Bugfix: Fixed an issue where Actors were not being removed from the World Outliner when they were added and removed on the same frame.
+
+* Bugfix: Fixed an issue where the Empty Level is named "NewWorld" in the World Outliner as opposed to "Untitled".
+
+#### Static Mesh Editor
+
+* Bugfix: Fixed a crash caused from trying to import an FBX file after a previous import attempt had failed. 
+
+* Bugfix: Fixed a crash from the FBX importer when it attempted to import a mesh only made up of degenerate faces.
+
+    * An error is now logged if the importer tries to import a mesh with no valid faces.
+
+* Bugfix: Fixed a crash that would happen when reimporting a static mesh that was exported from a geometry brush asset.
+
+* Bugfix: Using the Material list reset button in the Static Mesh Editor now prevents the Static Mesh Component materials array from growing.
+
+* Bugfix: Reimporting or overriding a LOD other than the base mesh now prevents the Material list array from growing. 
+
+#### Source Control
+
+* New: Git Plugin implemented the Sync operation to update local files using the git pull --rebase command.
+
+* New: Git Plugin now allows you to make an initial commit when initializing a new project.
+
+* Bugfix: Fixed an issue where the wrong status icons were being shown.
+
+* If multiple assets need to be checked out, a single notification is now created for all files rather than creating one notification for each file. 
+
+* When submitting changed assets to source control, files that were previously deleted are no longer resubmitted to source control.
+
+#### World Browser
+
+* Bugfix: Sub-Level layers were not showing up in the Layers tab when loaded as part of the world composition.
+
+* Bugfix: The World Composition levels window was not correctly updated when creating or renaming levels using the Content Browser.
+
+* Bugfix: World Composition origin shifting was not accounting for shifting the Dynamic Directional Light shadow casting.
+
+* Bugfix: World composition sub-levels that had a child actor inside were becoming modified if any other sub-level was unloaded.
+
+#### Foliage
+
+* New: Added support for Lighting Channels to Foliage and Landscape Grass.
+
+* New: A Foliage Type now includes all relevant fields from its Body Instance for editing.
+
+* Bugfix: Fixed an issue where the "Rebuilding hierarchical instanced mesh component, please resave map" warning was showing up when loading or entering PIE in maps with valid foliage.
+
+* Bugfix: Fixed an issue where deleting a Foliage actor from the Foliage menu did not remove actors from PIE until the editor was restarted.
+
+#### Cooker
+
+* New: Added more options to the Cooker for selecting Garbage Collection frequency:
+
+    * The new Cook Settings’ Max Memory Allowance setting is a percentage value which will trigger Garbage Collection if exceeded.
+
+    * The new Cook Settings’ Min Free Memory setting allows you to set a minimum free memory value which will trigger Garbage Collection.
+
+* New: Added -skipcompile flag for running from Visual Studio.
+
+    * This will skip compilation when the -multiprocess flag is used.
+
+* Bugfix: Fixed a deterministic cooking issue with Static Mesh Component caused by using a Static Mesh before postload was called.
+
+* Bugfix: Newly discovered session instances were not being added to an existing session in the Session Browser.
+
+* Bugfix: Selecting an instance in a session with multiple instances did not deselect the previously selected instance correctly.
+
+* Bugfix: Deterministic cooking issue where the shader name was stored inside the shader code for debugging purposes, but is now removed.
+
+* Added a minor cooking optimization to unsolicited package markup.
+
+* Added an optimization to string asset reference resolution to bring back performance lost from memory optimizations.
+
+* Prevented clients from cooking when building only server configs.
+
+* Improved cooking performance:
+
+    * Allow cooker to save packages opportunistically if the current package is compiling shaders.
+
+    * Allow cooker to load additional packages if blocked by other packages and memory is free.
+
+* Improved cooker peak memory usage:
+
+    * Moved resolving string asset references to be just before package save instead of after load.
+
+    * Provide opportunity to garbage collect more frequently.
+
+#### Sequencer
+
+* New: Added an option to maximize the viewport when starting recording with the sequence recorder.
+
+* New: Added option to Convert to Possessable.
+
+* New: Added option to record actors in the Sequence Recorder as possessables.
+
+* New: Added play/pause/stop events to Level Sequence player as it was with the Matinee actor.
+
+* New: Added "Reverse Play" and “Change Playback Direction” functions.
+
+* New: Added support for exporting the scene and animation to FBX. This functionality should match the Matinee FBX export.
+
+* New: Added visibility options to show/hide/only when selected 3D trajectories per transform section.
+
+* New: Added an option to rewind the sequence when a recording is started. This defaults to 'on'.
+
+* New: Added hotkey (alt-R) for recording selected actors with the sequence recorder.
+
+* New: Added more play controls - shuttle backward, pause, shuttle forward (j, k, l).
+
+* New: Added arbitrary property recording to the Sequence Recorder.
+
+* New: Added support for additive Skeletal Animations. Evaluates all overlapping Skeletal Animation sections.
+
+* New: Added support for montage based animation from Sequencer
+
+* New: Added the ability to add burn-ins to Level Sequences.
+
+* New: Added a draw tracking focus debug point for Cinematic Camera actors.
+
+* New: Changed "Actor Name" in Sequence Recorder to be a pulldown menu so that the user can choose the actor directly rather than typing in a name.
+
+* New: Display name of next Sequence in Sequence Recorder.
+
+* New: Expose compressed EXR frames to Sequencer Capture UI and command line.
+
+* New: Added {shot} and {shot_frame} format args for movie captures.
+
+* New: Track colors have been tweaked.
+
+    * Audio track brighter. 
+
+    * Transform, bool, event tracks less saturated. 
+
+    * Recording subsection more saturated. 
+
+    * Fade track gradient.
+
+* New: Added support for the named "PerformanceCapture" event which, like Matinee, calls Performance Capture in the Engine to output a screenshot when the event fires.
+
+* New: Added the ability to show the Context Menu for more than one selected node. 
+
+    * Certain operations (ie. Lock, Active, Label) now operate on all the selected nodes.
+
+* Bugfix: Fixed a crash when right-clicking on a grouped vector keyframe in the timeline. 
+
+* Bugfix: Fixed a crash which would happen sometimes when rendering a movie in a separate process.
+
+* Bugfix: End any transactions on mouse down that weren't ended cleanly. 
+
+    * This fixes a bug where if you drag with the left mouse button and click with the right mouse button, the transactions are left in an ambiguous state.
+
+* Bugfix: Updated the Event track to no longer fire if the playback status is stopped.
+
+    * This fixes a bug where when playback is stopped and the movie scene sequence is returned to the start of playback, we don't want all the events from the last playback position to the start of playback to fire.
+
+* Bugfix: Find corresponding PIE world actor for "Actor to Record". 
+
+    * This fixes a bug where if you set the actor before PIE and then record, the recording does not work.
+
+* Bugfix: Fixed an issue where Skeletal Animations would evaluate one frame ahead of where they were supposed to be.
+
+* Bugfix: Fixed an issue where Transforms would not be captured if an anim recorder exists but the skeletal mesh that the anim recorder is supposed to capture doesn't exist.
+
+    * This fixes the first person template character not getting recorded.
+
+* Bugfix: Fixed the countdown timer in the Sequencer Recorder to now draw in Simulate mode as well as Play.
+
+* Bugfix: Fixed dragging the leading edge of a skeletal animation section to now adjust the start offset of the animation clip.
+
+* Bugfix: Key editor commit when user tabs away would lose focus. Key navigation buttons and track color button are no longer focusable.
+
+* Bugfix: Fixed an issue where the Player would stop after 60 seconds. Now assigns the player to a UPROPERTY so that it is not Garbage Collected.
+
+* Bugfix: Fixed an issue where the level visibility track would mark animated levels as dirty.
+
+* Bugfix: Updated spawnables so that they can not be deleted from the viewport.
+
+* Bugfix: Fixed an issue where Blueprint classes as spawnables were not being re-instanced correctly.
+
+* Bugfix: Fixed an issue where jitter would happen when jumping around in Sequencer.
+
+* Bugfix: Fixed an issue where Sequencer would redraw all viewport when in PIE.
+
+* Bugfix: Updated node selection so that the selected nodes now only update if they have changed.
+
+    * This fixes a bug in the curve editor where undo re-selects and auto-frames.
+
+* For legacy, properties with the CPF_Interp flag can be animated in Matinee. It does not require the CPF_Edit flag as well. This makes Sequencer consistent with Matinee.
+
+* Optimized adding keys to rich curves to improve performance when recording data with the sequence recorder.
+
+    * Set the times and values at once as an ordered set.
+
+* Removed curve editor visibility as a toggleable config.
+
+    * It is now just a toggleable state that defaults to false. 
+
+    * This fixes unexpected behavior of staying in the curve editor when restarting the editor or switching to a different level sequence asset.
+
+* Renamed "Expose to Matinee" to "Expose to Cinematics" in Blueprints.
+
+* Select actors for corresponding selected keys or sections.
+
+* Set event track evaluation order to fire first. 
+
+    * This fixes some ambiguity and also a bug where transform tracks don't evaluate completely if the eval position is changed in the middle of evaluation.
+
+* Set relative location, rotation, scale explicitly to identity instead of calling Reset Relative Transform because we don't want overlaps to fire until after the update pass.
+
+* Improved recording indicator icon, text & timer.
+
+#### VR-Editor
+
+* New: Updated foliage mode to now be fully integrated with VR Editor mode. You can paint down foliage while in VR!
+
+* New: Added World button to the Quick Select menu to allow access to the World Settings menu.
+
+* Bugfix: Fixed an issue where duplications could not be undone.
+
+* Bugfix: Updated laser double-click events to now be sent when the trigger is pressed rather than when it is released.
+
+* Bugfix: Fixed an issue where the Mesh Paint tool was not showing the Brush Preview.
+
+* Bugfix: Fixed an issue where Foliage tools were active even without the controller’s trigger being pressed. 
+
+    * Foliage reapply tool.
+
+    * Foliage Lasso tool.
+
+* Bugfix: Fixed an issue where the laser’s hover point in the VR Editor actor wasn’t projected onto the corrected plan while rotating, translating, and scaling using gizmo.
+
+    * This fix snaps the hover point to the nearest line or plane according to the transformation.
+
+* Bugfix: Fixed an issue where the brush preview was still present while using a Foliage brush and hovering on top of the UI. 
+
+* Removed the "Simulate" button in the toolbar of the blueprint editor while in VR Editor mode because we currently cannot simulate in this mode.
+
+### Gameplay Framework
+
+* New: Added support for multi-jumping to Character Movement, and fixed infinite jumps when Jump Max Hold Time is set.
+
+* New: Added "Mixed" replication mode for active gameplay effects. 
+
+    * This will replicate active gameplay effects normally to their ability systems component's owner, but replicate a stripped down version to non-owners. This can help reduce bandwidth by as much as 2 KB per second.
+
+* New: Added "Get Parent Actor" convenience function.
+
+* New: Added Activate and Deactivate events to the Actor Component.
+
+* New: Added arc-tracing functions "Predict Projectile Path" and “Suggest Projectile Velocity Medium Arc” to Gameplay Statics.
+
+* New: Added "Relative To Initial FOV" option to Camera Anim, defaulting to true.
+
+    * If turned off, camera anims will use the camera's current FOV as the initial FOV for the animation.
+
+* New: Added "Is Valid" Blueprint-exposed method for Gameplay Attribute.
+
+* New: Added max depenetration distance settings for the Character Movement Component. Added controls to throttle logging when character is stuck in geometry so it doesn't spam the log.
+
+    * Depenetration settings are separated based on whether overlapping a Pawn versus other geometry, and furthermore by whether the Character is a proxy or not. Simulated proxies typically should not depenetrate a large amount because that effectively ignores the server authoritative location update.
+
+    * "Stuck" logging is controlled by the console var "p.CharacterStuckWarningPeriod". Set to number of seconds between logged events, or less than zero to disable logging.
+
+* New: Added new "Load Level Instance" Blueprint node, which allows to create multiple transformed instances of a .umap without having to include in persistent level's list. 
+
+* New: Added new Blueprint library math/vector functions: "Find Closest Point On Segment", “Find Closest Point On Line”, “Get Point Distance To Segment”, and “Get Point Distance To Line”.
+
+    * Fixed comments on "Find Nearest Points On Line Segments".
+
+    * Fixed comments on "Point Dist To Line", and renamed "Line" parameter to "Direction".
+
+* New: Added player Connect and Disconnect events to GameMode.
+
+* New: Added the ability to separate players in local multiplayer for the Keyboard and Gamepads by using the setting "Offset Player Gamepad Ids" in the Project Settings under Maps & Modes. 
+
+* New: Added the ability to Get and Set the controller ID of a player controller from Blueprints.
+
+* New: Added the ability to  Get and Set Tick Interval for Actor and Component primary ticks from Blueprints.
+
+* New: Updated Cheat Manager so it can now be subclassed via Blueprints and properties can be set for the Player Controller and Debug Camera Controller.
+
+* New: Exposed "Override With" and “Copy Properties” in Player State to Blueprint Children.
+
+* New: Made minimum quadtree size configurable in the Procedural Foliage Spawner.
+
+* New: Added "Is Player Controller" function to Controller. The variable already existed, but wasn’t exposed.
+
+* New: Added check to Post Actor Construction to avoid a Begin Play call on an Actor that is Pending Kill.
+
+* New: Anim Root Motion was applying too much velocity to Character Movement when framerate is low causing moves to be substepped, or when movement mode changes during root motion playback.
+
+    * Added "Anim Root Motion Velocity" to the Character Movement Component to track last velocity extracted from anim root motion.
+
+* Bugfix: Fixed a crash when cancelling PIE while in the middle of a seamless travel as a client.
+
+* Bugfix: Fixed a crash in Game Session if the GameMode’s Player State Class has become null.
+
+* Bugfix: Updated Character movement to make sure the Last Update Location, Rotation, and Velocity are updated on client and server error corrections. The Force Position Update function calls Perform Movement regardless of velocity (there may be root motion or gravity effects).
+
+* Bugfix: Updated Child Actors in a sublevel to no longer lose customizations made by the owning Actor's construction script when hiding and then unhiding the level.
+
+* Bugfix: Updated Child Actors to no longer sometimes collide with objects at the origin when spawned.
+
+* Bugfix: Updated Data Assets to now correctly report their size.
+
+* Bugfix: Fixed an issue where Actor Attachment Replication was not being cleared on detachment, which would cause issues with the actor's location for clients joining in progress.
+
+* Bugfix: Fixed an issue with Model Component replication in client/server PIE if BSP has been rebuilt.
+
+* Bugfix: Updated comments on "Load Game From Slot" and “Delete Game in Slot”.
+
+* Bugfix: Fixed an issue where Mesh smoothing on clients was causing a popping briefly when crouching. 
+
+    * This was due to a change in 4.12 where we started smoothing from Z location rather than always zeroing it (in certain movement modes).
+
+* Bugfix: Fixed an issue where the Root Motion network connections were not clearing the Root Motion Data, which would cause clients to jitter.
+
+* Bugfix: Fixed an issue where it was not possible to stream in a sub-level whose name is a substring of another sub-level.
+
+* Bugfix: Fixed an issue where an Instanced Static Mesh Component would not keep its instances when copy-pasting in the Editor.
+
+* Bugfix: Updated the "Hidden In Game" checkbox to now be properly displayed when editing Scene Component properties.
+
+* Bugfix: Improved pasting multiple cells into the Property Matrix to no longer depend on your selected tiles, only your target cell.
+
+* Bugfix: Prevented forced drawing of the debug visualization for spline components in Test configuration.
+
+* Removed GWorld from the Timer Manager.
+
+* Updated "Remove Instances" on a Hierarchical Instanced Static Mesh Component to now properly rebuild the cluster tree, allowing “Get Instances Overlapping Box” to return correct results.
+
+* Updated significance to now be correctly calculated when there are multiple viewpoints and significance is being sorted descending.
+
+* Updated HUD so that Show HUD is no longer a config variable, but can be edited from a Blueprint.
+
+* Improved "Move Component To Rotation".
+
+* Added minor performance improvements in Player Input and Input Component by inlining and avoiding memory allocations.
+
+* Updated numerous HUD properties to make them readable, and most of them also writable, from Blueprints.
+
+* Updated Components so they can no longer be added during world tear down.
+
+* Improved character movement linear smoothing in the presence of low net frequency updates.
+
+* Improved comment about the Blocking Hit property in Find Floor Result, explaining it is a valid blocking hit that was not in penetration. Other conditions can be determined from the Hit Result itself.
+
+* Made up/down behaviour for console history consistent and reverted to old ordering by default
+
+* Moved Uninitialize Components outside (Actor Initialized) to always uninitialize components when actors gets destroyed early.
+
+* Updated multi-arg console commands to now accept string commands with or without quotes.
+
+* Optimized NaN and Infinite checks for FVector, FQuat, FRotator, FMatrix, FTransform. SIMD version "Vector Contains NaN Or Infinite" function also optimized on relevant platforms. Added startup tests for “Vector Contains NaN Or Infinite”.
+
+    * All our "Contains NaN" tests are in fact "Is NaN or Is Infinite", which is the same as "not Is Finite", so exploited this to simplify the tests. Both NaN and +/-Inf are not finite (http://en.cppreference.com/w/cpp/numeric/math/isfinite). In the future we should rename "Contains NaN".
+
+### Learning Resources
+
+#### Content Examples
+
+* New: Added TP_VirtualRealityBP to list of Blueprint templates
+
+### Localization
+
+* New: Updated the Dialogue Wave voice actor so the direction can now be localized.
+
+    * Also, exported dialogue scripts now include a column that says whether they have a localized recording of that line of dialogue.
+
+* Bugfix: Fixed a crash caused by bad access of shared ‘this’ when closing an active IME context.
+
+* Bugfix: Improved changing the loading phase of a target in the localization dashboard to now write to the default config.
+
+* Bugfix: Fixed an issue where game targets were not being able to depend on other game targets in the Localization Dashboard.
+
+* Bugfix: Fixed Issues with gathering text from Blueprint Byte code. 
+
+    * Bytecode in Blueprints is very volatile, and can only be safely gathered after it's been compiled (which is not guaranteed to have happened by the time we save the package). This change avoids caching any assets that contain scripts (non-data-only Blueprints), and instead will always load them to perform a gather (which will ensure the Blueprint bytecode is up-to-date due to compile-on-load).
+
+* Bugfix: Fixed an issue with Unicode string handling in the web browser widget’s Execute Javascript method.
+
+* Updated Localization compilation to now specify an Archive Name to use.
+
+* Made sure that en-US-POSIX is in our list of available cultures.
+
+* Updated PO files to now only contain a single entry in the case of a native translation being exported.
+
+    * They used to contain the original entry, as well as an entry for the native translation, however the original entry would never be used.
+
+### Networking
+
+* New: Added a "-CapturePercentage" command line parameter for Oodle which sets a percentage change of a connection to generate an Oodle capture.
+
+* New: Added ability to amortize replay checkpoints. Use the console variable "demo.CheckpointSaveMaxMSPerFrame" to set the maximum time allowed (in ms) to spend saving checkpoints per frame.
+
+* New: Added ability to get network version from FArchive (for replay compatibility support). Use Engine Net Ver/Game Net Ver.
+
+* New: Added ability to skip over changed properties in Fast Array Serializer for replays (and will log incompatible property names).
+
+* New: Added ability to skip over RPC parameters that have changed in replays for backwards compatibility.
+
+* New: Allow Oodle captures to be generated in shipping configurations.
+
+* New: Improved support for streaming levels in replays.
+
+* New: Created a batch file, ExtractAllArchives.bat, which can pull down Oodle captures stored in Amazon S3.
+
+* New: Created a batch file, BuildDictionary_Generic.bat, that can generate an oodle dictionary from capture files in an arbitrary location.
+
+* New: Low level network code will now notify the game of two new error types to be handled (Net Guid Mismatch or Net Checksum Failures). This kind of errors occur when there are incompatibilities in between client and server, usually because different cooked data is involved. This shouldn't occur in properly shipped builds but is good during development when the possibility to mix and match builds might happen on accident.
+
+* New: Made some Oodle trainer commandlet improvements:
+
+    * Added a command line parameter to many Oodle trainer commandlet tasks for a changelist number to only process capture files from a specific changelist.
+
+    * When processing Oodle capture files, the Oodle trainer commandlet will now also look recursively in subdirectories.
+
+    * Refactored how the Oodle trainer commandlet processes capture files to make it easier to add new command line options for different tasks.
+
+* New: Added some miscellaneous seamless travel cleanup + sanity checks.
+
+    * Cleanup the naming of variables and logic used to pick which actors will transfer to new world during seamless travel to more explicitly convey intent.
+
+    * Warn when there is no context for a world when destroying actors. It's likely that this world is the floating world during seamless travel, and if we hit this case, something is likely in a bad state
+
+* New: Updated Oodle packet capture file names to now contain the branch name, the changelist of the build, and process ID used to generate the capture.
+
+* New: The console variable "demo.AsyncLoadWorld" can now be overridden with the URL option "Async Load World Override" when playing a replay.
+
+    * When this setting is true, the engine will use seamless travel to load the replay map instead of using the default non-seamless travel method.
+
+* New: The stat for the number of network actors (STAT_NumNetActors) is now tracked on both the server and client.
+
+* New: Updated Oodle compatibility to version 2.1.5.
+
+* New: When properties are incompatible in replays, they will now be shown in the log.
+
+* Bugfix: Fixed a crash that could occur when using "open" command in a PIE window.
+
+* Bugfix: Fixed a crash that would occur when opening a project from a network drive.
+
+* Bugfix: Fixed assert in channel cleanup code that could occur if the connection was cleaned up, and there were "Keep Processing Actor Channel Bunches Map" in-flight still.
+
+* Bugfix: Fixed incorrect warning when Received Bunch fails for replays.
+
+* Bugfix: Fixed a memory leak when recording a replay on a client that was built using a client-only executable configuration.
+
+* Bugfix: Fixed an issue causing improper http header parsing when one or more of the headers are empty.
+
+* Bugfix: Fixed an issue that could cause character movement to be saved incorrectly in replays.
+
+* Bugfix: Fixed an issue that could cause the spectator controller for a client-recorded replay to incorrectly report that it was local controller during recording.
+
+* Bugfix: Fixed an issue that was causing RPCs to be missing from replays that were recorded on clients.
+
+* Bugfix: Fixed several issues and playback artifacts with character positions in replays.
+
+* Deprecated the  "Engine Net Version", “Engine Min Net Version”, and “Engine Negotiation Version” (use the Network Version struct instead).
+
+* Disabled the ability to use hot-reload while PIE windows are open.
+
+* Added extra logging in the form of "UPendingNetGame::InitNetDriver: Sending hello: 127.0.0.1" to help diagnose connection issues.
+
+* Highly optimized replay checkpoint recording to take about 50% less time.
+
+* Improved playback quality of character movement for replays that are recorded on a client.
+
+* Keeping track of the number of actors stat (STAT_NumActors) is now more efficient.
+
+* Improved 'Too many calls to RPC' log message informativeness.
+
+* Removed unused "In Bunch" constructor.
+
+* The Oodle trainer commandlet will now only process capture files that end with the .ucap file extension.
+
+### Online
+
+* New: Added PS+ check to account privilege check.
+
+* New: Move Curl and PS4 HTTP request processing to a separate thread. HTTP manager will now start and tick HTTP requests on its own thread, removing the processing from the game thread. Hitches in HTTP requests no longer impact the game thread, and HTTP requests can be ticked more frequently by adjusting [HTTP] Http Thread Tick Rate in .ini config files.
+
+* New: Online Subsystem modules have been moved to Plugins, out of Engine/Source. Engine's dependency on OnlineSubsystem modules have been removed, in some places replaced with a new interface UOnlineEngineInterface for a small subset of features Engine requires. UOnlineEngineInterface is implemented in UOnlineEngineInterfaceImpl. FUniqueNetId and FUniqueNetIdWrapper have been promoted to a basic type in core, but they are still opaque types and must be created by online subsystems.
+
+* New: Added -port command line option for the Steam Online Subsystem.
+
+* Bugfix: Fixed a rare bug in the Build Patch Tool that would cause a manifest to not contain a chunk's information.
+
+* Bugfix: Fixed shutdown crash in HTTP request completion delegate.
+
+* The "Cancel Find Sessions" function in Online Session did not null out search result for LAN results. Future Find Session calls would fail with "search in progress".
+
+### Other
+
+* New: Changed the Gameplay Ability Actor Info's Anim Instance property to a Skeletal Mesh Component.
+
+* New: Game window titles now show only the project title by default in shipping builds.
+
+    * The default window title for non-shipping builds remains unaltered.
+
+    * If there is a value stored in the Project Displayed Title setting, it will be displayed for all configurations.
+
+    * If there is a value stored in the Project Debug Title Info setting, it will be appended to the window title in non-shipping builds.
+
+* Bugfix: Fixed a crash after clearing Sound Wave Asset and Deleting Media Player from Content Browser.
+
+* Bugfix: Fixed a crash and incorrect re-initialization of media texture resource.
+
+* Bugfix: Fixed failing to load Precached Media Source.
+
+* Bugfix: Fixed Local file paths without file:// open when pressing Enter, but not when clicking Go button.
+
+* Bugfix: Fixed Media player does not play the second item in the playlist on PS4.
+
+* Bugfix: Fixed media source asset cannot be inherited in other modules.
+
+* Bugfix: Fixed memory leak in the Wmf Media source resolver.
+
+* Bugfix: Fixed packaged game does not load in Windows XP due to Wmf Media.
+
+* Bugfix: Fixed problem where users couldn't close the Crash Report Client without sending the crash report.
+
+* Added code in Auto Populate Instance Properties to populate vector and scalar parameter data.
+
+* Allow symbol files to be skipped when staging build products as they are not essential for the staged project to run.
+
+* Increased maximum compressed payload size for crash reports from 16MB to 32MB.
+
+* Made PDBs optional build products for Cs Compile Build Graph task and added .dll.mdb check for Mono equivalent.
+
+### Paper2D
+
+* New: Exposed Paper Tile Map and Paper Tile Set as Blueprint Types to make it easier to work with tile maps in Blueprints.
+
+* Bugfix: Fixed a crash when mutating grouped sprite components at runtime, and cleaned up how collision rendering is done for grouped sprite components addressing some other issues as well.
+
+* Bugfix: Fixed a bug that caused the tabs in the Flipbook, Sprite, and Code Project editors to show the editor name rather than the asset name.
+
+* Bugfix: Fixed an issue causing incorrect normals (and therefore lighting) on tilemaps in a number of circumstances including standalone or cooked builds.
+
+* Bugfix: Fixed BP-created tile map components being incapable of having collision generated for them (still requires calling Set Layer Collision with rebuild=true or Rebuild Collision).
+
+* Bugfix: Fixed inconsistent normals / lighting on lit grouped sprites.
+
+### Physics
+
+* New: Added option to allow skeletal mesh simulation to NOT affect component transform.
+
+* New: Added support for 'teleport' option when changing the Instanced Static Mesh transform.
+
+* New: Allowed access to the Enable Gravity flags on individual bones within a Skeletal Mesh Component using Bone Name.
+
+* New: Allowed 'show collision' console command to work in Test configurations.
+
+* New: Brought back cone constraint visualization, but with some improvements like colors and free vs locked. Turned on constraint visualization when drawing a skeletal mesh with constraint show flags.
+
+* New: Changed constraint warnings to now show up in Message Log to be more visible when testing.
+
+* New: Introduced a default shape complexity project setting (Complex As Simple vs Simple As Complex vs Simple And Complex).
+
+* New: Updated physics blend weights to now only affect rendering data.
+
+* Bugfix: Fixed crash that could occur when the Phys Scene pointer was null (the world has no Physics Scene) in the Skeletal Mesh Component’s Term Articulated function. That could happen when loading a world without fully instantiating it, such as when right-clicking a world in the context browser rather than opening the world directly.
+
+* Bugfix: Fixed crash when using very large simulation meshes for clothing - limits have been increased and it is no longer possible to import a simulation mesh that is too large.
+
+* Bugfix: Fixed a crash in sub-stepping by limiting skeletal mesh bodies to be in either the sync scene or the async scene, but not both. Simplifies a lot of code and improves performance.
+
+* Bugfix: Fixed hit events on a skeletal mesh component not respecting the AND between a skeletal mesh component and the ragdoll bodies.
+
+* Bugfix: Fixed the case where physics hit events were passing the wrong component's bone info to the hit event.
+
+* Bugfix: Fixed vehicle drag adding instead of removing energy when in reverse.
+
+* Bugfix: Fixed "Add Impulse At Location" for destructible meshes not using the correct location.
+
+* Bugfix: Fixed "Angular Rotation Offset" properly disabling when motion for axis is Free or Locked.
+
+* Bugfix: Fixed bug with overlap queries against welded bodies.
+
+* Bugfix: Fixed bug with physics body welding resulting in duplicate collision shapes.
+
+* Bugfix: Fixed clothing association sometimes using stale references to sections.
+
+* Bugfix: Fixed collision rendering in level editor when mesh was mirrored.
+
+* Bugfix: Fixed destructible meshes always appearing to change when generating patches.
+
+* Bugfix: Fixed "Max Angular Velocity" setting being reset unintentionally.
+
+* Bugfix: Fixed memory churn during PhysX simulate calls by providing per-scene scratch space for the simulation to use.
+
+* Bugfix: Fixed the Paper Grouped Sprite not updating physics bodies.
+
+* Bugfix: Fixed physics bodies of an Instanced Static Mesh Component not moving when the mesh is updated.
+
+* Bugfix: Fixed Physics Constraint Component not being editable on instances (only on defaults).
+
+* Bugfix: Fixed race condition in clothing code if cloth is cleared between Pre and Post Physics Tick.
+
+* Bugfix: Fixed rare contact pre-filter crash when invalid shapes were passed to pre-filter callbacks.
+
+* Bugfix: Fixed warning about using CCD on kinematic bodies.
+
+* Added more information when convex cooking partially succeeds (fails initially but succeeds with inflation) indicating the full path of the mesh that failed.
+
+* Corrected Convex Hull 2D so that it returns an empty set of indices when passed an empty points array.
+
+* Ensure that skeletal mesh components do NOT teleport unless explicitly asked to.
+
+* Removed call to set a frame decay constant on destructible shapes in PhysX. This was used to better settle large piles but was slower to process. PhysX can now handle this better so it is no longer necessary.
+
+### Platforms
+
+* New: Added a setting to make the game window preserve its content's aspect ratio while being resized by user. This is on by default.
+
+* New: Added support for encrypting .ini files from the Project Settings and the Project Launcher.
+
+* New: Added support for the Unreal Front End to launch to multiple devices at once.
+
+* New: Generalize Android Plugin Language to Unreal Plugin Language for modifying plug-in package data during packaging process
+
+    * iOS plist data can now be transmogrified during packaging with UPL.
+
+* New: Added "Get Game Bundle Id" Blueprint function.
+
+* New: Added support to UAT for building target platforms with multiple cook flavors ex: -targetplatform=Android -cookflavor=ETC1+ETC2
+
+* Bugfix: Fixed a crash when entering the Games App ID into the Android Project Settings menu before clicking "Configure Now" in the APK Packaging section.
+
+* Bugfix: Fixed Android device rule for Alcatel Pixi 3.
+
+* Bugfix: Analog input for Gamepad is now properly mapped on iOS.  
+
+* Bugfix: Browser protocol used will continued to be used and canvas focus especially when inside an iframe.
+
+* Bugfix: Resolved issue with reflection environments on Intel Macs running Mac OS X El Capitan and using Metal Shader Model 5 by disabling tiled reflections for Intel GPUs on this version of Mac OS.
+
+* Bugfix: Resolved an issue where iOS Movie Player can't handle videos at resolutions that aren't multiples of 16. 
+
+    * Movies that are not multiple of 16 will report a runtime error.
+
+* Bugfix: Issue where a movie wasn’t played on iOS occasionally. 
+
+* Bugfix: Error Message displays as Unknown Error when failing to supply a Remote Build server for iOS on Windows has been fixed.
+
+    * A user friendly error is now provided. 
+
+* Bugfix: Reduced the number of buffers used per-frame in Metal when updating textures by pulling the buffers from the pool rather than allocating them each time.
+
+* Bugfix: Message box could not be closed after accessing the home screen on iOS.
+
+* iOS improvements
+
+    * Fix an issue where invalid dsyms could be produced on shipping builds.
+
+    * Added Bitcode support.
+
+    * Add full support for producing XCArchive packages including debugging support for bitcode.
+
+* Change Mac OpenGL shader version at runtime to 3.30 to support more intrinsics such as asuint(), asfloat() used in some materials.
+
+#### Android
+
+* New: AAR and JAR files from repositories is now supported which allowed updating the following libraries:
+
+    * Google Play Games native C++ SDK to v2.1
+
+    * Google Play Services to 9.2.0
+
+    * android-support-v4 to 23.0.0
+
+* New: Added the ability to set thread affinity for a device in the Device Profiles (ex: +CVars=android.SetThreadAffinity=RT 0x02 GT 0x01).
+
+* New: Added detection of Amazon Fire TV to disable requiring virtual joysticks.
+
+* New: Added Harfbuzz texture shaping compiled libraries for all Android platforms.
+
+* New: Added optional log spew filtering callback to Run
+
+    * This is used with the new adb with percentage completed indicator which spammed the logs.
+
+* New: Added the ability to specify the desired install location for Android packages in Project Settings
+
+* New: Return player ID and handle auth token for Google Play Games on Android.
+
+* New: You can now package and launch on multiple connected Android devices simultaneously.
+
+    * Select multiple devices in the Project Launcher.
+
+    * It is recommended that you use the Android_Multi target and enable a set of texture formats to support all of your devices.
+
+* New: Added additional scope (plus.login) to Google Play Games builder. 
+
+* New: Added "r.AndroidDisableThreadedRendering" console variable and added device profiles for some devices that need threaded rendering disabled due to crash bugs when an application regains focus
+
+    * This fix was specifically applied to the following devices
+
+        * Lenovo S939
+
+        * Lenovo A7600
+
+        * Wiko Highway Signs
+
+        * Alcatel Pixi 3
+
+* New: Added handling of OnTextChanged for virtual keyboard input on Android.
+
+* New: Added HUAWEI Mali device recognition for Vulkan.
+
+* New: Added missing error check of the response code returned from consuming in app products.
+
+* New: Allocated Optional Shadow Depth Color Surface to match Depth Stencil dimensions (allow up to 4 resolutions for now) for mobile shadows.
+
+* New: Downloader Backgrounds can now be set for both landscape and portrait orientation on Android.
+
+* New: Enabled Google Play Games for ARM64.
+
+* New: Android Vulkan now include path checks. 
+
+* Bugfix: Fixed a crash on device when playing sounds when packaged using Android_Multi.
+
+* Bugfix: Fixed a crash with Vulkan on Nexus 6P with Android N enabled extensions. 
+
+* Bugfix: Fixed an issue with the Android string implementation for finding a WIDECHAR string within another WIDECHAR string.
+
+* Bugfix: Fixed an issue where the application could become unresponsive on some devices if the Game Activity is paused while an Alert Dialog is showing.
+
+    * Hid the active AlertDialog on pause and show it again on resume after the render thread is restarted. 
+
+* Bugfix: There was an assumption in Unreal Plugin Language that caused it to only look in Engine Configs for Project Settings.
+
+* Bugfix: Fixed an issue with iterative deploy for new ADB.
+
+* Bugfix: Fixed an issue with linker warning (mismatched function/variable for gl Map Buffer OES and gl Unmap Buffer OES).
+
+* Bugfix: Fixed an issue with proguard entry for Android mediaplayer tracks classes that are used in the new media framework. 
+
+* Bugfix: Handle Android media prepare failure that could happen (eg. URL without internet connection).
+
+* Bugfix: Fixed an issue where launching to an Android device from a computer that has not been authorized by the device would hang during deployment.
+
+* Bugfix: Fixed an issue where setting Anti-Alias to true in the Draw Line function in a Blueprint Widget would cause the line to not be drawn on Android.
+
+* Bugfix: Fixed an issue where taking a screenshot on Android using the command console would result in an upside-down screenshot.
+
+* Bugfix: Fixed multiple Vulkan extensions for Android.
+
+* Removed unused methods from the Android JNI header.
+
+* Renamed some functions with "Error" in the name to prevent false coloring in the logs.
+
+* Use Malloc Ansi instead of Malloc Binned on Android ARM64 for now.
+
+* Use mmap/munmap for Android Binned Alloc From OS.
+
+* Cleaned up AndroidProperties.h (removed old platform types).
+
+* Deal with missing Android movie frame rate by defaulting to 30 if not available.
+
+* Device Tilt values on Android will now be consistent with values on iOS.
+
+* Disable Oculus stress tests on Android (for now) to remove shader dependency unhandled on Mac editor.
+
+* Disabled default engine crash handler on Android. System crash handler works more reliably across different Android versions and devices.
+
+#### iOS
+
+* New: Added flags to optionally produce flat dsym files or dsym bundles for third party crash reporting tools.
+
+* New: Added a new virtual keyboard method that no longer uses a pop-up, you can type directly into the Slate widget.
+
+    * Off by default, because the keyboard may cover your text boxes.
+
+    * Use -newkeyboard on the iOS command line to enable it.
+
+    * Also, see the On Virtual Keyboard Shown and On Virtual Keyboard Hidden delegates that will let you slide your UI out of the way.
+
+* New: Added support for local notifications, device orientation change callbacks, and Blueprint platform delegates in BlueprintPlatformLibrary.h.
+
+* New: Enabled iOS building from PC utilizing a remote Mac for the build.
+
+* New: Added additional log information when using an unsupported audio type.
+
+* Bugfix: Fixed a crash when accessing Message Box Ext in Shipping builds on iOS.
+
+* Bugfix: Fixed a crash when initializing push notifications on iOS 7.
+
+* Bugfix: Fixed a crash for an occasional crash while running Metal on iOS 8 devices. 
+
+* Bugfix: Resolved issue for iOS device name mappings. 
+
+* Bugfix: Fixed an incorrectly formed path issue has been addressed when trying to deploy to a device from Mac.
+
+* Bugfix: Fixed an issue with the remote build from PC to Mac where occasionally the build would fail with linker errors and then succeed on the next build.
+
+* Bugfix: Fixed an issue with checks for iOS SDK on PC when performing package operations. 
+
+* Increased the stack size on IOS and Mac when running in development. 
+
+    * Added support for the Apple Replay Kit Framework.
+
+* IPhonePackager and supporting assemblies are now set to be 64-bit by default.
+
+* Prevented crashes in iOS games using Metal when the game moves to the background and loses access to the frame buffer by ignoring draw calls until it is resumed.
+
+* Properly return success or failure when utilizing an IAP blueprint node the very first time in the game.
+
+* Reduce the number of shaders to compile on iOS for Metal by only caching Shadow Depth VS by setting Is For Geometry Shader to true when the RHI can handle the underlying feature.
+
+* Reducing memory footprint of engine on iOS by reducing the number of objects to hold before garbage collection. Further reduced memory when using Metal by reducing the command buffer sizes.
+
+* Removed local notification support for tvOS as they are not actually supported in the OS.
+
+* tvOS applications may now have an underscore in the name.
+
+* Changed all references to AppleTV (related to the OS) to tvOS in the editor.
+
+* Disabled MRT support on iPad Air 1, because it doesn't support wide enough MRTs for the engine to use them (ie for GPU particles)
+
+* Disabled user interaction when utilizing SSH to prevent hangs.
+
+#### Linux
+
+* Bugfix: Corrected error message when not specifying linux server device command line correctly in UAT.
+
+* Bugfix: Arguments passed to clang for the Linux Tool Chain now properly escape. 
+
+#### Mac
+
+* New: Added a command-line switch for Mac "-RedirectNSLog" that will capture the output of NS Log using an NS Pipe and write it into our log instead when not running under the debugger, thereby capturing output from underlying libraries that would otherwise be lost when outside a debugger.
+
+* New: Added more Metal statistics that track the number & memory size of id<MTLBuffer>'s allocated and released each frame.
+
+* New: Added more information to the engine log for, OS X version information, Mac model, CPU name, and number of virtual cores to improve our ability to reproduce bugs and report errors to Apple.
+
+* New: Added View->Enter Full Screen menu item for games on Mac.
+
+* New: Consolidate optional Metal feature support in Metal RHI for the latest versions of Apple's products.
+
+    * Support native Depth-16 texture formats where available.
+
+    * Use native stencil texture views for accessing the stencil component of a depth-stencil texture rather than copying the data to a temporary texture when available.
+
+    * Expose support for the optional Base Vertex Index & First Instance draw call arguments on newer iOS devices that support them as well as on Mac.
+
+    * Expose support for indirect draw and dispatch calls on newer iOS devices as well as on Mac.
+
+* New: Implemented rudimentary runtime debugging of Metal RHI with several different levels of complexity controlled by the new console-variable "rhi.Metal.RuntimeDebugLevel".
+
+    * By default "rhi.Metal.RuntimeDebugLevel" is 0 (off) & it is compiled out for Shipping - when enabled each higher level collects more data or changes the Metal RHI behaviour to make a command-buffer failure mode debuggable at the expense of performance.
+
+    * In order for this to work the number of command-buffers that may be allocated from the command-queue is also controllable via a console-variable "rhi.Metal.CommandQueueSize" which must be set prior to startup.
+
+* New: Initial AV Foundation implementation of Media Framework for Mac, iOS & tvOS.
+
+    * Supports playback of local movie files in .mp4, .mov & .m4v containers using the H.264 video format.
+
+    * When using Metal video frames are supplied directly as Metal textures without conversion.
+
+    * On Mac audio is played through the Media Audio Sink provided to the movie player but currently on iOS & tvOS audio will play through the speakers directly.
+
+* New: The Metal shader compiler will now retain the text shader code when "r.Shaders.KeepDebugInfo" is enabled to allow for easier debugging.
+
+    * For Mac, shaders are still compiled to bytecode and the Mac runtime will load the bytecode, while debug tools will load the text shader from the system temporary directory.
+
+    * On iOS & tvOS the runtime will load the text shaders so that they appear in the debug tools.
+
+* New: Added validation code to Metal RHI to ensure that we do not erroneously re-clear existing render contents when restoring the render command encoder to continue rendering.
+
+    * This will only be active when the Metal debug layer is enabled and is completely compiled away in Test & Shipping builds for performance.
+
+* Bugfix: Crash fix when loading an empty sub-level on iOS projects using Metal by adding validation of resource bindings.
+
+    * Setting "rhi.Metal.RuntimeDebugLevel" to 1 or greater enables the resource binding validation.
+
+    * Fixed the underlying error in Slate which was attempting to bind an uninitialised texture atlas.
+
+* Bugfix: Crash fix to always bind the SpeedTree uniform buffer data when required by the shader when importing SpeedTree assets if running the Mac Editor on the Metal graphics API.
+
+* Bugfix: Crash fix in Core Audio caused by adding and deleting the same node in a single tick.
+
+* Bugfix: Crash fix with Metal rendering one-pass point light shadows by ensuring that we bind the correct type of texture to the shader.
+
+* Bugfix: Crash fix with Metal's Update Texture 2D implementation on Mac by adding code to perform the required conversion from PF_G8 + sRGB to RGBA8_sRGB.
+
+* Bugfix: Crash fix on shutdown in debug caused by Metal RHI trying to update internal class counts after module shutdown.
+
+* Bugfix: Crash fix on shutdown in Sequence Recorder Module by checking if the UObject system is still active before attempting to unregister a class layout.
+
+* Bugfix: Crash fix with the search box in Blueprint Context Menu not getting focus.
+
+* Bugfix: Assert fix in Metal caused by releasing a shader resource view or unordered access view and the source object within a single command-buffer.
+
+* Bugfix: Crash fix in Metal due to race condition in Metal Blend States’ constructor.
+
+* Bugfix: Crash fix when deallocating textures in Metal RHI by explicitly retaining and releasing the MTL Texture objects rather than making assumptions about whether they are the same object.
+
+* Bugfix: Crash fixes for various GPU crashes when using Metal by explicitly zero-initialising variables and bounds-checking buffer access in Metal shaders.
+
+    * Added a shader compile option "r.Shaders.ZeroInitialise" that turns on explicit zero-initialisation of local and temporary variables in hlslcc for Metal
+
+    * Added a shader compile option "r.Shaders.BoundsChecking" flag to control whether shader platforms should manually enforce buffer access bounds - HLSL returns zero or ignores invalid reads & writes but Metal leaves the behaviour undefined and some drivers then fail.
+
+* Bugfix: Playback of movies on Mac/iOS that had spaces in the filename.
+
+* Bugfix: Fixed a buffer overflow bug in Detour Tile Cache Builder when removing a vertex that could corrupt the heap on Mac.
+
+* Bugfix: Fixed a compilation error found by the newer version of clang in Xcode 8.
+
+* Bugfix: Fixed an issue with Unreal Build Tool and Unreal Automation Tool not finding Unreal Tournament icon on Mac.
+
+* Bugfix: Fixed a problem with windows not updating if the main editor window was minimized. 
+
+* Bugfix: Fixed a race condition when enabling Parallel Execution in Metal by ensuring that ownership of committed command-buffers is properly transferred to the command-queue.
+
+* Bugfix: Fixed a Metal validation error when always binding a velocity uniform data buffer when the shader reads from one.
+
+* Bugfix: Fixed an issue with garbage data appearing in some circumstances using Metal by clearing render targets on first use to approximate the behaviour of other APIs.
+
+* Bugfix: Fixed incorrect rendering in various circumstances on Mac Metal by explicitly initialize all variables to zero, apart from local variable arrays which can't be done currently.
+
+* Bugfix: Fixed intermittent command-buffer failures in Metal by properly handling RHI Set Stream Source calls that override the stream stride.
+
+* Bugfix: Updated Metal implementation to assert on unimplemented calls to RHI Clear UAV on a structured-buffer or texture rather than failing silently.
+
+* Bugfix: Fixed an issue with Metal occlusion queries not reporting command-buffer failures correctly.
+
+* Bugfix: Fixed Metal's internal buffer pool to ensure the CPU and GPU are not accessing a buffer at the same time by only returning a buffer to the pool upon completion of the last command-buffer that uses it.
+
+* Bugfix: Fixed rendering of Atmospheric Fog on Metal without affecting other platforms with a minor change to the shader and by using the precise Metal intrinsics for floating point operations.
+
+* Bugfix: Fixed an issue with Screenshot and movie capture of gameplay on Mac when using Metal to render by inserting an intermediate render-target between the game and the back-buffer.
+
+* Bugfix: Fixed shader compilation for Mac when Apple's 'metal' shader compiler is unavailable by compiling Metal shaders as text instead.
+
+* Bugfix: Fixed the command-line script used to get the number of parallel tasks to use when invoking builds with distcc from Unreal Frontend.
+
+* Bugfix: Fixed the tiled reflections pass on Nvidia GPUs using Metal by removing Max & Min LOD overrides from our calls to "set Sampler State" and “set Sampler States”.
+
+* Flush writable file handles on close on Apple platforms to ensure that they reach the disk prior to program termination.
+
+* Increased the number of input buses for Core Audio 3D Mixer to support 64 audio channels. 
+
+    * Also, added a warning to Start Sources in Audio Device so it doesn't silently ignore sound source initialization failures.
+
+* Reduced memory use in the Mac Editor when running on the Metal graphics API by caching a fallback depth-stencil surface used for canvas tile rendering.
+
+* Reduced the number of temporary memory allocations required to set uniform parameters in Metal.
+
+* Separated texture & buffer references during Metal shader compilation to ensure that buffers don't get assigned an inaccessible index.
+
+* When archiving on for Mac, delete the dest icon if it exists before trying to call File.Move
+
+* Change the Metal buffer pool used inside Metal RHI to better match Apple's internal memory management behaviour to avoid hidden performance problems.
+
+* Disabled "r.DFShadowScatterTileCulling" as well as “r.AOScatterTileCulling” on Mac because we don't have the necessary RW textures on Metal in Mac OS X El Capitan.
+
+* Disabled the Shader Cache by default on Mac to match other platforms as it was not designed to support the new RHI Thread & Parallel Execution features of Unreal Engine 4.
+
+* Disambiguated Metal command-buffer failures by GPU vendor and error type and made failure fatal to avoid cases where attempting to continue causes a kernel panic.
+
+* Enable global clip-plane support for Metal shader platforms. This enables support for Planar Reflection Actors.
+
+#### Playstation 4
+
+* New: Added Visual Studio 2015 support for Windows use of DualShock Pads.
+
+* New: Updated Playstation 4 to SDK 3.508.201
+
+* Bugfix: Fixed a rare GPU hang when async texture reallocs would overlap with End Frame.
+
+* Bugfix: Fixed a compilation failure when Unsafe Command Buffers are enabled.
+
+* Bugfix: Fixed an issue where COTF was not updating cached files.
+
+* Bugfix: Fixed an issue where GS mode was not getting disabled when using parallel contexts.
+
+* Bugfix: Fixed the initial Garlic and Onion memory allocation sizes.
+
+* Bugfix: Fixed an issue where Media Player Audio was cutting out.
+
+* Bugfix: Fixed an issue where Audio was missing when using A3D.
+
+* Bugfix: Fixed null depth target when eliminating fast clears, which fixes a GNM validation failure. 
+
+* Bugfix: Fixed an issue where Num Instances were not getting reset after Draw Indirect calls.
+
+* Bugfix: Fixed packaging of non-code projects when plugins were enabled.
+
+* Bugfix: Fixed render target masks for disabled render target, which fixes a validation check for shaders expecting to write to NULL render targets.
+
+* Bugfix: Resolved issue with Render Target memory allocation.
+
+* Bugfix: Fixed an issue with shader output/render target format mismatch for sparse MRT. 
+
+* Bugfix: Fixed uninitialized Head Position from HMD tracker when tracking fails resolves a problem with A3D audio not working.
+
+* Bugfix: Fixed an issue where two different PS4RHI helper functions had the same name and signature, "handleReserveFailed", which caused unity builds to fail. The names have been changed to address this.
+
+* Smooth Frame Rate now defaults to False.
+
+* Improved checking that correct SDK is installed.
+
+#### Xbox One
+
+* New: Added an example of adding rating information to a title manifest has been added to ShooterGame.
+
+* New: D3D resources can now be created directly on the Render Thread, removing frequent stalls of the RHI thread.
+
+* New: Framerate smoothing is now allowed on Xbox One. Default is set to off.
+
+* New: Oodle network traffic compression is now enabled for Xbox One.
+
+* New: Update Texture 2D can now be done directly on the render thread, removing stalls on the RHI thread.
+
+* New: Xbox One now supports the August 2016 XDK by default.
+
+* Bugfix: Fixed a crash when disconnecting controllers which only occurred if games have been built with Visual Studio 2015 Update 3.
+
+* Bugfix: Eliminated a circular dependency on Xbox One target settings that could cause settings to not load in the editor.
+
+* Bugfix: Fixed ESRAM issues that were causing some odd rendering at times (bloom would be super bright along the bottom of the screen). This fix is also a performance win in some cases.
+
+* Bugfix: Fixed an issue where in some cases a game would not launch if a localized resource value (i.e. Project ID or Description) was not set.
+
+    * If these values are empty, a default value of "<Value Not Set>" will be used so that the project can run. A warning message will also be printed to the log. Make sure to set these values before shipping.
+
+* Bugfix: Fixed an issue where builds could fail if more than one version of the XDK was installed side by side.
+
+* Bugfix: Fixed an issue where composite fonts created before the release of hardware decompress support would not load correctly.
+
+* Bugfix: Fixed an issue where Product Id, Content Id, and Sandbox Id values were being read from the wrong INI when packaging.
+
+* Bugfix: Fixed an issue where streaming Install chunks would attempt to mount before completely downloaded.
+
+    * This issue affected large chunks where the OS reports 100% progress before the chunk completes downloading.
+
+    * Mounting now waits for the "Chunk Completed" event to be fired by the OS.
+
+* Bugfix: Fixed a memory leak of 3D textures.
+
+* Bugfix: Fixed an issue with missing GPU particles when Fast Semantics was enabled.
+
+    * As a result, Fast Semantics is now the default RHI on the Xbox One for 4.13.
+
+* Bugfix: Fixed an issue where the reflection environment was appearing bright green on some maps.
+
+* Moved the Xbox One platform settings into the default INI files so they are available to the PC tools and settable in the editor.
+
+* Removed CPU/GPU sync when updating/unlocking textures.
+
+* Deferred deletion of Xbox One D3D resources moved to happen once per real frame as there were circumstances in which it could be called more than once (rendering cube map faces, etc). As a result the latency of deletion has been reduced to two frames, saving memory.
+
+* Dynamic vertex/index buffers use less memory when locking/unlocking. This is also a performance improvement in some cases.
+
+* Improved processing of system events during load. 
+
+    * This will allow a game to suspend earlier in the loading process.
+
+* Large memory savings when locking/updating textures as we now only allocate a linear copy of the mip/array slice that is being locked, rather than the entire texture.
+
+* Occlusion queries now use 8x less memory than in 4.12 by packing 8 queries into each 256 byte allocation.
+
+* Optimized many calls in the Xbox One D3D11 state cache.
+
+* Shaders now use less memory, as intermediate DXBC bytecode has been stripped away in all possible cases.
+
+* Small stats optimization on Xbox One.
+
+* Vastly optimized the D3D11 Uniform Buffer destructor as we don't need to consider dynamic buffers any more.
+
+* Changed the Streaming Install implementation on Xbox One to not rely on the Find Chunk From File function.
+
+    * This avoids an issue with the Xbox One OS that will sometimes report the incorrect chunk index.
+
+#### VR
+
+* The Morpheus HMD constructor now initializes the Player Position vector.
+
+#### HTML5
+
+* New: Call End Session on Before Unload.
+
+* New: Added the HarfBuzz libraries for HTML5.
+
+* New: Added Third Party Software documentation and license file for Emscripten.
+
+* Bugfix: Corrected URL pathname generator to find pak files with the "complete" pathname.
+
+* Bugfix: Readme file cleanup.
+
+* Render and performance optimizations.
+
+#### Windows
+
+* New: The UCRT and Visual Studio 2015 runtime can now be setup for app-local deployment.
+
+    * A configuration option called App Local Prerequisites Directory was added to the packaging settings in the editor.
+
+    * Files stored in this directory will be deployed side by side with the executable.
+
+    * DLLs staged this way will be available without requiring the runtimes to be conventionally installed.
+
+    * Redistributable versions of the UCRT and Visual Studio 2015 runtime files are now included with the engine.
+
+    * The runtimes have been structured to be compatible with app-local deployment in the directory Engine\Binaries\ThirdParty\AppLocalDependencies.
+
+* Bugfix: Changing the desktop resolution while a full screen title is minimized will cause the resolution of the title to change to match the new desktop resolution.
+
+#### All Mobile
+
+* New: Added ability to consume IAP products during Restore Purchases which weren't consumed when they were purchased.
+
+* New: Added additional error code to represent an in-app product already being owned prior to the current action.
+
+* New: Added Raw Price to In-App Purchase info returned for IAP queries.
+
+* New: Added support for the Web Browser widget Execute Javascript function on iOS and Android.
+
+* New: Added <insertNewline/> command to Unreal Plugin Language.
+
+* New: Added fast ASTC texture compression, using ISPC.
+
+* Bugfix: Fixed Custom Touch Interface objects to now properly fade to their Inactive Opacity.
+
+* Bugfix: Fixed the virtual joystick to return -1 to 1 ranges for thumbsticks (diagonals aren't shortened).
+
+* Reduced the memory overhead of a cached file handle on mobile to save memory in some loading situations.
+
+### Programming
+
+* New: Improved pivot painter 2.0's internal documentation through better warning strings.
+
+* New: Minor optimization in the FProjectedShadowInfo::AddSubjectPrimitive() function.
+
+* New: Optimized the FPooledRenderTarget::IsFree() function. Now ~50% quicker.
+
+* New: Optimized the FSplineMeshVertexFactoryShaderParameters::SetMesh() function.
+
+* New: Vastly optimized GameplayTagNodeMap by storing an inverse map, and avoiding O(n) traversal of the map.
+
+* Bugfix: Fixed Ctrl-C not killing child processes when running AutomationTool from the console.
+
+* The value of the P4PORT variable is now found by querying Perforce, rather than just checking the environment variable. P4 on Windows stores these values in the registry.
+
+* Changed Shipping configuration of UE4Game to use Shipping PhysX libraries on Windows.
+
+* FRHITexture member TextureName should have a non empty name for most texture types (for easier debugging).
+
+### Rendering
+
+* New: Added component dependencies for each textures in the "ListStreamingTextures" console command.
+
+* New: Added a stat ("Translucency GPU Time (MS)") for the total GPU time spent on translucency rendering.
+
+* New: Added 'Allow CPU Access' flag to a Static Mesh, allowing CPU access to geometry in cooked builds.
+
+    * This is required to convert a Static Mesh to a Procedural Mesh Component in-game.
+
+* New: Added the GPU Shared Flag to render targets so they can be shared with other applications.
+
+* New: Added Blueprint node (Export Render Target and Export Texture 2D) to export 2D render targets and textures as HDR images to the disk.
+
+* New: Added Composite Mode to Scene Capture 2D, which can be used to additively accumulate or composite instead of the default overwrite behavior.
+
+    * Added "Capture On Movement" to Scene Capture, which can be disabled so the only source of scene capturing is a manual capture by calling Capture Scene.
+
+* New: Added the Console command "r.DumpShaderDebugWorkerCommandLine" to generate a file when dumping shaders to be able to debug using the Shader Compiler Worker's -directcompile.
+
+* New: Added support to be able to set a Canvas Render Target 2D NOT to clear to green whenever it's updated (so you can just draw pixels that have changed instead of always having to redraw everything).
+
+* New: Capsule shadows are now supported in stereo (VR).
+
+* New: Create Project Settings to optionally disable certain rendering features which cause extra shader permutations and extra shader compile time if they are not required for a given project.
+
+    * Can be found under Project Settings > Rendering Overrides (Local). 
+
+    * If an attempt to use a disabled feature is made, a warning will be displayed.
+
+* New: For use when doing automated screenshot verification, the new console command "r.ResetViewState" will reset some state (e.g. Temporal AA index). This allows for more deterministic tests.
+
+* New: Added a new Blueprint function "Clear Render Target 2D", which is the only way to set a render target alpha directly.
+
+* New: Added a new Blueprint function "Create Render Target 2D", which allows self-contained render-to-texture blueprints
+
+* New: OpenGL4 now supports 32 combined samplers.
+
+* New: Stationary Mobility for primitive components is allowed again, with the meaning 'moves rarely'. 
+
+    * Used for caching optimizations (cached shadowmaps).
+
+* Bugfix: Fixed a crash when removing instances while lighting is building.
+
+* Bugfix: Fixed a crash caused by quitting OpenGL 4. 
+
+* Bugfix: Fixed a crash when applying materials to some BSP models. 
+
+* Bugfix: Fixed a crash where Canvas Render Target 2D assets would crash on load in the editor due to trying to update the asset during post-load.
+
+* Bugfix: Fixed a crash while in Collision Rendering mode. 
+
+* Bugfix: Fixed a crash where the first time a Vulkan app is compiling shaders it was running out of descriptor sets.
+
+* Bugfix: Fixed a crash when selecting LODs with missing render data.
+
+* Bugfix: Fixed a crash when setting a composite texture on a render target cube. 
+
+* Bugfix: Fixed an issue where texture asset import data was being ignored when async loaded.
+
+* Bugfix: Fixed an issue where Billboards could be translucent during PIE.
+
+* Bugfix: Fixed a possible lighting build warning when using foliage. 
+
+* Bugfix: Fixed a regression with the material node "Camera Vector WS" that could affect post-process blendables. 
+
+* Bugfix: Fixed a resource leak with decal rendering in D3D11 RHI.
+
+* Bugfix: Fixed an issue where bounds were too large due to Scene Component sprites in Editor Builds. 
+
+* Bugfix: Fixed issues with Scene Capture components.
+
+    * Scene cube captures were not updating, leaving the render target unchanged.
+
+    * Scene capture components set to "Capture Every Frame" were only capturing every second frame.
+
+* Bugfix: Fixed an issue where particles using Spherical Particle Opacity wouldn’t render. 
+
+    * This now allows a smoother appearance when particles are intersecting geometry.
+
+* Bugfix: Fixed an issue with normals when using the GPU skin cache. 
+
+* Bugfix: Fixed an issue where refraction with a world space normal would behave differently from tangent space normals.
+
+* Bugfix: Fixed an issue with Shader Complexity viewmode when using Resolution Scaling. 
+
+* Bugfix: Fixed an issue where simulation times from Blueprints were not matching the Material editor expression time (eg. during PIE/Simulate).
+
+* Bugfix: Fixed an issue where SSAO was being applied in the Unlit viewmode. 
+
+* Bugfix: Fixed an issue where Static Mesh bounds were including simple collision bounds.
+
+* Bugfix: Fixed an issue with `Refraction Depth Bias` material parameter. 
+
+* Bugfix: Fixed an issue where World-space camera position was not being exposed in decal pixel shaders. 
+
+    * This also fixed decal lighting having missing specular and reflection.
+
+* Bugfix: Fixed an issue where translucency was appearing in Stationary Light Overlap viewmode. 
+
+* Bugfix: Fixed a Mesh Merge issue with merging mix of different extra bone influence with multiple sections.
+
+* Refactoring of resident mip calculations:
+
+    * Cooker takes into account the same compression block thresholds that the runtime previously used.
+
+    * Runtime doesn't attempt to calculate which mips to perma-load, but just looks at the ones whose bulk data is flagged as end-of-file or separate-file.
+
+* Removed "Update Actor Position".
+
+    * This was not needed in a vast majority of use cases and was causing a crash due to multithreading issues during end of frame updates.
+
+* Removed vertex density viewmode.
+
+* Shader compilation errors will now print the details of the error to the log by default.
+
+* Shader compiler does a recreate render state even during blocking compile.
+
+    * This fixes saving a material giving different behavior from applying changes with global distance fields.
+
+* Improved usability for DBuffer Decals.
+
+    * 'Show Decals' works correctly, previously would fetch from uninitialized textures.
+
+    * DBuffer being enabled forces a full prepass, previously decals would render incorrectly unless correct settings of "r.EarlyZPass" were used.
+
+    * Materials using DBuffer blend modes will fail to compile when the DBuffer project setting is disabled, instead of just being invisible.
+
+* Minor optimization of the Translucency lighting volume clear by moving it to Async Compute.
+
+* Noise material node functions have been renamed to more accurately represent what they do, and all now have tooltips with performance and quality information.
+
+* Renamed some cvars for consistency: 
+
+    * "r.SkinCacheShaders" -> “r.SkinCache.CompileShaders”
+
+    * "r.SkinCaching" -> “r.SkinCache.Mode”
+
+    * "r.MaxGPUSkinCacheElementsPerFrame" -> “r.SkinCache.MaxGPUElementsPerFrame”
+
+* Shader compile error logs are now unsuppressed - shows the cause for materials that fail to compile in cooks.
+
+#### FX
+
+* New: Added "Opacity Source Mode" to Sub UV Animation, which is useful with textures created for additive particles.
+
+* New: Added Particle Light Quality scalability setting since lower-end machines have trouble with particle lights. 
+
+    * They are disabled on low and medium spec machines. HQ lights are only allowed on high-end machines.
+
+* New: Beam emitters can now use particles as their source and target methods.
+
+* New: GPU particle simulation can be separately frozen/unfrozen with "r.GPUParticle.Simulate".
+
+* New: Particle emitters using vector fields can now be configured to use fix DT or variable DT for simulation.
+
+    * Other emitters only use variable DT.
+
+    * Fix DT allows for a more consistent particle behavior when running at different game speed (like slomo) and at different frame rates.
+
+* New: The transforms of individual mesh particles are now available in materials via the "Transform Vector" and “Transform Position” nodes for advanced material effects on mesh particles.
+
+* New: Quality level spawn rate scaling now works on GPU emitters.
+
+* Bugfix: Fixed a crash with a Sub UV Animation asset when texture source is cleared.
+
+* Bugfix: Fixed a crash that could happen on mesh import if the mesh is used on a currently active particle emitter. 
+
+* Bugfix: Fixed a crash that could happen when changing the particle module distribution properties. 
+
+* Bugfix: Fixed a crash caused by a recent feature addition to the Bone Socket module. 
+
+* Bugfix: Fixed a crash when activating particle systems with instanced parameters in packaged games. 
+
+* Bugfix: Fixed a crash caused by several issues when setting large values on parameters of the Beam Noise module. 
+
+* Bugfix: Updated particles so they can now use "Use Single Sample Shadow From Stationary Lights".
+
+* Bugfix: Fixed an issue with distortion in particle textures when using the macro UV module.
+
+* Bugfix: Fixed an issue with visible pop of particles on collision.
+
+* Bugfix: Fixed an issue where some FX would be spawned at the wrong LOD level when used with the pooling system.
+
+* Bugfix: Fixed an issue where particle cutouts were not working with certain Sub UV Animation texture sizes and Sub UV dimension combinations.
+
+* Bugfix: Fixed several issues with Particle Ribbons using sequential particle selection.
+
+* Bugfix: Fixed an issue preventing inherit rotation from working with mesh particles.
+
+* Increased size of texture for curves on GPU particles.
+
+    * This prevents bugs when many GPU particle systems using curves are active.
+
+* Particle Collisions can be ignored for systems not currently in view.
+
+    * Added a console variable to control how long these systems must be inactive for to have their collisions ignored.
+
+    * "fx.ParticleCollisionIgnoreInvisibleTime"
+
+* Reduced tunneling/missed collisions with depth buffer collision.
+
+#### Lighting
+
+* New: Added "r.Shadow.ForceSingleSampleShadowingFromStationary" for scalability.
+
+* New: Added a per-object shadow setting for directional lights
+
+    * "r.Shadow.PerObjectDirectionalDepthBias" - which allows us to set a depth bias for per-object shadows independently of the CSM shadows.
+
+* New: Added cvar "r.Fog" for scalability.
+
+* New: Added "Lower Hemisphere Solid Color" to sky lights.
+
+* New: Added "r.AOSpecularOcclusionMode", which determines how specular should be occluded by DFAO.
+
+    * 0: Apply non-directional AO to specular.
+
+    * 1: (default) Intersect the reflection cone with the unoccluded cone produced by DFAO. This gives more accurate occlusion than 0, but can bring out DFAO sampling artifacts.
+
+* New: Fading out planar reflections based on roughness since they don't have support for variable roughness (fade starts at .2 roughness, ends at .3)
+
+* New: Lights now have their own "Max Draw Distance" setting, along with a fade range to smoothly fade the light out. 
+
+    * Brightness will approach 0 at max draw distance and fade to 100% within the specified fade range.
+
+* New: Planar Reflection Improvements
+
+    * Added "Render Scene Two Sided" to planar reflections, which can be useful to limit leaking.
+
+    * Added "Show Only Actors" and “Hidden Actors” to Scene Capture Component for easy use without having to call BP functions.
+
+    * Added "Show Preview Plane" to Planar Reflection actors
+
+    * The view state is recreated on planar reflection edit, which resets the Temporal AA history, allowing instant previewing of changes
+
+* New: Scaled sphere intersection for indirect capsule shadows.
+
+    * Fixes the discontinuity when capsule axis points close to the light direction.
+
+* New: Skylight occlusion tint now applies to specular.
+
+    * Skylight occlusion tint on diffuse is now correctly affected by Diffuse Color.
+
+* New: Spreading precomputed visibility to neighbors now uses a 2D grid to find neighbors.
+
+    * This speeds up the process for 800k cells from 40 mins to 20s.
+
+* New: Volume lighting samples use adaptive sampling final gather.
+
+    * Increases their build time by 2x but improves quality in difficult cases (eg. small bright sources of bounce lighting).
+
+* Bugfix: Fixed a crash with threading related to level load with precomputed lighting volumes. 
+
+* Bugfix: Fixed Capsule Shadows artifacting on Skeletal Meshes with scaling.
+
+* Bugfix: Fixed an issue with HLOD and mesh LODs getting hit by Lightmass ray traces that didn't originate from a mesh.
+
+    * Volume lighting samples and precomputed visibility cells are now only placed on LOD0 (of both mesh LODs and HLOD).
+
+* Bugfix: Fixed an issue so that Lightmap density viewmode will now show the expected component's lightmap resolution when the precomputed lighting is not yet built.
+
+* Bugfix: Fixed an issue where Actors hidden in game still cast shadows and emit light in Lightmass, which was broken in 4.12.
+
+* Bugfix: Fixed an issue where atmospheric fog was not working on translucency.
+
+* Bugfix: Fixed issues with hair shading NaNs.
+
+* Bugfix: Fixed an issue where light shapes were not showing up in reflection captures.
+
+* Bugfix: Fixed an issue with NaNs resulting from 0 Roughness materials if the light’s "Min Roughness" was also set to 0.
+
+* Bugfix: Fixed an issue with Single Sample Shadow From Stationary Lights which had been broken from selective GBuffer outputs optimization.
+
+* Bugfix: Fixed a typo in tube lights that should no longer allow them to go negative.
+
+* Bugfix: Fixed per object shadows from point/spot lights to do the light source pull back based on subject box size, not subject radius, since the box is used to find a valid < 90 degree projection.
+
+* Bugfix: Fixed skylights so that they are not captured until their level is made visible.
+
+    * This fixes the case where skylights capture too early.
+
+* Bugfix: Fixed height fog to now work properly in planar reflections.
+
+* Bugfix: Fixed HLODs to now work with static and dynamic shadowing.
+
+* Lighting channels can now be edited on components with static mobility, since dynamic lights can still affect them.
+
+* Improved hair shading when lit by non shadow casting light.
+
+* Improved screen space reflection quality.
+
+#### Materials
+
+* New: Added a purely computational Perlin-style noise to the Noise material node. 
+
+    * Unlike texture-based noise, this one will never repeat.
+
+* New: Added material property "Allow Negative Emissive Color” which is useful when rendering to an offscreen render target.
+
+* New: New material functions available. These were incorrectly listed in the previous release but are available now.
+
+    * 1D to 2D Index - Convert given input from 1d to 2d.
+
+    * 1D to 3D Index - Convert given input from 1d to 3d.
+
+    * 2D to 1D Index - Convert given input from 2d to 1d.
+
+    * 3D to 1D Index - Convert given input from 3d to 1d.
+
+    * Contrast Preserve Color - Allows you to adjust the contrast of an image without adjusting the colors.
+
+    * Luminosity and Color - Allows you to separately specify the Luminosity and Color to create an output image.
+
+    * Landscape Manual UVW - A helper function to help map textures and height gradients to landscapes.
+
+    * Chroma Key Alpha - This is a green screen type node that extracts a mask based on the specified chroma color.
+
+    * Distance Limited Reflections - This returns parallax corrected reflections conformed to a sphere.
+
+    * Brick and Tile UVs - A useful function for creating pseudo random UVs for procedural texture creation.
+
+    * Texture Bombing - A material function that multi-samples a texture with random offsets to reduce tiling artifacts.
+
+    * Texture Bombing POM - An experimental function combining texture bombing with Parallax Occlusion Mapping. Note this function is extremely expensive.
+
+* New: Added new console variables "r.Streaming.FullyLoadUsedTextures 1" to load all used textures to their max resolution.
+
+    * Used as an alternative to disabling texture streaming with "r.TextureStreaming 0".
+
+    * Takes less memory as only used texture get streamed in.
+
+* New: Panner nodes now have an optional speed input.
+
+* New: Support for setting Scalar and Vector Materials by Index rather than by name on MIDs. 
+
+    * This feature allows much better performance in cases where large numbers of parameters are being set per frame and where the indices can be cached by the calling code in an initialization step.
+
+* New: Textures have an Alpha Preservation property which dictates how much of a channel's coverage to preserve when generating mipmaps.
+
+* Bugfix: Fixed a crash when passing a Vector4 to a "Transform Position" material node. 
+
+* Bugfix: Fixed an issue where Texture group and texture LOD Bias setting would be applied twice in cooked builds.
+
+* Bugfix: Fixed an issue with texture streaming where Skeletal Meshes would understream textures. 
+
+* Bugfix: Fixed an edge-case causing negative blend results on tri-planar mapping material function.
+
+* Bugfix: Fixed an issue where "Invalid Lightmap Settings" display material for instanced meshes.
+
+* Bugfix: Fixed the Shader Complexity viewmode for materials using depth pixel offset.
+
+    * This could cause masked materials to appear cheaper than they actually are.
+
+* Bugfix: Fixed a typo on `If` material node inputs.
+
+* Bugfix: Fixed various material compile errors of the form "Undeclared identifier LocalX" caused by interpolation optimization.
+
+* Removed additional material node created when adding a "Rotate About Axis" node.
+
+#### Postprocessing
+
+* Bugfix: Fixed a regression (seen only in SM4) that could break the fixed shutter exposure.
+
+* Bugfix: Fixed an issue with ghosting in the screen space reflections when temporal AA is disabled.
+
+* Bugfix: Fixed an issue with occasional flickering visual artifacts in the motion blur.
+
+* Bugfix: Fixed post-processing not working correctly in case project Mobile HDR setting is off, on non-mobile platforms.
+
+#### Optimizations
+
+* New: Added a 'resource' name for landscape to improve display in the mesh summary list of Profile GPU.
+
+* New: Added a way to control how far down the HLOD hierarchy to allow showing (can be used to limit quality loss and streaming texture memory usage on high scalability settings).
+
+    * Controlled by the new cvar "r.HLOD.MaximumLevel", which can be set to the following values:
+
+        * -1: No maximum level (default)
+
+        * 0: Prevent ever showing a HLOD cluster instead of individual meshes
+
+        * 1: Allow only the first level of HLOD clusters to be shown
+
+        * 2+: Allow up to the Nth level of HLOD clusters to be shown
+
+    * Note: This does not affect the memory used by the HLOD meshes itself or their always loaded low mip levels, it will only save the memory associated with streaming in the higher mip levels.
+
+* New: Added triangle and draw call summaries to Profile GPU output, broken up by asset and material.
+
+    * This is controlled by "r.ProfileGPU.PrintAssetSummary", which defaults to 1, but you really need “r.ShowMaterialDrawEvents 1” enabled as well for a complete picture.
+
+    * It can also output a summary line for specific asset names using a comma separated list in "r.ProfileGPU.AssetSummaryCallOuts" (e.g., "LOD,HeroName").
+
+* New: Moved Texture streaming computations for dynamic primitives on the worker thread.
+
+* New: Realtime GPU profiler - preview feature (not yet enabled by default).
+
+    * A non-hierarchical, realtime, high level GPU profiler with support for cumulative stat recording. Requires a build with stats enabled (e.g. development).
+
+    * To enable, use "r.gpustatsenabled 1".
+
+    * Stats are visualized with "stat GPU".
+
+    * Additional stat can be added with SCOPED_GPU_STAT macros, e.g. SCOPED_GPU_STAT(RHICmdList, Stat_GPU_Distortion).
+
+* New: Reduced memory allocations when creating sections in a Procedural Mesh Component.
+
+* New: Simple forward shading path for scaling down to low end PC.
+
+    * Support for a directional light in the base pass, which can be shadowed either by distance field shadowmaps or a single shadow factor from the Indirect Lighting Cache.
+
+    * Binding and outputting to GBuffers is disabled when using simple forward shading, along with dependencies of the GBuffers.
+
+    * Specular, normal maps and the 3rd SH band of skylighting are all hardcoded off. Materials are forced fully rough.
+
+    * Enabled with "r.SimpleForwardShading 1", which can be changed at runtime.
+
+    * Added "r.SupportSimpleForwardShading", which determines if the necessary shaders are compiled.
+
+    * Removed "r.SimpleDynamicLighting", whose functionality (unshadowed directional light in the base pass) is a subset of simple forward shading, and was unused.
+
+* Added "FreezeRendering" to the console autocomplete list.
+
+* Added "ProfileGPU" to the console autocomplete list.
+
+* Corrected a misleading help comment on "r.StaticMeshLODDistanceScale", explaining how it affects LOD calculations (it multiplies the effective distance, so larger numbers make transitions happen sooner).
+
+* Updated to d3dcompiler_47.dll offering large speed-up in shader compilation times.
+
+#### Mobile Rendering
+
+* New: Added support for lighting channels on mobile.
+
+    * Multiple directional lights are supported in different channels but primitives are only affected by the directional light in the first channel they have set.
+
+    * CSM shadows from stationary or movable directional lights correctly follow their lighting channels.
+
+    * No channel limitations for dynamic point lights.
+
+* New: Enabled quality level overrides for iOS Metal, Android ES3.1 and Android Vulkan shader platforms.
+
+* New: Renamed "ANDROIDES31" and “OpenGLES31” to “ANDROIDESDEFERRED” and “OpenGLESDeferred” to avoid confusion with mobile renderer supporting ES3.1.
+
+* Bugfix: Fixed an issue in which changing mobile preview device type in editor would not rebuild materials.
+
+* Bugfix: Fixed an issue that could cause changes to quality overrides in the 'High' quality column to be ignored.
+
+* Bugfix: Fixed "black dots" rendering corruption that may appear on particles with additive blending.
+
+* Bugfix: Fixed decals not rendering on Zenfone 2 (Added proper detection of FP16 render target support).
+
+* Bugfix: Fixed inverted R/B color channels with "r.TonemapperFilm" enabled on Nexus 5 (Android 4.4.4).
+
+* Bugfix: Fixed performance regression introduced in 4.10. Compute pixel world position in vertex shader instead of doing it in pixel shader.
+
+* Mobile DoF now requires DoF type to be Gaussian otherwise DoF will be disabled.
+
+### UI
+
+* Bugfix: Fixed a regression that broke safe zone debugging features in Windows.
+
+* Bugfix: Fixed an issue in a Slate List Panel so that those with horizontal arrangement (i.e. from Slate Tile View) use the number of desired items instead of the number of actual items in order to calculate the desired size of the geometry.
+
+    * This fixes the case where a Slate Tile View is contained within a Slate Scroll Box; the Slate Tile View is now allowed to expand to its full size within the Slate Scroll Box.
+
+* All textures in the UI texture group (TextureGroup_UI), should now render a lot smoother if they have mipmaps. Removed code that was forcing all UI textures to mip level 0. If this behavior was desired by anyone, you can use the cvar "r.ForceHighestMipOnUITextures" to re-enable it.
+
+* Sound Concurrency asset now has a unique icon.
+
+#### Slate
+
+* New: Added input settings:
+
+    * Control if the viewport locks the mouse on acquire capture.
+
+    * Control if the viewport acquires capture on the application launch (first window activate).
+
+* New: Slate Remote Server (for the iOS touch input app) is now disabled by default, so we do not open up the socket unless desired by the user.
+
+* Bugfix: Fixed a crash if a Slate widget tried to refresh itself while the editor is shutting down.
+
+* Bugfix: Fixed a crash that could happen when the fallback/last resort font was missing.
+
+* Bugfix: Fixed a potential crash when shutting down Slate if the font services were shut down first.
+
+* Bugfix: Fixed regression in editable text box alignment.
+
+    * Text was no longer vertically aligned center since editable text was converted to use a text layout. This vertical alignment is now handled by the outer editable text box instead.
+
+* Bugfix: Fixed the Scroll Box widget improperly clamping by content size.
+
+* Bugfix: Fixed some measuring issues with bi-directional text within a right-flowed document.
+
+    * This could cause text to sometimes significantly under-measure and lead to wrapping issues.
+
+* Significant performance improvements when pasting a large amount of text.
+
+* Made sure the slide handle of a Slider Widget no longer gets cut off at the edges of the slider.
+
+* Slider widget will no longer consume OnKeyDown events that they do not handle.
+
+* Updated ScrollBy function on ScrollBox widget to work properly when called while the widget is animating.
+
+#### UMG
+
+* New: Added ability to replace the widget the track is currently bound to. Also includes changes in Widget Blueprint Editor to send delegate to Animation Tab Summoner when switching from Graph to Designer.
+
+* New: Adding ability to animate the root (owner) widget in UMG.
+
+    * Hierarchy Root adds the preview widget instead of the CDO to selected objects in widget blueprint. The properties are then migrated back to the CDO.
+
+* New: Added additional styling options for the default items that are created in the combobox, such as Font.
+
+* New: Updated playback speed so that it can now be specified when starting a UMG animation.
+
+* New: Updated the Palette in UMG to no longer display widget content coming from Engine or Developer folders, based on the same setting configured in the content browser. If you're not showing Engine content in the Content Browser, you should no longer see all the User Widgets that were added to support the VR Editor.
+
+* New: Updated the Widget Component to extend from the Mesh Component. It can now have a custom material applied to it in order to achieve cooler effects - like ignoring the depth buffer.
+
+    * Users who use this option are encouraged to start with the widget components default material and work from there, as you need to start by sampling a texture parameter called Slate UI.
+
+    * The widget component now offers the ability to automatically size the render target to be the desired size of the widget - note that this could go really bad if your widget's desired size is very large.
+
+* New: Added a new Slate Texture Atlas Interface that any UObject can now implement if it wishes to integrate with UMG to provide texture atlas data in a form Slate can understand and render.
+
+    * For reference you can look at "Paper Sprite" which implements this interface allowing Paper 2D Sprites to be used on Slate Brushes.
+
+* New: Added "Get Url" method and “On Url Changed” property to the Web Browser widget.
+
+* New: Added meta flags for "Determines Output Type" and “Dynamic Output Param” to “Get All Widgets Of Class” in the Widget Blueprint Library, removing the need for the extra cast. 
+
+    * The output pin now becomes the same as the type requested.
+
+* New: Worked on better support for VR interactions with Slate widgets by introducing the Widget Interaction Component and performance improvements to the Widget Component.
+
+    * Direct mouse interaction with World Space widgets in the virtual world is no longer supported. Those kinds of interactions must all use the Widget Interaction Component now.
+
+    * Performance improvements - the viewport no longer arranges all 3D widgets every frame.
+
+    * Widget Components now support several methods for reducing how often they redraw to help control performance, they also support manual refresh. This automatically works in tandem with the widget interaction component to request refresh whenever the widget interaction component is interacting with the widget, thus giving you a simple way to only redraw widgets that the user is hovering over them.
+
+* Bugfix: Fixed a crash when using "Use Selected asset from Content Browser" for font materials when the selected asset is not a material.
+
+* Bugfix: Fixed animation playback not playing the full length if it starts on a negative time.
+
+* Bugfix: Fixed a bug where magenta-color text added to a Widget Blueprint would turn gray when the blueprint was compiled.
+
+* Bugfix: Fixed Widget Component to render widget in the correct position when Constrain Aspect Ratio is enabled.
+
+* Bugfix: Fixed a bug where audio track from UMG doesn't play correctly in PIE.
+
+* Bugfix: Fixed "Is Hovered" on “User Widgets”. The Object Widget did not properly call the super for Mouse Enter/Mouse Leave preventing the Widget's implementation from setting the correct state.
+
+* Bugfix: Fixed the "Designer Rebuild" flag detection so that we can just refresh the slate widget without recreating the preview UObject, which causes the destruction of the details panel, and the slate widget recreation was the only part that was required.
+
+* Bugfix: Fixed an issue where the web browser widget plugin wasn't loading soon enough to be properly loaded in time if it was referenced by game content that loads in the Default stage of the pipeline, so moving it to PreDefault.
+
+* Set viewed animation to current animation after switching from Graph to Designer (This is for "No Animation Selected" showing up when switching).
+
+* Removed the "Designer Widget Tree", instead directly injecting the widget tree into the partially constructed User Widget during design time, when refreshing the preview.
+
+    * This avoids updating the living class instance with a new designer tree that all new instances of the class would become biased to using.
+
+    * Making the preview widget explicitly non-transactional as there's no reason to track changes to the preview, all the changes that need to be tracked should be on the template widget. This should fix the crash in the widget designer when you undo just after compiling the widget blueprint.
+
+* Changed the Input Key Selector Widget to no longer add a bogus Selected Key property to the details panel.
+
+* Updated the Slider widget to now allow the thumbsticks to also control movement of the slider when it's in edit mode, via controller input.
+
+* Widgets transferred from one Widget Blueprint to another Widget Blueprint will now allow both to render and compile correctly.
+
+### VR
+
+* New: Added a "Start in VR" project setting, which will force a project to try and start in VR mode, regardless of whether or not it has -vr specified on the commandline.
+
+* New: Added additional Motion Controller types, to allow for extension to devices that have more than just left and right trackers.
+
+* New: Added GoogleVR (Cardboard) support for iOS, updated SDK.
+
+* New: Added in "Get HMD Device Name" function to the HMD function library, which allows you to identify which plugin (e.g. Oculus Rift, SteamVR, PSVR) is currently active.
+
+* New: Added Oculus 1.6 runtime support.
+
+* New: Added fix for flickering shadows on translucent surfaces in VR.
+
+* New: Updated haptics system so that it can now take a buffer of values to playback asynchronously on devices.
+
+* New: Added PlayStation VR '2DVR' reprojection mode.
+
+    * 2DVR is a way to display a static image on a quad temporarily while in VR output mode. Intended for loading screens, etc.
+
+    * Exposed to blueprint through Morpheus Function Library.
+
+* New: Implemented HMD Setup dialogs for PlayStation VR at startup and if the HMD is disconnected at runtime.
+
+* New: Exposed PlayStation VR "Hmd Reprojection Set Output Min Color" function to blueprint.
+
+    * Sets the minimum output color from reprojection processing. Can be used to mitigate certain vr artifacts (including dark smearing and rendered area edge visibility).
+
+    * Note: Tone Mapping and Color Correction can achieve some of the same effects, but are cross-platform.
+
+* New: Added Steam VR Stereo Layer Support.
+
+* Bugfix: Fixed crash when running with an HMD with no microphone.
+
+* Bugfix: Fixed a crash with the Steam Controller plugin enabled where applying a dynamic force feedback to the controller could crash the editor.
+
+* Bugfix: Fixed crash when polling for position and orientation of an HMD from animation skel controls.
+
+* Bugfix: Fixed crash with Gear VR when transitioning levels with an active HUD.
+
+* Bugfix: Fixed bloom artifacts when changing screen percentage in VR.
+
+* Bugfix: Fixed module loading race condition in Steam VR plugin.
+
+* Bugfix: Fixed assert on Google Daydream soon after startup (likely related to heat-induced throttling).
+
+* Bugfix: Fixed Motion Controllers on remote clients using data from attached motion controllers. Can now replicate motion controllers across the network without having local devices interfere.
+
+* Bugfix: Fixed outstanding issues with OSVR rendering performance and stereo convergence in the headset.
+
+* Prevented the output log from being spammed for every single shader compile when instanced stereo is enabled for a shader platform that doesn't support it.
+
+* Updated Steam VR detection code to now use a lighter weight function, which will not launch Steam VR unless it is chosen as the active platform.
+
+## Programming Release Notes
+
+#### AI
+
+* New: A generic, template-based structure for holding data in a form of 2D grids has been added (TSimpleCellGrid). Useful for things like influence maps and spatial lookup tables.
+
+* New: A generic, template-based A* implementation has been added.
+
+##### Behavior Tree
+
+* Bugfix: Fixed a bug with instanced behavior tree nodes writing over memory of other nodes.
+
+* Bugfix: Fixed missing memory callbacks for instanced behavior tree nodes.
+
+* Bugfix: Fixed multiple calls to behavior tree's instance cleanup.
+
+* Changed behavior tree node injection to modify shared template instead of forcefully switching nodes to instanced. Fixes loops in garbage collector references between AI using the same behavior tree.
+
+##### Navigation
+
+* New: Added a static UNavigationSystem::bStaticRuntimeNavigation flag that can be used to short-circuit dynamic navigation-related functions.
+
+* New: Added meta navigation paths handled by FMetaNavMeshPath class. They are using an array of points to define the shape of a high-level path and automatically expand traversed section into regular navmesh path during path following.
+
+* New: Pathfollowing Component now stores information that current move segment is a navigation link.
+
+* Bugfix: Fixed AITask_MoveTo being responsible for some observed navigation paths not getting released.
+
+* Bugfix: Fixed hash function for FRecastDebugPathfindingNode.
+
+* Exposed UNavigationInvokerComponent as part of ENGINE_API so that it can be spawned procedurally in C++ in game-specific code.
+
+#### Animation
+
+* New: Added check to PostUpdateAnimation to highlight when game code causes it to be called recursively.
+
+* New: Added option to FSkeletalMeshMerge to keep vertex buffer on CPU (e.g. for spawning particle effects).
+
+* New: Added optional UV transforms to FSkeletalMeshMerge, so UVs of each part can be translated and rotated during merging.
+
+* New: Added several new extensibility points to animation and physics for plugin development.
+
+    * Exported FAnimBlueprintCompiler class from KismetCompiler module.
+
+    * Exported AssetUserData interface virtuals on UAnimationAsset.
+
+    * Exported several virtuals on UAnimSequenceBase and USkeleton.
+
+    * Added SimEventCallbackFactory for custom simulation event callbacks.
+
+    * Added CreatePhysicsDelegate and DestroyPhysicsDelegate multi-cast delegates to UActorComponent.
+
+* New: Exposed bone selection widgets so that other modules can reference/use them.
+
+* New: Exposed FFIRFilterTimeBased from Engine module and add SetWindowDuration.
+
+    * Added CopyBonesTo function to FBaseCompactPose.
+
+* New: Improved performance of RecalcRequiredBones on SkeletalMeshComponents. It now takes a third of the time to execute on average.
+
+* New: Merged FSkelMeshChunk struct into FSkelMeshSection to make Skeletal Mesh code cleaner and easier to understand.
+
+* New: Refactored bone reference widget so that the selection tree menu can be used separately to the combo box.
+
+* New: Removed rigid vertex array in USkeletalMesh to make code cleaner and easier to understand.
+
+* New: Removed several animation related headers from global inclusion to reduce compile times.
+
+* New: Removed SkeletalMesh.h from being globally included to reduce compile times.
+
+* New: Renamed 'LocalAtoms' to 'BoneSpaceTransforms' and 'SpaceBases' to 'ComponentSpaceTransforms' to make it clearer what they are.
+
+* Bugfix: FAnimInstanceProxy::EvaulateAnimation is now split into two for easier extensibility.
+
+* Removed UPROPERTY from transient properties in SkinnedMeshComponent (PredictedLODLevel, OldPredictedLODLevel & MaxDistanceFactor).
+
+#### Audio
+
+* New: Allowed AudioVolume to be subclassed outside Engine module.
+
+* New: Threaded Audio.
+
+* Removed USoundWave::CompressionName.
+
+* Added USoundWave::HasCompressedFormat.
+
+#### Automation
+
+* Added map name to editortest command line in UAT.
+
+#### Blueprints
+
+* New: UEdGraphSchema_K2::GetVariableTypeTree now uses an enum filter rather than multiple boolean parameters.
+
+* Bugfix: Crash fix in FKismetDebugUtilities::GetWatchText()
+
+* Bugfix: Fixed a compiler warning in FBlueprintVarActionDetails::GetVariableReplicationType where a local variable was hiding a member variable.
+
+* Deprecated FName UEdGraphNode::GetPaletteIcon(FLinearColor&) in favor of FSlateIcon UEdGraphNode::GetIconAndTint(FLinearColor&) to allow for icons in external style sets to be used.
+
+* UInheritableComponentHandler::GetAllTemplates() is now exposed as a callable API outside of the editor context.
+
+#### Core
+
+* New: Added FMath::ClosestPointOnInfiniteLine() to distinguish it from the (poorly named) ClosestPointOnLine() that actually works on segments.
+
+* New: Added a CopyTemp function for making a temporary as a copy (instead of a move) of an object to aid visibility of unintended copies of setter functions defined only as taking an rvalue reference.
+
+* New: Added a DEPRECATED_MACRO deprecation macro for deprecating macros.
+
+* New: Added a new LerpRange function to allow interpolation of FRotator without taking the shortest path.
+
+* New: Added Algo::Copy and Algo::Transform.
+
+* New: Added C++11 initializer list support to TArray and TSet.
+
+* New: Added Core/Containers/TArrayView based on Engine/FixedSizeArrayView.
+
+    * TArrayView allows a function to accept a range of elements as if it was taking a TArray, while allowing the caller to pass any of:
+
+        * TArray (with any allocator)
+
+        * C Array
+
+        * Pointer + Size
+
+        * Initializer list!
+
+    * Note: The new and improved TArrayView is not const-propogating. const TArrayView<T> does not imply "const T" and will not bind to a const TArray!
+
+        * You want TArrayView<const T> as the parameter in 99% of functions that would currently take a const TArray<T>&. Note: TArrayView should be a value parameter, not a reference.
+
+* New: Added FastDecimalFormat overloads to write into an existing string.
+
+    * This helps avoid an extra allocation if you already have a pre-sized string that you're writing the number to (as is the case in FText::Format).
+
+* New: Added Find and FindByPredicate algorithms for operating on arbitrary ranges.
+
+* New: Added LexicalConversion::ToString overloads for const CharType* and bool.
+
+* New: Added move semantics support to TArray::HeapPush.
+
+* New: Added some extra logging information when a class fails to hot reload.
+
+* New: Added static version of FName::IsValidXName.
+
+    * This allows you to verify name-like strings without having to convert them to an FName (and thus add them to the name table)
+
+* New: Added TArray::EmplaceAt for direct construction of array elements at a given index.
+
+* New: Added TWeakObjectPtrMapKeyFuncs as the KeyFuncs for TMaps which use TWeakObjectPtr as a key, so that the map is not invalidated when the pointer becomes stale.
+
+* New: Added Visual Studio debugger visualization of TTuples of up to 6 elements.
+
+* New: All configuration files are now scanned for an Engine section containing TaggedPropertyRedirects. This allows plugin configurations to register their own property redirects in cases where a property's name changes. Please note that this engine section will ONLY be searched for TaggedPropertyRedirects, and no other options are currently inspected for non-Engine configuration files.
+
+* New: Async Loading will no longer wait for all packages to finish serializing before calling PostLoad on packages that have already been serialized.
+
+* New: Errors are now raised when a UnrealHeaderTool markup is in a skipped preprocessor block.
+
+* New: Exposed the Key and Value typedefs for TMap and TSet.
+
+* New: FModuleManager::UnloadOrAbandonModuleWithCallback's 'abandon only' functionality has been split into its own function.
+
+* New: Garbage Collector will no longer be responsible for generating class token stream. Instead the token stream will be generated on startup or when a (Blueprint) class has finished loading.
+
+* New: Made check() calls assume the condition is true even in !DO_CHECK builds.
+
+* New: Made FFeedbackContext's Errors and Warnings list thread safe.
+
+* New: Made the TAnd, TAndValue, TOr and TOrValue metafunctions variadic.
+
+* New: Prevented FStrings being concatenated with integers, as this would concatenate a Unicode character of the value of that integer, rather than the expected decimal string representation of that integer.
+
+* New: Removed obsolete PLATFORM_COMPILER_HAS_EXPLICIT_OPERATORS, FORCEINLINE_EXPLICIT_OPERATOR_BOOL and SAFE_BOOL_OPERATORS macros.
+
+* New: Split inline shader registration from serialization. Serialization can happen on the async loading thread but registration should only happen on the game thread. Removed a lot of critical section locks.
+
+* New: UEngine now defines a Start() function that subclasses can use to start game related things after initialization of the engine.
+
+    * This is done so that after the Init() call on UEngine, we can then perform a module load for the ELoadingPhase::PostEngineInit phase of loading, then inform the UEngine that it's time to start the game.
+
+    * UGameEngine now uses the Start() phase to tell the GameInstance to Start, rather than at the base of Init(), before all modules have had a chance to load.
+
+* New: Unified output formatting across FOutputDeviceFile, FOutputDeviceMemory and FOutputDeviceArchiveWrapper and fixed potential race conditions caused by serializing log lines in chunks.
+
+* New: UStruct::SerializeTaggedProperties has been restructured to use new virtual UProperty::ConvertFromType for backwards compatibility type conversions.
+
+* New: Added ForceInline to TIndexedContainerIterator<...>::operator!=(...). 
+
+    * This change improved DebugGame performance of a for loop using ranged-based syntax significantly.
+
+* New: Added an assert to PackageTools::GetFilteredPackageList to help track down a bad object.
+
+* Bugfix: Crash fix for intermittent issue that does not delete linkers during FFeedbackContext::ProgressReported().
+
+* Bugfix: Fix hang in FTaskGraphInterface::BroadcastSlow_OnlyUseForSpecialPurposes caused by back to back calls.
+
+* Bugfix: Fixed a hash conflict that could occur when both the case-sensitive and case-insensitive FName hashes were identical.
+
+* Bugfix: Fixed a long-standing issue with undo/redo nulling lazy pointers.
+
+    * This particularly affected Landscape.
+
+* Bugfix: Fixed comparisons of UnrealHeaderTool manifests which may affect UHT makefile correctness.
+
+* Bugfix: Fixed const correctness in UObject* overload of IPropertyHandle::SetValue.
+
+* Bugfix: Fixed const correctness of FLogCategoryBase::IsSuppressed.
+
+* Bugfix: Fixed hot reload on Mac when there are no changes.
+
+* Bugfix: Fixed many Visual Studio static analysis warnings.
+
+* Bugfix: Fixed some hot reload output which was being reported as warnings.
+
+* Bugfix: Fixed the parsing of non-ISO8601 datetimes in JSON.
+
+* Bugfix: Fixed UnrealHeaderTool generating the wrong properties for int and unsigned int when the build environment's int sizes differ from the the target environment.
+
+* Bugfix: Fixed 'unreferenced local variable' warnings used in UE_CLOG macros when logging is disabled.
+
+* Bugfix: Fixed three bugs required for script support to work properly.
+
+* Added a non-const version of MultiFindPointer() to TMultiMap.
+
+* Added a proper forwarding constructor to FAsyncTask to allow references to be passed.
+
+* Clarified some wording in TFunction assert messages and comments.
+
+* FAsyncPackage now creates the meta-data before processing the remaining exports
+
+    * This matches the behavior of FLinkerLoad::LoadAllObjects, as other objects may depend on the meta-data being loaded before them.
+
+* Improved Algo::Accumulate when used with movable types (e.g. FString).
+
+* Made GetDefault<T>(UClass*) assert that the class is castable to T. Otherwise it is an unsafe cast.
+
+* Made sure async loading is not marked as complete until all async packages have been deleted.
+
+* Made sure SQLLite database handle is nulled after it's been closed.
+
+* Optimized some FString and FPaths functions to produce fewer temporary strings.
+
+* Prevented the possibility of move constructors and move assignment operators being defined in UCLASSes.
+
+* Removed a developer-only check in STATS which was generating lots of redundant strings.
+
+* Removed obsolete macro-based delegate implementation.
+
+* Removed PLATFORM_COMPILER_HAS_VARIADIC_TEMPLATES.
+
+* Replaced some check()s with checkSlow()s in the container allocators.
+
+* Reversed the order the UObject delete listeners were notified of UObject deletion to avoid skipping the next listener if the current one removed itself from the array.
+
+* Split up TTypeTraits into multiple individual traits to avoid erroneous VS2015 compiler errors and for composability.
+
+* Stomp allocator can now be used in the editor builds. At least in smaller projects.
+
+* Stopped FError::Throwf from being usable outside of UnrealHeaderTool.
+
+* Suspending or resuming thread heartbeat will not generate a heartbeat for threads that haven't sent any heartbeats yet.
+
+* When using FFeedbackContextAnsi messages logged with Display verbosity will no longer be added to the Warnings array.
+
+* Worked around some spurious VS2015 compiler errors in defaulted and deleted class functions.
+
+#### Editor and Tools
+
+* New: Added FAssetData::GetTagValue and FAssetData::GetTagValueRef to handle getting asset tag values in a type-correct way.
+
+    * This allows type-conversion using LexicalConversion, and also has specializations for FString, FText, and FName.
+
+* New: Added LANDSCAPE_API to Landscape MaterialExpression classes so user-created modules can make use of them.
+
+* New: Allowed whitespaces for meta class names. Ex: UPROPERTY(meta = (AllowedClasses = "Texture2D, MaterialInterface")) will now show both Texture2D and MaterialInterface types.
+
+* New: Removed the default "Toggle Pause/Play" command keybinding from BaseInput.ini and made this a customizable keybinding option in the "Editor Preferences".
+
+* New: UBT on windows will now look for existence of a "UE4CSharp.prog" file next to a C# program project, that would like to be added to the UE4.sln.
+
+* Changed function signature for FSequencer::NotifyMovieSceneDataChanged so that it takes an enum describing the type of change.
+
+* World Outliner now correctly calls ProcessEditDelete on editor modes that have asked to process delete operations
+
+##### Landscape
+
+* Bugfix: Fixed Landscape class hierarchy.
+
+    * LandscapeProxy is now LandscapeStreamingProxy.
+
+    * A new base class has been added for both Landscape and LandscapeStreamingProxy (called LandscapeProxy).
+
+    * Most code shouldn't notice the change.
+
+#### Gameplay Framework
+
+* New: Added a GetEffectiveFrameRateLimit method to UGameUserSettings, to allow games to override the default behavior for things like a different limit in menus versus gameplay.
+
+* New: Added GetSquaredDistanceToBody to FBodyInstance and virtual GetSquaredDistanceToCollision to UPrimitiveComponent.
+
+* New: Added minor support to data tables to support game-specific post-import fixup of data
+
+    * Added ability to mark variables in data table row structs with the meta field DataTableImportOptional, which implies the variable is not necessarily expected to be in the imported source file as a column (warnings won't be generated on import). Useful for variables that will only be populated by a post-import step in code.
+
+    * Added a new method, OnPostDataImport, to FTableRowBase to allow games to provide custom parsing, fix-up, etc. of data imported into data tables.
+
+    * Fixed the data table importer dialog to show data table row structs alphabetically.
+
+* New: Added support for UNumericProperty (int64, etc...) and doubles as allowable types in data table structures.
+
+    * Note: Blueprints still can't access these types, so use is limited to C++ for now.
+
+* New: Added time spent in each hitch bucket to the FPS chart .log output and analytics output.
+
+* New: AHUD::WhiteColor, GreenColor, and RedColor have been made static const and deprecated. Use FColor::White, Green, and Red instead.
+
+* New: Made CalcSceneView and GetProjectionData in ULocalPlayer virtual.
+
+* New: Marked engine performance target cvars ECVF_Scalability so they can be adjusted via scalability buckets at runtime for games that target different framerates on different levels of hardware.
+
+* New: Scalability buckets can now have more or fewer than 4 entries, controlled by sg.[BucketName].NumLevels cvars.
+
+* New: SceneComponent AttachChildren, AttachParent, AttachSocketName are now private. Use GetAttachChildren/Parent/SocketName, AttachTo, or SetupAttachment as appropriate.
+
+* New: UActorComponent::GetWorld() marked final and inlined to optimize overhead.
+
+* New: Until now the only type used for attributes in the gameplay ability system has been floats. A new type, FGameplayAttributeData, has been added. At some point in the future, support for float based attributes will probably be removed. At that point only attributes using the FGameplayAttributeData struct or a struct derived from it will be valid attributes. Until then it is possible for projects to have both FGameplayAttributeData and float based attributes.
+
+* New: Added comment to FHitResult for ImpactPoint to explain behavior when it's penetrating.
+
+* New: Added nullptr checks to ProceduralMeshComponent's SceneProxy.
+
+* New: Added perf stat for SetHitResultFromShapeAndFaceIndex().
+
+* New: GetNetMode() optimizations.
+
+    * Added IsNetMode() which is more efficient in non-editor builds when considering dedicated servers.
+
+    * Wrapped in inline functions that check IsDedicatedServer() where possible.
+
+    * Converted some components to use new methods.
+
+* New: AHUD::bShowHUD now defaults to true
+
+* New: Minor optimization in UGameplayTagsManager::FindTagNode() to stop taking copies of a tags children array and churning memory.
+
+* Bugfix: Crash fix when calling FAudioDevice::CreateComponent as part of shut down.
+
+* Bugfix: Fixed crash when ShouldSpawnAtStartPoint returns true if Player->StartSpot is null.
+
+* Bugfix: Fixed issue that could cause the data table reimporter to change the reimported table back to a base Data Table class instance instead of the previously imported subclass, in the event of using subclasses of Data Table.
+
+* Bugfix: Fixed optimization to IsLocalController() using wrong platform define.
+
+* Bugfix: Fixed overlaps being dropped from within a FScopedMoveUpdate when rotation occurs. Removed invalid assert.
+
+* Bugfix: Fixed Net relevant WorldSettings objects so they no longer inadvertently end up in the Actors array twice.
+
+* Bugfix: Fixed seamless travel so that if it chooses not to carry a GameMode over, it will now correctly create a new one.
+
+* Bugfix: Fixed the UEngine::GetAllLocalPlayerControllers function to no longer add null entries to the output PlayerList array, which was possible in some rare cases before.
+
+* Changed the code in AActor::IncrementalRegisterComponents to obey bAutoRegister for the root component rather than asserting that it is true. Note: If children components have bAutoRegister=true, they will still pull the root component into the fray and cause it to be registered first.
+
+* Converted TAutoConsoleVariables in CharacterMovementComponent to FAutoConsoleVariableRefs in (renamed) CharacterMovementCVars namespace.
+
+    * Avoids need for GetValueOnGameThread() which was adding unnecessary overhead and preventing some users from easily multithreading portions of the code.
+
+* Disabled concurrent tick queuing in Windows.
+
+* Embedded FPS chart preamble/postamble/row .html files into ChartCreation.cpp to permanently solve packaging woes.
+
+* FTickableGameObjects will no longer tick multiple times when multiple worlds are ticking, such as having a BP preview world, a persona world, or multiple PIE instances.
+
+* Made FSavedMove_Character::CanCombineWith easier to debug. Consolidated duplicate code to one block.
+
+* Marked FHitResult, FOverlapResult, FOverlapInfo as POD types. Avoids destructors etc when in TArrays.
+
+* Optimized UCapsuleComponent::CalcBounds() to remove sqrt and use tighter bounding sphere radius.
+
+* Optimized AudioComponent to set bUseAttachParentBounds=true to reduce cost of transform updates (avoids virtual CalcBounds() when transform changes).
+
+* Optimized character networking to remove allocations after initial spawn from client saved move processing.
+
+* Prevented ACharacter::bNotifyJumpApex from being editable and cleaned up comments.
+
+* Reduced allocations during movement and overlap queries and when fetching shapes from physx actors.
+
+* Reduced allocations in UI Canvas items.
+
+* Reduced allocations in UObjectBase::CreateStatID(). It throws off profiling results significantly in some cases.
+
+* Worked around some of the drawbacks of Windows scheduler. This would occasionally cause hitches on quad core machines with additional load in the tick task manager.
+
+* Wrapped call from UCharacterMovementComponent::PostPhysicsTickComponent() to UpdateBasedMovement() in a FScopedMovementUpdate to accumulate moves with better perf.
+
+#### Localization
+
+* New: Added the ability to version each localized string individually when loaded into the localization manager.
+
+    * The single 32-bit global history has now been replaced with two 16-bit histories. One is global, and is updated whenever the culture is changed (or a LocRes file is loaded), and the other is local to each string, and is updated if the display string is changed outside of a culture update (to handle cases where the display string is changed, but the key is preserved). Changing the global history will reset all local histories.
+
+* Renamed FContext to FManifestContext to better reflect its purpose and avoid naming conflicts with other code.
+
+#### Networking
+
+* New: Added support for prioritizing actors for replay recording. This feature can be used along with the per-frame recording time limit (see UDemoNetDriver::SetMaxDesiredRecordTimeMS) to improve playback quality for the important actors in the scene.
+
+    * This feature is enabled by calling UDemoNetDriver::SetActorPrioritizationEnabled on the recording UDemoNetDriver.
+
+    * If desired, customize the priority of actors by overriding the AActor::GetReplayPriority function. This works similarly to the AActor::GetNetPriority function for non-replay replication.
+
+    * You may also override the viewer used for the priority calculation with the UDemoNetDriver::SetViewerOverride function.
+
+* New: Added the ability to limit the size of the buffer used by the in-memory replay streamer (FInMemoryNetworkReplayStreamer).
+
+    * This will prevent the memory buffer that stores the replay data from growing unbounded over the course of a game. Only the most recent data will be kept, so it won't be possible to play back the replay from an earlier time that's no longer in the buffer.
+
+    * To set the buffer size, call INetworkReplayStreamer::SetTimeBufferHintSeconds on an instance of FInMemoryNetworkReplayStreamer. The implementation will keep at least the most recent desired number of seconds available for playback.
+
+* New: Added the ability to pause and resume replay recording, and get the current paused state. See UDemoNetDriver::PauseRecording and UDemoNetDriver::IsRecordingPaused.
+
+* New: Added the ability to set a desired per-frame time limit for recording replays. The time is not guaranteed, but setting it can help reduce the CPU load of replay recording in exchange for potentially lower quality playback.
+
+    * After the time limit is hit, additional actors that would have been recorded on the current frame will be deferred until a future frame. Since these actors are being recorded at a lower frequency, this may cause a reduction in playback quality.
+
+    * To set the desired per-frame time limit, call UDemoNetDriver::SetMaxDesiredRecordTimeMS on the recording UDemoNetDriver.
+
+    * This feature may be combined with replay actor prioritization (see UDemoNetDriver::SetActorPrioritizationEnabled) to record important actors at a higher frequency and maintain more playback quality.
+
+* New: Updated the FUniqueNetId of a connecting client to be passed along in the client beacons so the server can tell who is making requests.
+
+* New: Added ENGINE_API to FRepMovement, as the constructor is in a source file.
+
+* Bugfix: Fixed an issue that could cause arrays to get out of sync with the game server when using a FFastArraySerializer with predictive elements and recording a replay on the client.
+
+* Bugfix: Fixed an issue where network profiles after the first one created during an execution of the engine would have broken name and address mappings.
+
+* Bugfix: Fixed issue wherein replicated actors destroyed on the server would not always be properly destroyed on the client when used in conjunction with network dormancy and queued bunches.
+
+* Any actors that are set not to load on the client (bNetLoadOnClient = false) will be force deleted when initializing a level’s actors on a client.
+
+* Fixed misspelling in comment in NetworkVersion.h.
+
+* Corrected filename case in #include "IPAddress.h" statement.
+
+#### Online
+
+* New: FQosInterface::GetDefaultRegionString() has been replaced by FQosInterface::GetRegionId(). It is just a rename of the function
+
+* New: FUniqueNetIdRepl is stored on the UNetConnection and passed into engine login functions rather than TSharedPtr<FUniqueNetId> as it is a more robust wrapper of the same data.
+
+* New: Added a base class UpdateManager that can coordinate patch and hotfix checks inside a game. Taking advantage of the OnlineHotfixManager class can help setup a basic infrastructure for using that class to check for hotfixes while the game is running.
+
+* New: Added an Expiration Date field to FOnlineStoreOffer. This field defaults to Unreal Max Time.
+
+* New: Added the capability for actors to temporarily pause replication of themselves over a certain channel. This can significantly reduce bandwidth costs and reduce cheating. ShooterGame has implemented an example - look for AShooterCharacter::IsReplicationPausedForConnection and AShooterCharacter::OnReplicationPausedChanged.
+
+* Bugfix: Fixed issues with FSocketBSD::Recv returning false on SE_EWOULDBLOCK when polling a non-blocking socket.
+
+* FSocketBSD::Recv, RecvFrom and the IPv6 equivalents will now return false when recv returns 0, which indicates the connection was shutdown in an orderly manner.
+
+#### Other
+
+* New: Added Framework zip paths to list of required external files when precompiling targets.
+
+* New: Added support for per file Intellisense PCH settings, which greatly improves its startup speed.
+
+    * Set bUsePerFileIntellisense to true to activate, which can be done through your local BuildConfiguration.xml file.
+
+    * This has been seen to use a lot more memory and even cause crashes in large Visual Studio solutions. Only supported after VS2015 Update 3.
+
+* New: Added TArray<FString> support to FAnalyticsEventAttribute.
+
+* New: Made a first pass at async loading improvements. Mostly disabled.
+
+* New: Removed hard coded staging of dlls for Crash Reporter and use its receipt instead.
+
+* New: Added validation of tag names to FindOrAddTagSet so that they only contain a single '#' character and any spaces come after.
+
+* Bugfix: Fixed an issue in ability system where gameplay effect handles would be improperly invalidated whenever a sublevel was loaded.
+
+* Ensured that TempStorageFileList uses forward slashes as its path separators so that it's easily used on Mac and Windows.
+
+* Improved the CPU burden of loading in several places. Made substantial progress on the complete loading revamp (currently disabled).
+
+* Removed EnterScope and LeaveScope from ReadGraphBody so that included files are treated as being in the same scope (allows use of properties across files).
+
+#### Paper2D
+
+* New: Added PAPER2D_API to FSpriteDrawCallRecord so it can be used externally.
+
+#### Physics
+
+* New: Exported ConvertToCollisionChannel, ConvertToObjectType and ConvertToTraceType members of UCollisionProfile.
+
+* New: Exposed mobility filtering query params. This allows users to filter out static mobility, for example from scene queries.
+
+* New: Refactored constraint instance so that we can easily swap parameters at runtime using profiles.
+
+* Bugfix: Fixed Unreal collision bounds calculation not accounting for rotation of convex elements.
+
+#### Platforms
+
+* Bugfix: Fixed bVerfiyPeer flag to only be set after CertBundlePath has been set, otherwise trying to verify the SSL session without any cert file doesn't make sense.
+
+* Made sure dSYM generation action in UBT on Mac does not start before the source dylib is ready.
+
+##### Android
+
+* New: Added a new variable to AndroidRuntimeSettings called "bCreateAllPlatformsInstall" that makes packaging create install files for all platforms (PC and Mac/Linux).
+
+* Renamed FManifestReader to FAndroidFileManifestReader in AndroidFile to prevent conflict.
+
+* Ensured that the OnControllerConnectionChange delegates will now properly dispatch on a gamepad connection.
+
+##### iOS
+
+* New: Added PRAGMA_DISABLE_OPTIMIZATION_ACTUAL and PRAGMA_ENABLE_OPTIMIZATION_ACTUAL defines for iOS.
+
+##### Linux
+
+* Updated LinuxToolChain to use the same output delegate for all of its actions when cross-compiling.
+
+    * This avoids the compile and link actions being split into different batches.
+
+##### Mac
+
+* New: Added a check if UnrealBuildTool.csproj exists before trying to compile it in Mac GenerateProjectFiles.sh. Fixes a problem in binary distribution where the script would show an error (but still succeed) due to missing UnrealBuildTool.csproj.
+
+* Bugfix: Fixed a bug in Kismet where a reference to an element within a TMap was being used to add a new value, which results in use-after-free as the memory is reallocated before initialising the new element.
+
+* Bugfix: Fixed compilation on iOS 10 by disabling depth-clip mode on iOS, despite earlier headers it isn't actually available.
+
+* Bugfix: Fixed UBT not building on non HFS partitions on OSX.
+
+* Bugfix: Crash fix in USoundNodeLooping::NotifyWaveInstanceFinished.
+
+* Changed BinnedAllocFromOS/BinnedFreeToOS/PageProtect on Apple platforms to use mmap/munmap/mprotect to match expected behaviour.
+
+* Disabled the "-Wdelete-non-virtual-dtor " warning on iOS, tvOS & macOS.
+
+* Implemented RHIGetSupportedResolution & RHIGetAvailableResolutions for Mac Metal.
+
+* Made bGeneratedSYMFile true by default and changed some ifs in Mac UBT code so that non-debug configs always build dSYM files on Mac, unless bGeneratedSYMFile is set to false in BuildConfiguration.xml.
+
+##### Playstation 4
+
+* New: Added validation that the texture pool size is not set to automatic (-1), so it fails with a helpful error message rather than a crash attempting to allocate too much memory.
+
+##### Xbox One
+
+* New: Added support for Visual Studio 2015 Update 3.
+
+* New: Added the bDisableSymbolCache to the link environment configuration to allow Xbox One builds to skip the creation of the runtime symbol cache.
+
+    * Disabling the symbol cache creation will slightly decrease build times and increase available memory at runtime, but stacks will be printed without symbols.
+
+* New: Changed appWinPumpMessages so that it doesn't rely on the current thread to retrieve the dispatcher.
+
+* New: Enabled the use of constexpr for the engine on Xbox One.
+
+* New: Unreal Build Tool now generates the correct build environment for Xbox One directly.
+
+    * This method puts less reliance on the environment a build is launched in and eliminates a build issue if the PATH environment variable is too long.
+
+* Bugfix: Fixed an issue where builds could fail if the DurangoXDK environment variable did not contain a trailing slash.
+
+* Disabled a warning that was impacting some users because of a toolchain issue.
+
+##### HTML5
+
+* New: The Engine will now obtain the game's w&h resolution with javascript.
+
+* New: Allowed use of TMaps in JSON.
+
+* Emscripten SDK moved from Engine/Source/ThirdParty/HTML5/emsdk to Engine/Extras/ThirdPartyNotUE/emsdk.
+
+* Ensured HTML5 LaunchOn with pathnames that contain whitespaces is now handled properly - escape quotes needs to be escaped itself.
+
+##### Windows
+
+* Bugfix: Fixed errors compiling WmfMedia and SteamVR plugins for Windows XP.
+
+#### Programming
+
+* New: Added compatibility for Visual Studio 2015 Update 3.
+
+* New: Added feature to reuse existing response files if the contents are unchanged.
+
+* New: Added SlackIntegrations module for sending messages to Slack from UE4 projects.
+
+* New: FActiveMorphTarget is no longer tagged as a USTRUCT.
+
+* New: Fixed an issue where compiling CPP files using the Visual C toolchain could be too long for XGE to process.
+
+* New: The UnrealVS Visual Studio extension now exposes the "UnrealVS.CompileSingleFile" command. Binding to a shortcut key (eg. Ctrl-F7) from the Visual Studio options dialog allows quickly compiling the current file without linking binaries.
+
+* New: Added a UBT config setting to allow overriding the output directory for PCH files. To use, edit Engine\Saved\UnrealBuildTool\BuildConfiguration.xml and add: <BuildConfiguration> <PCHOutputDirectory>D:\TestOutputDir</PCHOutputDirectory> </BuildConfiguration>.
+
+* Changed Unreal Build Tool batch files to use MSBuild version 14 (the version that ships with Visual Studio 2015).
+
+* Ensured that UnrealBuildTool on Windows will now prioritize user settings in the My Documents folder over the one in the AppData folder, but reads both files if present. Previous behavior was to ignore the one with the older timestamp.
+
+* Updated OpenSSL Runtime Dependencies to include VS2015 and remove VS2012.
+
+#### Rendering
+
+* New: Adapted GetStaticBatchElementVisibility to support changing visibility of batches with a single element (for view-dependent culling of static mesh elements).
+
+    * Added GetStaticBatchElementShadowVisibility that allows culling static mesh elements when rendering shadow depths using information about the light.
+
+    * Note: Any user-created vertex factory that implements GetStaticBatchElementVisibility will need to set bRequiresPerElementVisibility on their FMeshBatch.
+
+* New: Added ability to request a screenshot without a file name, for screenshots that won't be written to disk.
+
+* New: Added cvar r.StencilForLODDither, which is compiled into shaders and forces a full prepass when enabled. This allows r.EarlyZPass to remain changeable at runtime (when not forced to a value by features that are enabled).
+
+* New: Added cvars that can be used to force the engine to target shader model 4 instead of shader model 5 on cards from a specific vendor (r.ForceAMDToSM4, r.ForceIntelToSM4, r.ForceNvidiaToSM4).
+
+* New: Added r.D3DDumpD3DAsmFile so d3d shaders dump out an assembly file when dumping shader compilations.
+
+* New: Allowed r.HLOD console command to be used in Test configurations.
+
+* New: Baked out distribution UObjects for particle systems.
+
+* New: Exported AddStreamingManager and RemoveStreamingManager functions on FStreamingManagerCollection.
+
+* New: Added forward renderer initial implementation - not ready for general use.
+
+    * Point and spot lights are culled to a frustum space grid, base pass loops over culled lights.
+
+    * Light culling uses a reverse linked list to avoid a per-cell limit, and the linked list is compacted to an array before the base pass.
+
+    * New cvars to control light culling: r.Forward.MaxCulledLightsPerCell, r.Forward.LightGridSizeZ, r.Forward.LightGridPixelSize.
+
+    * A full Z Prepass is forced with forward shading. This allows deferred rendering before the base pass of shadow projection methods that only rely on depth.
+
+    * Dynamic shadows are packed based on the assigned stationary light ShadowMapChannel, since stationary lights are already restricted to 4 overlapping.
+
+    * GBuffer render targets are still allocated.
+
+    * Fixed several issues in parallax corrected base pass reflections - not blending out box shape, discontinuity in reflection vector, not blending with stationary skylight properly.
+
+    * Forward shading is now used for TLM_SurfacePerPixelLighting translucency in the deferred path.
+
+* New: Moved RHI prefix out of command list macros so you can find callers in a Find in Files for RHIBlendState, etc.
+
+* New: Split ProceduralMeshComponent plugin into Runtime and Editor modules.
+
+* New: Updated VulkanSDK headers and libraries to 1.0.17.0.
+
+* New: Added code to prevent flooding of rendering thread with heartbeats.
+
+* Bugfix: Fixed crash bug in FCanvasBorderItem rendering caused by a null texture dereference.
+
+* Bugfix: Fixed ensures related to ECVF_Scalability.
+
+* Bugfix: Fixed a bad hash on landscape grass components.
+
+* Bugfix: Fixed a critical bug in the render thread fencing that could sometimes cause fences to complete before previous commands.
+
+* Bugfix: Fixed a large number of API validation warnings and errors across all shaders.
+
+* Bugfix: Fixed an uninitialized memory bug in the static mesh scene proxy HLODcoloration visualization code.
+
+* Bugfix: Fixed a bad call to RemoveAt() in FFinalPostProcessSettings::UpdateEntry() causing potential infinite loop.
+
+* Bugfix: FIxed calls to length(float) on Metal, GL & Vulkan.
+
+* Increased stack sizes for thread pools.
+
+* MaterialParameterCollections will now create default resources (uniform buffers) which are used when no valid FScene is present (eg DrawTile while exporting materials to lightmass).
+
+* Potential crash fix on foliage occlusion queries and reflection captures.
+
+* Removed outdated cvar, "r.ShaderModelDebug".
+
+* Removed unused setting bAllowDebugViewmodesOnConsoles (replaced some time ago by r.ForceDebugViewModes).
+
+* Renamed "r.SaveUncompressedEXRFrames" to "r.SaveEXR.CompressionQuality" as part of a refactor.
+
+* Renamed material function "IsUsedWithDeferredDecal" to "IsDeferredDecal".
+
+* Set the "ToggleLight" console command to now be flagged as a cheat.
+
+##### FX
+
+* New: Added a CVar to disable/freeze GPU particle simulation (r.GPUParticle.Simulate).
+
+##### Lighting
+
+* New: Refactored shadow depth.
+
+    * Shadow setup and render target allocation now happens in InitViews, and shadow depth rendering happens at one spot in the frame. This provides control over where shadow depths are rendered for things like async compute, and allows easy atlasing of shadowmaps for forward shading.
+
+    * The 33Mb of shadow depth buffers in FSceneRenderTargets has been removed, and shadow depth buffers are now allocated as needed. Shadow depth memory is now unbounded which can cause Out of Video Memory in extreme cases.
+
+    * A large amount of duplicated code to handle each shadow type has been combined.
+
+    * Cleaner parallel rendering: no more view hacking for the shadow depth pass, no more shadow depths in the middle of translucency.
+
+    * 'vis ShadowDepthAtlas' or 'vis WholeSceneShadowMap' must now be used to visualize the shadow depth textures.
+
+##### Optimizations
+
+* New: Added a cvar that controls unbinding of all texture resources between materials changes in the DX11 renderer (r.UnbindResourcesBetweenDrawsInDX11) to improve the readability of GPA captures, off by default as it has a minor CPU performance impact.
+
+* New: Added a scalability CVar for landscape LOD biasing (r.LandscapeLODBias), which can have a huge impact on triangle count and help low-end vertex bound GPUs.
+
+* New: Made the optimization to combine upscaling/downscaling and tonemapping optional based on the amount of upscaling that will occur.
+
+    * r.Tonemapper.ScreenPercentage has been renamed to r.Tonemapper.MergeWithUpscale.Mode.
+
+    * r.Tonemapper.MergeWithUpscale.Threshold is a new setting used when r.Tonemapper.MergeWithUpscale.Mode is set to 2, which indicates to only try to merge the passes if the ratio of the area before upscale/downscale to the area afterwards is greater than the threshold.
+
+    * This prevents running the tonemapper on all of the target res pixels when the source res is far smaller, as that can cause it to be a loss to merge the passes.
+
+* New: Made sure that Vulkan always runs using an RHI thread.
+
+* Removed a bogus autocomplete for ShowMaterialDrawEvents, which was previously renamed r.ShowMaterialDrawEvents.
+
+#### UI
+
+* Moved FSlateIcon definition to SlateCore.
+
+##### Slate
+
+* New: Added FText property handle type.
+
+    * Required for manipulating FText properties from a slate details customization.
+
+* FExternalDragOperation can now contain both text and file data at the same time.
+
+    * This better mirrors what the OS level drag-and-drop operations are capable of, and some applications will actually give you both bits of data at the same time.
+
+##### UMG
+
+* Demoted errors regarding widget-bound properties when first compiling a newly created widget blueprint - otherwise an ensure occurs.
+
+#### VR
+
+* New: Implemented VRHeadsetLost and VRHeadsetReconnected delegates for PlayStation VR.
+
+* New: Removed ISR's use of SV_CullDistance and now relying only on SV_ClipDistance to cut down on interpolators.
+
+## Upgrade Notes
+
+### C++ API Changes
+
+* SceneComponent AttachChildren, AttachParent, AttachSocketName are now private. Use GetAttachChildren/Parent/SocketName, AttachTo, or SetupAttachment as appropriate.
+
+* Replace uses of AHUD::WhiteColor, GreenColor, and RedColor with FColor::White, Green, and Red.
+
+* In the HTML5JavaScripteFx.js and .h files, the function signature for UE_MakeHTTPDataRequest() has changed.
+
+* With regards to localized strings, because of the change from an int32 to a uint16, 0, rather than INDEX_NONE, is now considered the "unset" value for a history.
+
+* FContext (in our internationalization manifest) has been renamed to FManifestContext to better reflect its purpose and avoid naming conflicts with other code
+
+### Animation
+
+* Added a validation path to see if mesh hierarchy matches with the animation skeleton when initializing animation.
+
+    * If you have issue with animation not playing, make sure to save the mesh correctly with the skeleton.
+
+    * Additive APIs are fixed to work with it.
+
+* "Get World From Context Object" does not behave the same as in non-game thread.
+
+    * This causes additive scale to save the difference of scale, which means 0 if no difference. If you're doing any math on those FTransform, it won't work correctly anymore due to 0 scale
+
+    * If you have issue with animation not playing, make sure to save the mesh correctly with the skeleton.
+
+    * Additive APIs are fixed to work with it.
+
+* Local Atoms has been changed to Bone Space Transforms, and Space Bases has been changed to Component Space Transforms. This includes functions that were accessing those members. The old methods have been deprecated and need to be updated to use the new functions.
+
+* Several animation related headers are no longer globally included. This means you may need to include them explicitly if you were using those classes.
+
+    * Common examples are Animation/AnimInstance.h, Animation/AnimMontage.h or Animation/AnimNotifies/AnimNotifyState.h
+
+* The SkeletalMesh.h header is no longer globally included, so you may need to explicitly include it in your project.
+
+* The 'rigid' vertex array "Rigid Vertices" in Skel Mesh Chunk has been removed, all vertices are now considered 'soft' (rigid status was not used by GPU skinning). Any code dealing with vertex data will need to be updated.
+
+* FSkelMeshChunk has been removed and properties are moved into FSkelMeshSection. Any code dealing with 'chunks' will need to be updated.
+
+### Audio
+
+* In addition to class changes to support threaded sound evaluation, there are a number of patterns that are no longer viable.
+
+    * It is no longer safe to reference any game state from a Sound Node. Instead game thread classes should marshall relevant data to audio thread owned structures, using "Run Command On Audio Thread" in the Audio Thread.
+
+### Automation
+
+* If you experience regressions in automated tests, they may have been ignoring ensures because they were previously classed as warnings.
+
+### Core
+
+* Code which tries to concatenate a String with an integral value should either explicitly cast to a character type if it should append a char, or changed to use something like "To String" in Lexical Conversion if it should append the string representation of that integer.
+
+* Using Throwf in FError and associated catch blocks should be replaced with a regular C++ throw call and proper exception class.
+
+* PLATFORM_COMPILER_HAS_DEFAULT_FUNCTION_TEMPLATE_ARGUMENTS has been removed. This is now expected to be 1 on all supported platforms.
+
+* PLATFORM_COMPILER_HAS_VARIADIC_TEMPLATES has been removed. This is now expected to be 1 on all supported platforms.
+
+* Replace PLATFORM_COMPILER_HAS_EXPLICIT_OPERATORS usage by assuming it's defined to 1.
+
+    * Replace FORCEINLINE_EXPLICIT_OPERATOR_BOOL with 'FORCEINLINE explicit operator bool'.
+
+    * Remove any usage of SAFE_BOOL_OPERATORS.
+
+* Specializations of Type Traits which override its "Needs Copy Constructor", “Needs Move Constructor”, “Needs Copy Assignment”, “Needs Move Assignment”, and “Needs Destructor” traits will be ignored and compiler-based traits will be used instead.
+
+    * HAS_TRIVIAL_ASSIGN(T), HAS_TRIVIAL_COPY(T) and HAS_TRIVIAL_DESTRUCTOR(T) should be replaced with TIsTriviallyCopyAssignable<T>::Value, TIsTriviallyCopyConstructible<T>::Value and TIsTriviallyDestructible<T>::Value respectively.
+
+    * "TIsFloatType" should be replaced with “TIsFloatingPoint”.
+
+    * "TIsIntegralType" should be replaced with “TIsIntegral”.
+
+    * "TIsSignedIntegralType" should be replaced with “TIsSigned”.
+
+    * "TIsArithmeticType" should be replaced with “TIsArithmetic”.
+
+    * "TIsPointerType" should be replaced with “TIsPointer”.
+
+    * "TNoDestructorType" should be replaced with “TIsTriviallyDestructible”.
+
+* Use of variables inside a DO_CHECK-only block will not compile if the check() is outside a DO_CHECK-only block. Either both or neither should be inside the block.
+
+* When "Unload Or Abandon Module With Callback" in the Module Manager would be called with a third bool argument, that argument should be removed (if the bool is 'false') or replaced by a call to “Abandon Module With Callback”, passing only the module name argument (if the bool is 'true').
+
+* You may get errors about a skipped UFUNCTION, UPROPERTY etc. when upgrading. This is an indication that these elements are not reflected and so may be an indication of a bug (e.g. a skipped UPROPERTY which is a UObject pointer will not be tracked by the garbage collector). The usage should be reviewed and fixed, perhaps by removing the markup, moving the element outside of the skipped preprocessor block or some other form of refactoring.
+
+### Editor and Tools
+
+* Make sure to rebuild the proxy meshes for all clusters.
+
+### Gameplay Framework
+
+* All Tickable Game Objects should implement the virtual "Get Tickable Game Object World" function to indicate which World they should be ticked with, otherwise they will be ticked after all Worlds as part of the set of tickable game objects that do not belong to a world.
+
+### Abilities System
+
+* Converting existing float based gameplay attributes to use the FGameplayAttributeData struct is a two step process. First, resave any assets that contain gameplay attributes. Second, replace the float variables with FGameplayAttributeData, or a game specific struct derived from it, in your code.
+
+### Online
+
+* Online Subsystem modules have been moved to Plugins and must now be included in the plugins section of your project's uproject file and must no longer be referenced via "DynamicallyLoadedModuleNames" in your project's Build.cs files. Game projects can continue to reference Online Subsystem classes as before. Any custom modifications made in Engine that reference Online Subsystem must be routed through the Online Engine Interface.
+
+#### Windows
+
+* The Applocal Prerequisites Directory setting should be set to a base directory that contains files the binary files.
+
+    * The Applocal Prerequisites Directory can be specified as an absolute path or relative to the engine or project directory using the variables $(EngineDir) and $(ProjectDir).
+
+    * Files should be laid out within the Applocal Prerequisites Directory in the format <platform>\<package>\*.dll, where:
+
+        * <platform> is the specific target platform (i.e. Win32 or Win64).
+
+        * <package> is a unique package name (i.e. UCRT or Microsoft.VC140.CRT).
+
+### Rendering
+
+* The bounding box of a Static Mesh instance no longer includes its simple collision geometry. Now, the bounding box only includes renderable geometry.
+
+* Any user-created vertex factory that implements the Get Static Batch Element Visibility function will need to set "Requires Per Element Visibility" on its Mesh Batch.
+
+#### Optimizations7
+
+* The console command "r.Tonemapper.ScreenPercentage" has been renamed to “r.Tonemapper.MergeWithUpscale.Mode”, and offers additional modes (see the main release notes or cvar help for more information).
+
+#### Mobile Rendering
+
+* Mobile depth of field now requires Depth Of Field Method to be Gaussian. Other options will result in depth of field being disabled.
+
