#UE4 Docs: minor formatting and text updates

diff --git a/Engine/Documentation/Source/Support/Builds/ReleaseNotes/2016/4_13/ReleaseNotes_4_13.INT.udn b/Engine/Documentation/Source/Support/Builds/ReleaseNotes/2016/4_13/ReleaseNotes_4_13.INT.udn
--- a/Engine/Documentation/Source/Support/Builds/ReleaseNotes/2016/4_13/ReleaseNotes_4_13.INT.udn
+++ b/Engine/Documentation/Source/Support/Builds/ReleaseNotes/2016/4_13/ReleaseNotes_4_13.INT.udn
@@ -4,8 +4,6 @@ Crumbs:%ROOT%, Engine
 Description:
 order:13
 
-Unreal Engine 4.13 Release Notes
-
 This release brings hundreds of updates for Unreal Engine 4, including 145 improvements submitted by the community of Unreal Engine developers on GitHub! Thanks to all of these contributors to Unreal Engine 4.13:
 
 [REDACTED_CONTRIBUTORS]
@@ -342,11 +340,14 @@ There is a setting under VR in the Experimental section of Editor Settings that
 
 Sequencer can now import and export CMX EDL files for interchange with non-linear editing packages. 
 
+[REGION:raw]
 ![image alt text](image_35.jpg)(h:529 w:869 convert:false)
+[/REGION]
 
 The above shot shows a sequence that was exported to Adobe Premiere Pro.  Each shot in a sequence will be written to a separate movie file which is referenced by the EDL file.  Any sequencing changes made in Premiere can then be imported back into UE4’s Sequencer!
 
 Sequencer’s ability to export HDR data in OpenEXR files has been expanded to give the user a choice of the color gamut used to encode the HDR data. 
+
 ![image alt text](image_36.jpg)(h:220 w:492 convert:false)
 
 Finally, Sequencer now supports importing FBX animation directly to an object or track. You can also export animated tracks to FBX!
@@ -354,7 +355,10 @@ Finally, Sequencer now supports importing FBX animation directly to an object or
 ## New: Sequencer Burn-ins on Renders
 
 When rendering out your movie sequence, you can now configure a "Burn-in" for the exported images.  This is very useful in dailies for identifying and tracking shots.
+
+[REGION:raw]
 ![image alt text](image_37.jpg)(h:458 w:865 convert:false)
+[/REGION]
 
 ## New: Media Framework Overhaul
 
