New TSet document and minor updates to TMap document.

diff --git a/Engine/Documentation/Source/Programming/UnrealArchitecture/TMap/TMap.INT.udn b/Engine/Documentation/Source/Programming/UnrealArchitecture/TMap/TMap.INT.udn
--- a/Engine/Documentation/Source/Programming/UnrealArchitecture/TMap/TMap.INT.udn
+++ b/Engine/Documentation/Source/Programming/UnrealArchitecture/TMap/TMap.INT.udn
@@ -12,15 +12,15 @@ There are two types of map: TMap and **TMultiMap**.  TMap's keys are unique, and
 
 TMaps are primarily defined by two types - a key type and a value type - which are stored as associated pairs in the map.  It is convenient to refer to these pairs as the element type of the map as if it was an individual object.  In this document, element means a key-value pair, while individual components are referred to as the element's key or the element's value.  The element type is actually a **TPair**< KeyType, ElementType >, though it should be rare to need to refer to the TPair type directly.
 
-Like TArray, TMap is a homogeneous container and so all of its elements are strictly the same type.  TMap is also a value type, and supports the usual copying, assignment and destructor operations, as well as strong ownership of its elements, which will be destroyed when the map is destroyed.  The key type and value type must also be value types.
+Like TArray, TMap is a homogeneous container, meaning that all of its elements are strictly the same type.  TMap is also a value type, and supports the usual copying, assignment and destructor operations, as well as strong ownership of its elements, which will be destroyed when the map is destroyed.  The key type and value type must also be value types.
 
-TMap is a hashing container which means that, by default, the key type must support the **GetTypeHash** function and provide an operator== for comparing keys for equality.  Hashing is covered in more detail later.
+TMap is a hashing container which means that, by default, the key type must support the `GetTypeHash` function and provide an `operator==` for comparing keys for equality.  Hashing is covered in more detail later.
 
-TMap can also take an optional allocator to control the memory allocation behavior.  Standard UE4 allocators (e.g. FHeapAllocator, TInlineAllocator) cannot be used as allocators for TMap.  It instead uses set allocators, which define how many hash buckets the map should use and which standard UE4 allocators should be used for hash and element storage.  See TSetAllocator for more information.
+TMap can also take an optional allocator to control the memory allocation behavior.  Standard UE4 allocators (e.g. `FHeapAllocator`, `TInlineAllocator`) cannot be used as allocators for TMap.  It instead uses set allocators, which define how many hash buckets the map should use and which standard UE4 allocators should be used for hash and element storage.  See `TSetAllocator` for more information.
 
 The final TMap template parameter is **KeyFuncs**, which tells the map how to retrieve the key from the element type, how to compare two keys for equality and how to hash the key.  These have defaults which will just return a reference to the key, use operator== for equality and the non-member GetTypeHash function for hashing.  If your key type supports these functions, it is usable as a map key without the need to provide a custom KeyFuncs.
 
-Unlike TArray, the relative order of TMap elements in memory cannot be relied upon and iterating over the elements is likely to return them in a different order from that which they were added.  Elements are also unlikely to be laid out contiguously in memory.  The backing data structure of a map is a sparse array, which is an array with holes.  As elements are removed from the map, holes in the sparse array will appear which are then filled by the addition of future elements.  However, even though TMap doesn't shuffle elements to fill holes, pointers to map elements may still be invalidated, as the entire storage can be reallocated when it is full and new elements are added.
+Unlike TArray, the relative order of TMap elements in memory cannot be relied upon, and iterating over the elements is likely to return them in a different order from the order in which they were added. Elements are also unlikely to be laid out contiguously in memory. The backing data structure of a map is a sparse array, which is an array with holes. As elements are removed from the map, holes in the sparse array will appear. These holes are then filled when future elements are added. However, even though TMap doesn't shuffle elements to fill holes, pointers to map elements may still be invalidated, as the entire storage can be reallocated when it is full and new elements are added.
 
 ## Creating and filling a map
 
@@ -75,7 +75,7 @@ Like TArray, we can also use **Emplace** instead of Add to avoid the creation of
 	// 	{ Key: 3, Value: "Orange"    }
 	// ]
 
-Here, the two arguments are passed directly to the constructors of the key type and value type respectively.  This has no real effect for the int32 here, but it does avoid the creation of a temporary FString for the value.  Unlike TArray, it's only possible to emplace elements into a map with single argument constructors.
+Here, the two arguments are passed directly to the constructors of the key type and value type respectively.  This has no real effect for the int32 here, but it does avoid the creation of a temporary FString for the value.  Unlike TArray, it's only possible to emplace elements into a map with single-argument constructors.
 
 It's also possible to insert all the elements from another map by using the Append function to merge them:
 
@@ -150,7 +150,7 @@ We can check if a particular key exists in the map using the **Contains** functi
 	// bHas7 == true
 	// bHas8 == false
 
-Most of the time you will want to look up elements without knowing whether or not the key exists.  Using Contains followed by operator[] means a double lookup of the key, which we don't really want to do.  The **Find** function will allow you to do a single lookup, returning a pointer to the value of the found element rather than a reference, and returning null when the key does not exist:
+Most of the time, you will want to look up elements without knowing whether or not the key exists.  Using Contains followed by operator[] means a double lookup of the key, which we don't really want to do.  The **Find** function will allow you to do a single lookup, returning a pointer to the value of the found element rather than a reference, and returning null when the key does not exist:
 
 	FString* Ptr7 = FruitMap.Find(7);
 	FString* Ptr8 = FruitMap.Find(8);
@@ -208,7 +208,7 @@ The **FindKey** function allows you to perform a reverse lookup (find a key give
 	// *KeyMangoPtr   == 5
 	//  KeyKumquatPtr == nullptr
 
-The **GenerateKeyArray** and **GenerateValueArray** functions allow an array to be populated with a copy of all the keys and values respectively.  In both cases, the array being passed is emptied before population, so the resulting number of elements will always equal the number of elements in the map:
+The **GenerateKeyArray** and **GenerateValueArray** functions populate a TArray with a copy of all the keys and values respectively.  In both cases, the array being passed is emptied before population, so the resulting number of elements will always equal the number of elements in the map:
 
 	TArray<int32>   FruitKeys;
 	TArray<FString> FruitValues;
@@ -236,7 +236,7 @@ Elements can be removed from maps by using the **Remove** function and providing
 
 Removing elements actually leaves holes in the data structure, which you can see when visualizing the map in Visual Studio's watch window, but they have been omitted here for clarity.
 
-The **FindAndRemoveChecked** function can be used to remove an element by and return the associated value.  The checked part of the name means that the key will be checked for existence and assert if it doesn't exist:
+The **FindAndRemoveChecked** function can be used to remove an element from the map and return its value.  The "checked" part of the name means that the key will be checked for existence and assert if it doesn't exist:
 
 	FString Removed7 = FruitMap.FindAndRemoveChecked(7);
 	// Removed7 == "Pineapple"
@@ -344,7 +344,7 @@ TMap also supports move semantics which can be invoked using the **MoveTemp** fu
 
 ## Slack
 
-TMaps also have a notion of **slack** and can be used to optimize the population of the map.  The **Reset** function acts like a Empty() call but does not free the memory that was previously used by the elements:
+TMaps also have a notion of **slack**, and this can be used to optimize the population of the map.  The **Reset** function acts like an Empty() call, but does not free the memory that was previously used by the elements:
 
 	FruitMap.Reset();
 	// FruitMap == [<invalid>, <invalid>, <invalid>]
@@ -354,7 +354,7 @@ Here the map is emptied in the same way as Empty would, but the memory used for
 TMap does not provide a way of checking how many elements are preallocated, like TArray::Max() does, but still supports preallocating slack.  The Reserve function can be used to preallocate slack for a particular number of elements before adding:
 
 	FruitMap.Reserve(10);
-	for (int32 i = 0; i != 10; ++i)
+	for (int32 i = 0; i < 10; ++i)
 	{
 		FruitMap.Add(i, FString::Printf(TEXT("Fruit%d"), i));
 	}
@@ -370,7 +370,7 @@ Note that the slack happens to have caused the new elements to be added in rever
 
 The **Shrink** function also works like the TArray equivalent in that it removes any wasted slack from the end of the container.  However, as TMap allows holes in its data structure, this will only remove slack from holes left at the end of the structure:
 
-	for (int32 i = 0; i != 10; i += 2)
+	for (int32 i = 0; i < 10; i += 2)
 	{
 		FruitMap.Remove(i);
 	}
@@ -399,7 +399,7 @@ The **Shrink** function also works like the TArray equivalent in that it removes
 	// 	{ Key: 1, Value: "Fruit1" }
 	// ]
 
-Note that only one invalid element has been removed from the Shrink call, because there was only one hole at the end.  The **Compact** function can be used to remove all the holes before shrinking:
+Note that only one invalid element has been removed by the Shrink call, because there was only one hole at the end.  The **Compact** function can be used to remove all the holes before shrinking:
 
 	FruitMap.Compact();
 	// FruitMap == [
@@ -428,11 +428,11 @@ As long as a type has an operator== and a non-member GetTypeHash overload then i
 
 KeyFuncs requires the definition of 2 typedefs and 3 static functions:
 
-* **KeyInitType** - used to pass keys around
-* **ElementInitType** - used to pass elements around
-* **KeyInitType GetSetKey(ElementInitType Element)** - returns the key of an element.
-* **bool Matches(KeyInitType A, KeyInitType B)** - returns if A and B are equivalent
-* **uint32 GetKeyHash(KeyInitType Key)** - returns the hash value of Key
+* **KeyInitType** - Used to pass keys around.
+* **ElementInitType** - Used to pass elements around.
+* **KeyInitType GetSetKey(ElementInitType Element)** - Returns the key of an element.
+* **bool Matches(KeyInitType A, KeyInitType B)** - Returns if A and B are equivalent.
+* **uint32 GetKeyHash(KeyInitType Key)** - Returns the hash value of Key.
 
 KeyInitType and ElementInitType are typedefs to the normal passing convention of the key type and element type.  Usually these will be a value for trivial types and a const reference for non-trivial types.  Remember that the element type of a map is a TPair.
 
@@ -531,7 +531,7 @@ A note about providing your own KeyFuncs: be aware that TMap assumes that two it
 
 ## Misc
 
-The **CountBytes** and **GetAllocatedSize** functions are used to estimate how much memory is currently being utilized by the array.  CountBytes takes an FArchive and GetAllocatedSize can be called directly.  They are typically used for stats reporting.
+The **CountBytes** and **GetAllocatedSize** functions are used to estimate how much memory is currently being utilized by the internal array.  CountBytes takes an FArchive and GetAllocatedSize can be called directly.  They are typically used for stats reporting.
 
 The **Dump** function takes an FOutputDevice and writes out some implementation information about the contents of the map.  It's usually used for debugging.
 
diff --git a/Engine/Documentation/Source/Programming/UnrealArchitecture/TSet/TSet.INT.udn b/Engine/Documentation/Source/Programming/UnrealArchitecture/TSet/TSet.INT.udn
new file mode 100644
--- /dev/null
+++ b/Engine/Documentation/Source/Programming/UnrealArchitecture/TSet/TSet.INT.udn
@@ -0,0 +1,252 @@
+Title:TMap
+Crumbs:
+Description:
+Availability:Docs
+Version:4.12
+
+**TSet** is similar to **TMap** and **TMultiMap**, but with an important difference: Rather than associating elements with separately-provided keys, a TSet uses the element itself as the key, via an overridable function that evaluates the element. TSets are very fast (constant time) for adding, finding, and removing elements. By default, TSets do not permit duplicate keys, but this behavior is supported, should it be desired.
+
+## TSet
+
+TSets are a fast container class to store (usually) unique elements in a context where order is irrelevant. For most use cases, just one parameter - the element type - is needed. However, TSet can be set up with different template parameters to change its behavior and make it more versatile. a derived struct based on `DefaultKeyFuncs` can be specified to provide hashing functionality, as well as permitting multiple keys with the same value to exist in the set. Finally, like the other container classes, a custom allocator for the elements can be specified.
+
+Like TArray, TSet is a homogeneous container, meaning that all of its elements are strictly the same type.  TSet is also a value type, and supports the usual copying, assignment, and destructor operations, as well as strong ownership of its elements, which will be destroyed when the set is destroyed. The key type must also be a value type.
+
+TSet uses hashes, which means that the **KeyFuncs** template parameter, if provided, tells the map how to determine the key from an element, how to compare two keys for equality, how to hash the key, and whether or not duplicate keys will be permitted.  These have defaults which will just return a reference to the key, use `operator==` for equality, and the non-member `GetTypeHash` function for hashing.  If your key type supports these functions, it is usable as a set key without the need to provide a custom KeyFuncs. To write a custom KeyFuncs, extend the `DefaultKeyFuncs` struct.
+
+Finally, TSet can take an optional allocator to control the memory allocation behavior.  Standard UE4 allocators (e.g. `FHeapAllocator`, `TInlineAllocator`) cannot be used as allocators for TSet.  It instead uses set allocators, which define how many hash buckets the set should use and which standard UE4 allocators should be used for element storage.  See `TSetAllocator` for more information.
+
+Unlike TArray, the relative order of TSet elements in memory cannot be relied upon, and iterating over the elements is likely to return them in a different order from the order in which they were added. Elements are also unlikely to be laid out contiguously in memory. The backing data structure of a set is a sparse array, which is an array with holes.  As elements are removed from the set, holes in the sparse array will appear. These holes are then filled when future elements are added. However, even though TSet doesn't shuffle elements to fill holes, pointers to set elements may still be invalidated, as the entire storage can be reallocated when it is full and new elements are added.
+
+## Creating and filling a set
+
+Creating a TSet can be done like this:
+
+	TSet<FString> FruitSet;
+
+This creates an empty TSet designed to store unique strings. We have not specified a KeyFuncs or an allocator, so the set will compare FStrings directly, hash them using GetTypeHash, and use the standard heap allocation. No memory has been allocated at this point.
+
+The standard way to populate a set is to use the **Add** function and provide a key (element):
+
+	FruitSet.Add(TEXT("Banana"));
+	FruitSet.Add(TEXT("Grapefruit"));
+	FruitSet.Add(TEXT("Pineapple"));
+	// FruitSet == [ "Banana", "Grapefruit", "Pineapple" ]
+
+While the elements are listed here in the order of insertion, there is no real guarantee as to the order of these elements. For a new set, they are likely to be in order of insertion, but the more insertions and removals the set has been subject to, the more likely it is that new elements will not appear at the end.
+
+Since we used the default allocator, keys are guaranteed to be unique. We can see what happens if we attempt to add a duplicate key:
+
+	FruitSet.Add(TEXT("Pear"));
+	FruitSet.Add(TEXT("Banana"));
+	// FruitSet == [ "Banana", "Grapefruit", "Pineapple", "Pear" ]
+	// Note: Only one banana entry.
+
+The set now contains 4 elements. "Pear" brought the count up to 4 from 3, but the new "Banana" didn't change the number of elements in the set because it replaced the old "Banana" entry.
+
+Like TArray, we can also use **Emplace** instead of Add to avoid the creation of temporaries when inserting into the set:
+
+	FruitSet.Emplace(TEXT("Orange"));
+	// FruitSet == [ "Banana", "Grapefruit", "Pineapple", "Pear", "Orange" ]
+
+Here, the argument is passed directly to the constructor of the key type. This avoids the creation of a temporary FString for the value. Unlike TArray, it's only possible to emplace elements into a set with single-argument constructors.
+
+It's also possible to insert all the elements from another set by using the Append function to merge them:
+
+	TSet<FString> FruitSet2;
+	FruitSet2.Emplace(TEXT("Kiwi"));
+	FruitSet2.Emplace(TEXT("Melon"));
+	FruitSet2.Emplace(TEXT("Mango"));
+	FruitSet2.Emplace(TEXT("Orange"));
+	FruitSet.Append(FruitSet2);
+	// FruitSet == [ "Banana", "Grapefruit", "Pineapple", "Pear", "Orange", "Kiwi", "Melon", "Mango" ]
+
+Here, the resulting set is equivalent to using Add/Emplace to add them individually, so that duplicate keys from the source set replace ones in the target.
+
+## Iteration
+
+Iteration over TSets is similar to TArrays.  You can use C++'s ranged-for feature:
+
+	for (auto& Elem : FruitSet)
+	{
+		FPlatformMisc::LocalPrint(
+			*FString::Printf(
+				TEXT(" \"%s\"\n"),
+				*Elem
+			)
+		);
+	}
+	// Output:
+	// 	"Banana"
+	// 	"Grapefruit"
+	// 	"Pineapple"
+	// 	"Pear"
+	// 	"Orange"
+	// 	"Kiwi"
+	// 	"Melon"
+	// 	"Mango"
+
+Sets also provide their own iterator type for more direct control over iteration.  The **CreateIterator** function provides read-write access to the elements and the **CreateConstIterator** function provides read-only access. The iterator objects themselves are of the element type that was specified as the first template argument in the declaration of the TSet.
+
+## Queries
+
+We can ask the set how many elements it holds by using the **Num** function:
+
+	int32 Count = FruitSet.Num();
+	// Count == 8
+
+We can use the `FSetElementId` struct to find the index of a key within the set. This enables us to look at that element via the indexing operator[]. Calling operator[] on a non-const set will return a non-const reference and calling it on a const set will return a const reference.
+
+	FSetElementId BananaIndex = FruitSet.Index(TEXT("Banana"));
+	// BananaIndex is a value between 0 and (FruitSet.Num() - 1)
+	FPlatformMisc::LocalPrint(
+		*FString::Printf(
+			TEXT(" \"%s\"\n"),
+			*FruitSet[BananaIndex]
+		)
+	);
+	// Prints "Banana"
+	
+	FSetElementId LemonIndex = FruitSet.Index(TEXT("Lemon"));
+	// LemonIndex is INDEX_NONE (-1)
+	FPlatformMisc::LocalPrint(
+		*FString::Printf(
+			TEXT(" \"%s\"\n"),
+			*FruitSet[LemonIndex]
+		)
+	); // assert
+
+We can check if a particular key exists in the map using the **Contains** function:
+
+	bool bHasBanana = FruitSet.Contains(TEXT("Banana"));
+	bool bHasLemon = FruitSet.Contains(TEXT("Lemon"));
+	// bHasBanana == true
+	// bHasLemon == false
+
+Most of the time, you will want to look up elements without knowing whether or not the set actually contains that element. Using Contains followed by operator[] means a double lookup of the key, which we don't really want to do.  The **Find** function will allow you to do a single lookup, returning a pointer to the value of the found element rather than a reference, and returning null when the key does not exist:
+
+	FString* PtrBanana = FruitSet.Find(TEXT("Banana"));
+	FString* PtrLemon = FruitSet.Find(TEXT("Lemon"));
+	// *PtrBanana == "Banana"
+	//  PtrLemon == nullptr
+
+If called on a const map, the returned pointer will also be const.
+
+The **Array** function returns a TArray populated with a copy of all the elements in the TSet.  In both cases, the array being passed is emptied before population, so the resulting number of elements will always equal the number of elements in the map:
+
+	TArray<FString> FruitArray = FruitSet.Array();
+	// FruitArray == [ "Banana","Grapefruit","Pineapple","Pear","Orange","Kiwi","Melon","Mango" ] (order may vary)
+
+## Removal
+
+Elements can be removed by key with the **Remove** function:
+
+	FruitSet.Remove(TEXT("Banana");
+	// FruitSet == [ "Grapefruit","Pineapple","Pear","Orange","Kiwi","Melon","Mango" ]
+
+Removing elements actually leaves holes in the data structure, which you can see when visualizing the set in Visual Studio's watch window, but they have been omitted here for clarity. If a TSet supports duplicate keys, **Remove** will remove all keys that match the input parameter. The Remove function returns the number of elements removed, and will be 0 if the key provided was not contained in the set.
+
+	int32 RemovedAmountPineapple = FruitSet.FindAndRemoveChecked(TEXT("Pineapple"));
+	// RemovedAmountPineapple == 1
+	// FruitSet == [ "Grapefruit","Pear","Orange","Kiwi","Melon","Mango" ]
+	FString RemovedAmountLemon = FruitSet.FindAndRemoveChecked("Lemon"); // assert!
+	// RemovedAmountLemon == 0
+
+Finally, all elements can be removed by using the **Empty** function:
+
+	TSet<FString> FruitSetCopy = FruitSet;
+	// FruitSetCopy == [ "Grapefruit","Pear","Orange","Kiwi","Melon","Mango" ]
+	
+	FruitSetCopy.Empty();
+	// FruitSetCopy == []
+
+Like TArray, Empty takes an optional slack value which can be used to optimize the case when you are about to repopulate the map with a given number of elements.
+
+## Sorting
+
+TSets can temporarily be sorted with the **Sort** function.  The next iteration over the set will present the elements in sorted order, though future modifications to the set will likely result in the set being unsorted again.  Sorting is unstable, so equivalent elements (in a TSet that permits duplicate keys) may appear in any order.
+
+The Sort function takes a binary predicate which specifies the sort order:
+
+	FruitSet.Sort([](const FString& A, const FString& B) {
+		return A > B; // sort by reverse-alphabetical order
+	});
+	// FruitSet == [ "Pear", "Orange", "Melon", "Mango", "Kiwi", "Grapefruit" ] (order is temporarily guaranteed)
+	
+	FruitSet.Sort([](const FString& A, const FString& B) {
+		return A.Len() < B.Len(); // sort strings by length, shortest to longest
+	});
+	// FruitSet == [ "Pear", "Kiwi", "Melon", "Mango", "Orange", "Grapefruit" ] (order is temporarily guaranteed)
+	
+## Operators
+
+Like TArray, TSet is a regular value type and as such can be copied via the standard copy constructor or assignment operator. As sets strictly own their elements, copying a set is deep and so the new set will have its own copy of the elements:
+
+	TSet<int32, FString> NewSet = FruitSet;
+	NewSet.Add(TEXT("Apple"));
+	NewSet.Remove(TEXT("Pear"));
+	// FruitSet == [ "Pear", "Kiwi", "Melon", "Mango", "Orange", "Grapefruit" ]
+	// NewSet == [ "Kiwi", "Melon", "Mango", "Orange", "Grapefruit", "Apple" ]
+
+
+## Slack
+
+TSets also have a notion of **slack** and this can be used to optimize the population of the set.  The **Reset** function acts like an Empty() call, but does not free the memory that was previously used by the elements:
+
+	FruitSet.Reset();
+	// FruitSet == [ <invalid>, <invalid>, <invalid>, <invalid>, <invalid>, <invalid> ]
+
+Here the set is emptied in the same way as Empty would, but the memory used for storage is not freed and remains as slack.
+
+TSet does not provide a way of checking how many elements are preallocated, like TArray::Max() does, but still supports preallocating slack.  The Reserve function can be used to preallocate slack for a particular number of elements before adding:
+
+	FruitSet.Reserve(10);
+	for (int32 i = 0; i < 10; ++i)
+	{
+		FruitSet.Add(FString::Printf(TEXT("Fruit%d"), i));
+	}
+	// FruitSet == [ "Fruit9", "Fruit8", "Fruit7" ... "Fruit2", "Fruit1", "Fruit0" ]
+
+Note that the slack happens to have caused the new elements to be added in reverse order. This is an example of why element order in maps should not be relied upon in any way.
+
+The **Shrink** function also works like the TArray equivalent in that it removes any wasted slack from the end of the container.  However, as TSet allows holes in its data structure, this will only remove slack from holes left at the end of the structure:
+
+	// Remove every other element from the set.
+	for (int32 i = 0; i < 10; i += 2)
+	{
+		FruitSet.Remove(FSetElementId::FromInteger(i));
+	}
+	// FruitSet == ["Fruit8", <invalid>, "Fruit6", <invalid>, "Fruit4", <invalid>, "Fruit2", <invalid>, "Fruit0", <invalid> ]
+
+	FruitSet.Shrink();
+	// FruitSet == ["Fruit8", <invalid>, "Fruit6", <invalid>, "Fruit4", <invalid>, "Fruit2", <invalid>, "Fruit0" ]
+
+Note that only one invalid element has been removed by the Shrink call, because there was only one hole at the end.  The **Compact** function can be used to remove all the holes before shrinking:
+
+	FruitSet.Compact();
+	// FruitSet == ["Fruit8", "Fruit6", "Fruit4", "Fruit2", "Fruit0", <invalid>, <invalid>, <invalid>, <invalid> ]
+	FruitSet.Shrink();
+	// FruitSet == ["Fruit8", "Fruit6", "Fruit4", "Fruit2", "Fruit0" ]
+
+## DefaultKeyFuncs
+
+As long as a type has an operator== and a non-member GetTypeHash overload then it can be used with TSet without any changes, as the type is both the element and the key. However, it may be useful to use types as keys where it is undesirable to overload those functions.  In these cases, you can provide your own custom **DefaultKeyFuncs**.
+
+KeyFuncs requires the definition of 2 typedefs and 3 static functions:
+
+* **KeyInitType** - Used to pass keys around. Usually drawn from ElementType template parameter.
+* **ElementInitType** - Used to pass elements around. Also usually drawn from ElementType template parameter, and therefore identical to KeyInitType.
+* **KeyInitType GetSetKey(ElementInitType Element)** - Returns the key of an element, which is generally just the element itself.
+* **bool Matches(KeyInitType A, KeyInitType B)** - Returns if A and B are equivalent.
+* **uint32 GetKeyHash(KeyInitType Key)** - Returns the hash value of Key, usually by calling the external GetTypeHash function.
+
+KeyInitType/ElementInitType are typedefs to the normal passing convention of the key/element type.  Usually these will be a value for trivial types and a const reference for non-trivial types.  Remember that the element type of a set is also the key type, which is why DefaultKeyFuncs uses only one template parameter, ElementType, to define both.
+
+A note about providing your own DefaultKeyFuncs: be aware that TSet assumes that two items that compare equal using DefaultKeyFuncs::Matches return the same value from KeyFuncs::GetKeyHash.  In addition, modifying the key of an existing element in a way which will change the results from either of these functions is considered undefined behavior, as this will invalidate TSet's internal hash.  These rules also apply to the overloads of operator== and GetKeyHash when using the default implementation of DefaultKeyFuncs.
+
+## Misc
+
+The **CountBytes** and **GetAllocatedSize** functions are used to estimate how much memory is currently being utilized by the internal array.  CountBytes takes an FArchive and GetAllocatedSize can be called directly.  They are typically used for stats reporting.
+
+The **Dump** function takes an FOutputDevice and writes out some implementation information about the contents of the set.  It's usually used for debugging.
+
