Submitted too quickly, fixing.

diff --git a/Engine/Documentation/Source/Programming/UnrealArchitecture/TSet/TSet.INT.udn b/Engine/Documentation/Source/Programming/UnrealArchitecture/TSet/TSet.INT.udn
--- a/Engine/Documentation/Source/Programming/UnrealArchitecture/TSet/TSet.INT.udn
+++ b/Engine/Documentation/Source/Programming/UnrealArchitecture/TSet/TSet.INT.udn
@@ -1,4 +1,4 @@
-Title:TMap
+Title:TSet
 Crumbs:
 Description:
 Availability:Docs
@@ -12,7 +12,7 @@ TSets are a fast container class to store (usually) unique elements in a context
 
 Like TArray, TSet is a homogeneous container, meaning that all of its elements are strictly the same type.  TSet is also a value type, and supports the usual copying, assignment, and destructor operations, as well as strong ownership of its elements, which will be destroyed when the set is destroyed. The key type must also be a value type.
 
-TSet uses hashes, which means that the **KeyFuncs** template parameter, if provided, tells the map how to determine the key from an element, how to compare two keys for equality, how to hash the key, and whether or not duplicate keys will be permitted.  These have defaults which will just return a reference to the key, use `operator==` for equality, and the non-member `GetTypeHash` function for hashing.  If your key type supports these functions, it is usable as a set key without the need to provide a custom KeyFuncs. To write a custom KeyFuncs, extend the `DefaultKeyFuncs` struct.
+TSet uses hashes, which means that the **KeyFuncs** template parameter, if provided, tells the set how to determine the key from an element, how to compare two keys for equality, how to hash the key, and whether or not duplicate keys will be permitted.  These have defaults which will just return a reference to the key, use `operator==` for equality, and the non-member `GetTypeHash` function for hashing.  If your key type supports these functions, it is usable as a set key without the need to provide a custom KeyFuncs. To write a custom KeyFuncs, extend the `DefaultKeyFuncs` struct.
 
 Finally, TSet can take an optional allocator to control the memory allocation behavior.  Standard UE4 allocators (e.g. `FHeapAllocator`, `TInlineAllocator`) cannot be used as allocators for TSet.  It instead uses set allocators, which define how many hash buckets the set should use and which standard UE4 allocators should be used for element storage.  See `TSetAllocator` for more information.
 
@@ -116,7 +116,7 @@ We can use the `FSetElementId` struct to find the index of a key within the set.
 		)
 	); // assert
 
-We can check if a particular key exists in the map using the **Contains** function:
+We can check if a particular key exists in the set using the **Contains** function:
 
 	bool bHasBanana = FruitSet.Contains(TEXT("Banana"));
 	bool bHasLemon = FruitSet.Contains(TEXT("Lemon"));
@@ -130,9 +130,9 @@ Most of the time, you will want to look up elements without knowing whether or n
 	// *PtrBanana == "Banana"
 	//  PtrLemon == nullptr
 
-If called on a const map, the returned pointer will also be const.
+If called on a const set, the returned pointer will also be const.
 
-The **Array** function returns a TArray populated with a copy of all the elements in the TSet.  In both cases, the array being passed is emptied before population, so the resulting number of elements will always equal the number of elements in the map:
+The **Array** function returns a TArray populated with a copy of all the elements in the TSet.  In both cases, the array being passed is emptied before population, so the resulting number of elements will always equal the number of elements in the set:
 
 	TArray<FString> FruitArray = FruitSet.Array();
 	// FruitArray == [ "Banana","Grapefruit","Pineapple","Pear","Orange","Kiwi","Melon","Mango" ] (order may vary)
@@ -160,7 +160,7 @@ Finally, all elements can be removed by using the **Empty** function:
 	FruitSetCopy.Empty();
 	// FruitSetCopy == []
 
-Like TArray, Empty takes an optional slack value which can be used to optimize the case when you are about to repopulate the map with a given number of elements.
+Like TArray, Empty takes an optional slack value which can be used to optimize the case when you are about to repopulate the set with a given number of elements.
 
 ## Sorting
 
@@ -207,7 +207,7 @@ TSet does not provide a way of checking how many elements are preallocated, like
 	}
 	// FruitSet == [ "Fruit9", "Fruit8", "Fruit7" ... "Fruit2", "Fruit1", "Fruit0" ]
 
-Note that the slack happens to have caused the new elements to be added in reverse order. This is an example of why element order in maps should not be relied upon in any way.
+Note that the slack happens to have caused the new elements to be added in reverse order. This is an example of why element order in sets should not be relied upon in any way.
 
 The **Shrink** function also works like the TArray equivalent in that it removes any wasted slack from the end of the container.  However, as TSet allows holes in its data structure, this will only remove slack from holes left at the end of the structure:
 
