Initial checkin of Pose Watching feature.

diff --git a/Engine/Documentation/Source/Engine/Animation/NodeReference/AnimationNodes.INT.udn b/Engine/Documentation/Source/Engine/Animation/NodeReference/AnimationNodes.INT.udn
--- a/Engine/Documentation/Source/Engine/Animation/NodeReference/AnimationNodes.INT.udn
+++ b/Engine/Documentation/Source/Engine/Animation/NodeReference/AnimationNodes.INT.udn
@@ -18,6 +18,7 @@ space nodes, and conversion nodes.
 	EventGraph Nodes
 	[/PARAM]
 	[PARAM:description]
+	Nodes in this graph are used to process incoming data that in turn is used to drive the pose data in the **AnimGraph**. 
 	[/PARAM]
 	[PARAM:links]
 	* [](Engine/Animation/NodeReference/Events "%Engine/Animation/NodeReference/Events:description%")
@@ -31,6 +32,7 @@ space nodes, and conversion nodes.
 	AnimGraph Nodes
 	[/PARAM]
 	[PARAM:description]
+	Using data from the **EventGraph**, nodes in this graph determine the actual frame-by-frame pose to use for the character. 
 	[/PARAM]
 	[PARAM:links]
 	* [](Engine/Animation/NodeReference/Blend "%Engine/Animation/NodeReference/Blend:description%")
@@ -40,6 +42,31 @@ space nodes, and conversion nodes.
 	[/PARAM]
 [/OBJECT]
 
+## Animation Node Pose Watching
+
+Animation Node **Pose Watching** allows you to see a representation of the pose being generated at any point in the **AnimGraph** dynamically. 
+You can also have multiple Pose Watches active at once allowing you to compare poses at different points to find the exact point in which any errors in your current pose are introduced. 
+
+To enable this, inside your **AnimGraph** on a node that contains a pose data, **Right-click** on the node and select **Toggle Pose Watch**. 
+
+[OBJECT:EmbeddedVideo]
+[PARAMLITERAL:width]
+640
+[/PARAMLITERAL]
+[PARAMLITERAL:height]
+360
+[/PARAMLITERAL]
+[PARAMLITERAL:videoid]
+nabapYAeE7c
+[/PARAMLITERAL]
+[/OBJECT]
+
+You can also specify different colors for each of your watched poses by clicking the icon in the upper left corner of the node. 
+
+![](PoseWatchA.png) ![](PoseWatchB.png)
+
+You can **Remove Pose Watch** from the color picker window or by reselecting **Toggle Pose Watch** on the Animation Node. 
+
 ## Pins and Properties
 
 Animation nodes can have pins (both input and output) which values can be passed to in the **EventGraph** 
