#UE4 Docs: Checking in this new doc and its images.

diff --git a/Engine/Documentation/Source/Engine/Rendering/LightingAndShadows/LightingAndShadows.INT.udn b/Engine/Documentation/Source/Engine/Rendering/LightingAndShadows/LightingAndShadows.INT.udn
--- a/Engine/Documentation/Source/Engine/Rendering/LightingAndShadows/LightingAndShadows.INT.udn
+++ b/Engine/Documentation/Source/Engine/Rendering/LightingAndShadows/LightingAndShadows.INT.udn
@@ -75,6 +75,7 @@ Version: 4.9
 %Engine/Rendering/LightingAndShadows/LightPropagationVolumes:Topic%
 %Engine/Rendering/LightingAndShadows/DistanceFieldAmbientOcclusion:Topic%
 %Engine/Rendering/LightingAndShadows/RayTracedDistanceFieldShadowing:Topic%
+%Engine/Rendering/LightingAndShadows/ShadowMapCaching:Topic%
 [/REGION]
 
 ## General
diff --git a/Engine/Documentation/Source/Engine/Rendering/LightingAndShadows/ShadowMapCaching/ShadowMapCachingForMovableLights.INT.udn b/Engine/Documentation/Source/Engine/Rendering/LightingAndShadows/ShadowMapCaching/ShadowMapCachingForMovableLights.INT.udn
new file mode 100644
--- /dev/null
+++ b/Engine/Documentation/Source/Engine/Rendering/LightingAndShadows/ShadowMapCaching/ShadowMapCachingForMovableLights.INT.udn
@@ -0,0 +1,122 @@
+Availability:Public
+Title: Shadow Map Caching for Movable Lights
+Crumbs: %ROOT%, Engine, Engine/Rendering/LightingAndShadows
+Description:System for caching movable point and spotlight dynamic shadow maps.
+Version: 4.13
+
+[VAR:Topic]
+[OBJECT:Topic]
+	[PARAM:image]
+		![%Engine/Rendering/LightingAndShadows/ShadowMapCaching:title%](SMC_Header.png)
+	[/PARAM]
+	[PARAM:title]
+		%Engine/Rendering/LightingAndShadows/ShadowMapCaching:title%
+	[/PARAM]
+	[PARAM:description]
+		%Engine/Rendering/LightingAndShadows/ShadowMapCaching:description%
+	[/PARAM]
+	[PARAM:path]
+		[RELATIVE:Engine/Rendering/LightingAndShadows/ShadowMapCaching]
+	[/PARAM]
+[/OBJECT]
+[/VAR]
+[VAR:TopicCompact]
+[OBJECT:TopicCompact]
+	[PARAM:image]
+		![%Engine/Rendering/LightingAndShadows/ShadowMapCaching:title%](SMC_Header.png)
+	[/PARAM]
+	[PARAM:title]
+		%Engine/Rendering/LightingAndShadows/ShadowMapCaching:title%
+	[/PARAM]
+	[PARAM:description]
+		%Engine/Rendering/LightingAndShadows/ShadowMapCaching:description%
+	[/PARAM]
+	[PARAM:path]
+		[RELATIVE:Engine/Rendering/LightingAndShadows/ShadowMapCaching]
+	[/PARAM]
+[/OBJECT]
+[/VAR]
+
+[REGION:fullwidth raw]
+![](SMC_Header.png)
+[/REGION]
+
+[TOC(start:2 end:2)]
+
+When a point or spot light is not moving, we can store off the shadow map for that light and reuse it next frame. This makes shadow casting movable Point and Spot lights much more affordable in games where the environment is often not moving. In the following document we will take a look at how you can use this feature in your Unreal Engine 4 (UE4) projects.
+
+## Enabling Shadow Map Caching 
+Shadow Map Caching should be automatically be enabled in your UE4 project. To check and see if Shadow Map Caching is enabled you can do the following:
+
+1. Select all of the lights in your project's level that you want to cast dynamic shadows.
+
+	![](SMC_Select_All_Lights.png)
+
+1. Set the lights **Mobility**  to **Movable** and make sure that **Cast Shadows** is enabled.
+
+	![](SMC_Set_Mob_Shadow_Cast.png)
+
+1. Press the **Backtick (`)** key to open up the Unreal Console and type **Stat Shadowrendering** to see the current cost of dynamic shadows. 
+
+	![](SMC_Console_ShadowRendering.png)
+
+1. Press the Backtick (`) key again to open the Console and type **r.Shadow.CacheWholeSceneShadows 0** to disable dynamic shadow caching.
+
+	![](SMC_Disable_Shadow_Cache.png)
+	
+	[region:note]
+	 Make sure to note the the numbers in the **CallCount** and **InclusiveAug** sections.
+	[/region]
+
+1. Now press the Backtick (`) key one more time to open the Console and type **r.Shadow.CacheWholeSceneShadows 1** to re-enable shadow caching. If you compare the Callout and InclusiveAug numbers you can see the impact this feature has on dynamic shadow performance.
+
+	[OBJECT:ComparisonSlider]
+		[PARAM:before]
+		![Shadow Caching On](SMC_On.png)
+		[/PARAM]
+		[PARAM:after]
+		![Shadow Caching Off](SMC_Off.png)
+		[/PARAM]
+	[/OBJECT]
+
+
+## Performance
+The impact that this features has on performance is quite significant. Here are the results of testing this feature using an NVIDIA 970 GTX GPU at a screen resolution of 1920x1200.
+
+* Before enabling this feature three shadow casting point lights without caching took 14.89ms to render Shadow Depths.
+
+* When Cached Shadow Maps where enabled the same three shadow casting lights took .9ms to render which is about **16 times faster**!
+	[region:note]
+	Note that it still costs 2 ms to render the contributions of the 33 point lights, which can be optimized in other ways but is not affected by this change.
+	[/region]
+* You can control the maximum amount memory used by the Shadow Map Cache using **r.Shadow.WholeSceneShadowCacheMb**
+
+
+
+## Limitations
+While enabling Shadow Map Caching can decrease the cost of using dynamic shadows in your UE4 projects, there are some limitations that could cause rendering artifacts if used with non-supported features. In the following section we will take a look at what the limitations of Shadow Map Casting are and what you can do about them.
+
+
+* By default, caching can only happen when an object meets the following requirements:
+
+	* Primitives have mobility their **Mobility** set to **Static** or **Stationary**. 
+	* Materials used in the level do not use **World Position Offset**.
+	* Light need to be either a **Point** or **Spot** light, have its **Mobility** set to **Movable** and have **Shadow Casting** enabled.
+	* Lights have to remain stationary and can not move at all.
+	* Materials that use animated **Tessellation** or **Pixel Depth Offset** can cause artifacts as their shadow depths are cached.
+
+
+
+
+
+
+
+
+
+
+
+
+
+
+
+
