#UE4 Docs: Fixing a few grammer errors.

diff --git a/Engine/Documentation/Source/Programming/UnrealArchitecture/TSet/TSet.INT.udn b/Engine/Documentation/Source/Programming/UnrealArchitecture/TSet/TSet.INT.udn
--- a/Engine/Documentation/Source/Programming/UnrealArchitecture/TSet/TSet.INT.udn
+++ b/Engine/Documentation/Source/Programming/UnrealArchitecture/TSet/TSet.INT.udn
@@ -8,11 +8,11 @@ Version:4.12
 
 ## TSet
 
-TSets are a fast container class to store (usually) unique elements in a context where order is irrelevant. For most use cases, just one parameter - the element type - is needed. However, TSet can be set up with different template parameters to change its behavior and make it more versatile. a derived struct based on `DefaultKeyFuncs` can be specified to provide hashing functionality, as well as permitting multiple keys with the same value to exist in the set. Finally, like the other container classes, a custom allocator for the elements can be specified.
+TSets are a fast container class to store (usually) unique elements in a context where the order is irrelevant. For most use cases, just one parameter - the element type - is needed. However, TSet can be set up with different template parameters to change its behavior and make it more versatile. a derived struct based on `DefaultKeyFuncs` can be specified to provide hashing functionality, as well as permitting multiple keys with the same value to exist in the set. Finally, like the other container classes, a custom allocator for the elements can be specified.
 
-Like TArray, TSet is a homogeneous container, meaning that all of its elements are strictly the same type.  TSet is also a value type, and supports the usual copying, assignment, and destructor operations, as well as strong ownership of its elements, which will be destroyed when the set is destroyed. The key type must also be a value type.
+Like TArray, TSet is a homogeneous container, meaning that all of its elements are strictly the same type.  TSet is also a value type and supports the usual copying, assignment, and destructor operations, as well as strong ownership of its elements, which will be destroyed when the set is destroyed. The key type must also be a value type.
 
-TSet uses hashes, which means that the **KeyFuncs** template parameter, if provided, tells the set how to determine the key from an element, how to compare two keys for equality, how to hash the key, and whether or not duplicate keys will be permitted.  These have defaults which will just return a reference to the key, use `operator==` for equality, and the non-member `GetTypeHash` function for hashing.  If your key type supports these functions, it is usable as a set key without the need to provide a custom KeyFuncs. To write a custom KeyFuncs, extend the `DefaultKeyFuncs` struct.
+TSet uses hashes, which means that the **KeyFuncs** template parameter if provided, tells the set how to determine the key from an element, how to compare two keys for equality, how to hash the key, and whether or not duplicate keys will be permitted.  These have defaults which will just return a reference to the key, use `operator==` for equality, and the non-member `GetTypeHash` function for hashing.  If your key type supports these functions, it is usable as a set key without the need to provide a custom KeyFuncs. To write a custom KeyFuncs, extend the `DefaultKeyFuncs` struct.
 
 Finally, TSet can take an optional allocator to control the memory allocation behavior.  Standard UE4 allocators (e.g. `FHeapAllocator`, `TInlineAllocator`) cannot be used as allocators for TSet.  It instead uses set allocators, which define how many hash buckets the set should use and which standard UE4 allocators should be used for element storage.  See `TSetAllocator` for more information.
 
@@ -33,7 +33,7 @@ The standard way to populate a set is to use the **Add** function and provide a
 	FruitSet.Add(TEXT("Pineapple"));
 	// FruitSet == [ "Banana", "Grapefruit", "Pineapple" ]
 
-While the elements are listed here in the order of insertion, there is no real guarantee as to the order of these elements. For a new set, they are likely to be in order of insertion, but the more insertions and removals the set has been subject to, the more likely it is that new elements will not appear at the end.
+While the elements are listed here in the order of insertion, there is no real guarantee as to the order of these elements. For a new set, they are likely to be in the order of insertion, but the more insertions and removals the set has been subject to, the more likely it is that new elements will not appear at the end.
 
 Since we used the default allocator, keys are guaranteed to be unique. We can see what happens if we attempt to add a duplicate key:
 
@@ -144,7 +144,7 @@ Elements can be removed by key with the **Remove** function:
 	FruitSet.Remove(TEXT("Banana");
 	// FruitSet == [ "Grapefruit","Pineapple","Pear","Orange","Kiwi","Melon","Mango" ]
 
-Removing elements actually leaves holes in the data structure, which you can see when visualizing the set in Visual Studio's watch window, but they have been omitted here for clarity. If a TSet supports duplicate keys, **Remove** will remove all keys that match the input parameter. The Remove function returns the number of elements removed, and will be 0 if the key provided was not contained in the set.
+Removing elements will actually leave holes in the data structure, which you can see when visualizing the set in Visual Studio's watch window, but they have been omitted here for clarity. If a TSet supports duplicate keys, **Remove** will remove all keys that match the input parameter. The Remove function returns the number of elements removed, and will be 0 if the key provided was not contained in the set.
 
 	int32 RemovedAmountPineapple = FruitSet.FindAndRemoveChecked(TEXT("Pineapple"));
 	// RemovedAmountPineapple == 1
@@ -240,7 +240,7 @@ KeyFuncs requires the definition of 2 typedefs and 3 static functions:
 * **bool Matches(KeyInitType A, KeyInitType B)** - Returns if A and B are equivalent.
 * **uint32 GetKeyHash(KeyInitType Key)** - Returns the hash value of Key, usually by calling the external GetTypeHash function.
 
-KeyInitType/ElementInitType are typedefs to the normal passing convention of the key/element type.  Usually these will be a value for trivial types and a const reference for non-trivial types.  Remember that the element type of a set is also the key type, which is why DefaultKeyFuncs uses only one template parameter, ElementType, to define both.
+KeyInitType/ElementInitType are typedefs to the normal passing convention of the key/element type.  Usually, these will be a value for trivial types and a const reference for non-trivial types.  Remember that the element type of a set is also the key type, which is why DefaultKeyFuncs uses only one template parameter, ElementType, to define both.
 
 A note about providing your own DefaultKeyFuncs: be aware that TSet assumes that two items that compare equal using DefaultKeyFuncs::Matches return the same value from KeyFuncs::GetKeyHash.  In addition, modifying the key of an existing element in a way which will change the results from either of these functions is considered undefined behavior, as this will invalidate TSet's internal hash.  These rules also apply to the overloads of operator== and GetKeyHash when using the default implementation of DefaultKeyFuncs.
 
