initial check-in for Custom Render Passes in Sequencer (HDR Export, Capture Camut, etc.)

diff --git a/Engine/Documentation/Source/Engine/Sequencer/Overview/SequencerOverview.INT.udn b/Engine/Documentation/Source/Engine/Sequencer/Overview/SequencerOverview.INT.udn
--- a/Engine/Documentation/Source/Engine/Sequencer/Overview/SequencerOverview.INT.udn
+++ b/Engine/Documentation/Source/Engine/Sequencer/Overview/SequencerOverview.INT.udn
@@ -256,6 +256,17 @@ Please see [](Engine/Sequencer/HowTo/ImportExportEDLs) for a step-by-step guide
 [/REGION]
 [/PUBLISH:docs]
 
+[PUBLISH:docs]
+### Custom Render Passes
+
+If you are looking to export your cinematic in different render passes, you can do so from the **Render Movie Settings** window. 
+This will allow you to specify the render pass(es) to use when exporting your sequence. You can also export HDR data as .exr files and define compression and color gamut settings. 
+
+
+[REGION:note]
+Please see [](Engine\Sequencer\Workflow\CustomRenderPass) for more information.
+[/REGION]
+[/PUBLISH:docs]
 
 
 ## Workflow Considerations
diff --git a/Engine/Documentation/Source/Engine/Sequencer/Workflow/CustomRenderPass/CustomRenderPass.INT.udn b/Engine/Documentation/Source/Engine/Sequencer/Workflow/CustomRenderPass/CustomRenderPass.INT.udn
new file mode 100644
--- /dev/null
+++ b/Engine/Documentation/Source/Engine/Sequencer/Workflow/CustomRenderPass/CustomRenderPass.INT.udn
@@ -0,0 +1,72 @@
+Availability:Docs
+Title: Exporting Custom Render Passes
+Crumbs: %ROOT%, Engine
+Description:Describes how you can export your Sequence in Custom Render Passes.
+Navigation:topic
+parent:Engine/Sequencer/Workflow
+order:1
+Version: 4.13
+Related: Engine\Sequencer\Overview
+Related: Engine/Sequencer/HowTo
+
+This page describes how you can render out custom passes from the **Render Movie Settings** window inside **Sequencer**. 
+Using a Custom Render Pass will output your sequence into a series of image files based on the Render Pass you specify. 
+By default the resulting files will be exported into the .png file format, however you can capture frames in HDR to export them in the .exr file format as well. 
+
+[TOC(start:2 end:2)]
+
+## Custom Render Passes
+
+To render out custom passes for your cinematic refer to the following steps. 
+
+1. Inside **Sequencer** click the **Render Movie Settings** button. 
+
+	![](RenderMovie1.png)
+
+1. In the **Render Movie Settings** window, change the **Output Format** to **Custom Render Passes**. 
+
+	![](RenderMovie2.png)
+
+1. Select the desired Render Pass(es) from the **Include Render Passes** drop-down.
+
+	![](RenderMovie3.png)
+
+	You can include more than one pass as part of the export process. 
+
+	![](RenderMovie4.png)
+
+	Click the **Minus** button next to each render pass to remove it. 
+
+	[REGION:note]
+	If you do not specify which render pass(es) to include, all available passes will be included as part of the export. Depending upon the number of frames in your sequence, this may result in hundreds of files being generated for each available render pass.
+	[/REGION]
+
+1. Specify the **Output Directory** and **Filename Format**, then click **Capture Movie**. 
+
+	![](RenderMovie5.png)
+
+	If you use the default **Filename Format**, your files will automatically be named based on the render pass it reflects.
+
+Once the render movie process completes, in your output directory there will be an image file for each frame of your sequence that corresponds to the selected render pass(es). 
+If you selected more than one render pass, you will have a separate set of image files for each pass. 
+
+
+## HDR Exporting
+
+**Sequencer** also provides the ability to export HDR data in [OpenEXR](http://www.openexr.com/) files when performing a custom render pass. 
+In addition, you can also specify the **HDR Compression Quality** (0 for no compression, 1 for the default compression) and the color **Capture Gamut** used when storing the captured HDR data. 
+The choice of Capture Gamut will largely depend on the project that will make use of the exported EXRs (see the table below and click each **Gamut Type** for external reference). 
+
+| Gamut Type | Description |
+| ---- | ----------- |
+| [**Rec.709 / sRGB**](https://en.wikipedia.org/wiki/Rec._709#Primary_chromaticities) | This is the most widely used color gamut for images that will be displayed on monitors, TVs and Projectors. |
+| [**P3 D65**](https://en.wikipedia.org/wiki/DCI-P3#System_colorimetry) | This is used widely for cinematic projection. |
+| [**Rec.2020**](https://en.wikipedia.org/wiki/Rec._2020#System_colorimetry), [**ACES**](http://www.oscars.org/science-technology/aces/aces-documentation), [**ACEScg**](http://duikerresearch.com/2015/09/acescg-a-common-color-encoding-for-visual-effects-applications/) | These gamuts are being used more with the advent of UHD TVs, HDR TVs and other advances in display technologies. ACEScg in particular now plays a prominent role in rendering and compositing for visual effects and is probably the most future-proof of gamuts.  |
+
+To export HDR frame data, perform a custom render pass and choose the **Capture Frames in HDR** option. 
+
+![](RenderMovie6.png)
+
+Then define the **HDR Compression Quality** and **Capture Gamut** settings as desired. 
+
+![](RenderMovie7.png)
diff --git a/Engine/Documentation/Source/Engine/Sequencer/Workflow/Workflow.INT.udn b/Engine/Documentation/Source/Engine/Sequencer/Workflow/Workflow.INT.udn
--- a/Engine/Documentation/Source/Engine/Sequencer/Workflow/Workflow.INT.udn
+++ b/Engine/Documentation/Source/Engine/Sequencer/Workflow/Workflow.INT.udn
@@ -38,6 +38,7 @@ If you are new to Sequencer it is recommended that you first check out the [**Ov
 	[PARAM:links]
 		* [](Engine\Sequencer\Workflow\CheatSheet "%Engine\Sequencer\Workflow\CheatSheet:description%")
 		* [](Engine\Sequencer\Workflow\ImportExport "%Engine\Sequencer\Workflow\ImportExport:description%")
+		* [](Engine\Sequencer\Workflow\CustomRenderPass "%Engine\Sequencer\Workflow\CustomRenderPass:description%")
 		* [](Engine\Sequencer\Workflow\ToolUsage "%Engine\Sequencer\Workflow\ToolUsage:description%")
 	[/PARAM]
 [/OBJECT]
