Minor adjustments to the documents.

diff --git a/Engine/Documentation/Source/Engine/Rendering/RenderTargets/BlueprintRenderTargets/BlueprintRenderTargets.INT.udn b/Engine/Documentation/Source/Engine/Rendering/RenderTargets/BlueprintRenderTargets/BlueprintRenderTargets.INT.udn
--- a/Engine/Documentation/Source/Engine/Rendering/RenderTargets/BlueprintRenderTargets/BlueprintRenderTargets.INT.udn
+++ b/Engine/Documentation/Source/Engine/Rendering/RenderTargets/BlueprintRenderTargets/BlueprintRenderTargets.INT.udn
@@ -11,7 +11,7 @@ Navigation:topic
 
 ![](Engine\Rendering\RenderTargets\BlueprintRenderTarget.png)
 
-Having the ability to use Blueprints to draw Materials into Render Targets and then use or save those results for later use can enable a huge variety of game-specific rendering effects. On the following page you will find several links to documentation pages that cover how to use Blueprints and Render Targets. 
+Having the ability to use Blueprints to draw Materials into Render Targets and then use or save those results can enable a huge variety of game-specific rendering effects. On the following page, you will find several links to documentation pages that cover how to use Blueprints and Render Targets. 
 
 
 [REGION:topics third]
diff --git a/Engine/Documentation/Source/Engine/Rendering/RenderTargets/BlueprintRenderTargets/HowTo/BlueprintRenderTargetsHowTos.INT.udn b/Engine/Documentation/Source/Engine/Rendering/RenderTargets/BlueprintRenderTargets/HowTo/BlueprintRenderTargetsHowTos.INT.udn
--- a/Engine/Documentation/Source/Engine/Rendering/RenderTargets/BlueprintRenderTargets/HowTo/BlueprintRenderTargetsHowTos.INT.udn
+++ b/Engine/Documentation/Source/Engine/Rendering/RenderTargets/BlueprintRenderTargets/HowTo/BlueprintRenderTargetsHowTos.INT.udn
@@ -1,5 +1,5 @@
 availability:docs
-title: Blueprint and Render Targets - How To's
+title: Blueprint and Render Targets - How Tos
 description: A How To Guide for using Blueprints to manipulate render targets. 
 crumbs:
 version:4.13
@@ -9,7 +9,7 @@ parent:Engine
 order:2
 Navigation:topic
 
-The **Blueprint and Render Targets - How To…** pages provide step by step examples on how to use Blueprints to work with render targets. Whether you want to use Blueprints and render targets to create Textures for use in your projects Materials, or you want to create a surface that deforms when the player shoots at it, these pages will guide you on how to accomplish this. Reference links are also included to provi
+The **Blueprint and Render Targets - How To** pages provide step by step examples on how to use Blueprints to work with render targets. Whether you want to use Blueprints and render targets to create Textures for use in your project's Materials, or you want to create a surface that deforms when the player shoots at it, these pages will guide you on how to accomplish this. Reference links are also included to provi
 
 ## Intermediate
 [REGION:wrapper transparent]
diff --git a/Engine/Documentation/Source/Engine/Rendering/RenderTargets/BlueprintRenderTargets/HowTo/CreatingTextures/CreatingTexturesHowTo.INT.udn b/Engine/Documentation/Source/Engine/Rendering/RenderTargets/BlueprintRenderTargets/HowTo/CreatingTextures/CreatingTexturesHowTo.INT.udn
--- a/Engine/Documentation/Source/Engine/Rendering/RenderTargets/BlueprintRenderTargets/HowTo/CreatingTextures/CreatingTexturesHowTo.INT.udn
+++ b/Engine/Documentation/Source/Engine/Rendering/RenderTargets/BlueprintRenderTargets/HowTo/CreatingTextures/CreatingTexturesHowTo.INT.udn
@@ -13,20 +13,20 @@ Navigation:topic
 
 ![](Creating_Textures_Header_00.png)
 
-When creating Textures for your Unreal Engine 4 (UE4) project, you will generally need to use an additional 2D painting program like Adobe Photoshop or GIMP. While authorizing content in this manner can give you amazing results there are some instances where having the ability to create content like this inside of UE4 would be easier and faster to work on. In the following How To we will take a look at how you can use Blueprints and Render Targets to create Textures to use in your UE4 projects.
+When creating Textures for your Unreal Engine 4 (UE4) project, you will generally need to use an additional 2D painting program like Adobe Photoshop or GIMP. While authoring content in this manner can give you amazing results, there are some instances where having the ability to create content like this inside of UE4 would be faster and easier. In the following How To, we will take a look at how you can use Blueprints and Render Targets to create Textures for your UE4 projects.
 
 ## Required Asset Creation & Setup
 [region:note]
-Please note that the following is not specific to any UE4 project type and can be accomplished in any UE4 project that used UE4 version 4.13 or later. 
+Please note that the following is not specific to any UE4 project type and can be accomplished in any project that uses UE4 version 4.13 or later. 
 [/region]
 
-Before we can create and use Textures, we will first need to create and set up the following three assets inside of the UE4 Content Browser.
+Before we can create and use Textures, we need to set up the following three assets inside of the UE4 Content Browser:
 * **Render Target**
 * **Material**
-* **Blueprint** (Based of Actor)
+* **Blueprint** (Based on Actor)
 
 ### Create a new Empty Level
-Creating a new **Empty Level** to place your Texture creation Blueprint in will make things easier to manage and also ensure that nothing else is affecting your Texture creation process. To create a new Empty Level inside of UE4 you will need to do the following:
+Creating a new **Empty Level** to place your Texture creation Blueprint in will make things easier to manage and also ensure that nothing else is affecting your Texture creation process. To create a new Empty Level inside of UE4, you will need to do the following:
 
 1. From the **Main Toolbar** go to **File** > **New Level**.
 
@@ -36,37 +36,37 @@ Creating a new **Empty Level** to place your Texture creation Blueprint in will
 
 	![](New_Level_02.png)
 
-1. After clicking on the Empty Level option UE4 will load a new level that has nothing inside of it. Once that has happened, save the level by clicking on the **Save** button on the Main Toolbar and giving it a name of **BlueprintTextureCreation**.
+1. After clicking on the Empty Level option, UE4 will load a new level that has nothing inside of it. Once that has happened, save the level by clicking on the **Save** button on the Main Toolbar and giving it a name of **BlueprintTextureCreation**.
 
 	![](New_Level_03.png)
 
 ### Render Target Creation & Setup
-To create a **Render Target** inside of the UE4 Content Browser you will need to do the following:
+To create a **Render Target** inside of the UE4 Content Browser, you will need to do the following:
 
-1. Open up the UE4 Editor and then  **right - click** Inside of the **Content Browser** and from the menu that is displayed, locate the **Create Advanced Assets** section.
+1. Open up the UE4 Editor and then **right-click** Inside of the **Content Browser**. From the menu that is displayed, select the **Create Advanced Assets** section.
 
 	![](Create_RTT_01.png)
 
-1. From the Create Advanced Assets menu, move your mouse to the **Materials & Textures** section and from the asset list that is displayed, select the **Render Target** option.
+1. From the Create Advanced Assets menu, move your mouse to the **Materials & Textures** section and select the **Render Target** option.
 
 	![](Create_RTT_02.png)
 
-1. Name the newly created Render Target, **RT_00** and when completed your Content Browser should look like the following image.
+1. Name the newly created Render Target, **RT_00**. Your Content Browser should look like the following image:
 
 	![](Create_RTT_03.png)
 	
 ### Material Creation & Setup
-To create and set up a Material to use a Render Target you will need to do the following inside of the UE4 Content Browser. 
+To create and set up a Material to use a Render Target, you will need to do the following inside of the UE4 Content Browser:
 
-1. Right - Click in the Content Browser and from the **Create Basic Asset** section select the **Material** option to create a new Material.  
+1. Right-click in the Content Browser and from the **Create Basic Asset** section, select the **Material** option to create a new Material.  
 
 	![](Create_RTT_Material_01.png)
 
-1. Name the Material **RT_Material_00** and then **double - click** on it to open it up in the Material Editor.
+1. Name the Material **RT_Material_00** and then **double-click** on it to open it up in the Material Editor.
 
 	![](Create_RTT_Material_02.png)
 
-1. With the RT_Material_00 open in the Material Editor, search for and add the following Material Expression nodes to the Material Graph.
+1. With the RT_Material_00 open in the Material Editor, add the following Material Expression nodes to the Material Graph:
 
 	* **Texture Coordinate**
 	* **Constant**
@@ -75,40 +75,40 @@ To create and set up a Material to use a Render Target you will need to do the f
 
 	![](Create_RTT_Material_03.png)
 
-1. Once the required Material Expression nodes have been added to the Material Graph, connect the nodes so that they match the image below and press the **Apply** button when completed to compile the Material.
+1. Once the required Material Expression nodes have been added to the Material Graph, connect the nodes so that they match the image below, then press the **Apply** button to compile the Material.
 
 	![](Create_RTT_Material_05.png)
 
 ### Blueprint Creation & Setup
-To create and set up a Blueprint that will work with both Render Targets and Materials you will need to do the following.
+To create and set up a Blueprint that will work with both Render Targets and Materials, you will need to do the following:
 
-1. Right - Click in the Content Browser and from the menu that is displayed in the **Create Basic Asset** section, select the **Blueprint** option to create a new Blueprint.  
+1. Right-click in the Content Browser. From the menu that is displayed, under the **Create Basic Asset** section, select the **Blueprint** option to create a new Blueprint.  
 
 	![](Create_Blueprint_01.png)
 
-1. From the **Pick Parent Class** window, select the **Actor** class and then name the Blueprint **BP_Texture_Creator**.
+1. From the **Pick Parent Class** window, select the **Actor** class, and then name the Blueprint **BP_Texture_Creator**.
 
 	![](Create_Blueprint_02.png)
 
-1. Double - click on the BP_Texture_Creator Blueprint to open it up and once opened, click on the **Event Graph** tab to open up the Event Graph.
+1. Double-click on the BP_Texture_Creator Blueprint to open it up. Once it is opened, click on the **Event Graph** tab to open up the Event Graph.
 
 	![](Create_Blueprint_03.png)
 
-1. Add the following Blueprint nodes to the Event Graph.
+1. Add the following Blueprint nodes to the Event Graph:
 	* **Custom Event**
 	* **Draw Material to Render Target** 
 
 	![](Create_Blueprint_04.png)
 
-1. Select the Custom Event node and press **F2** so you can re-name the Custom Event node to **Bake**. Then connect it's output to the input of the Draw Material to Render Target node.
+1. Select the Custom Event node and press **F2** so you can rename the Custom Event node to **Bake**. Then connect its output to the input of the Draw Material to Render Target node.
 
 	![](Create_Blueprint_05.png)
 
-1. Next, on the Draw Material to Render Target node, load the Render Target that was created previously into the **Texture Render Target** slot by clicking on the small white triangle next to the text **Select Asset** and then from the list find the Render Target, RT_00, and select it by clicking on it. Once that is completed, do the same for the **Material** section, selecting the **RT_Material_00**.
+1. Next, on the Draw Material to Render Target node, load the Render Target that was created previously into the **Texture Render Target** slot by clicking on the small white triangle next to the text **Select Asset** and then find the Render Target, RT_00, in the list and select it by clicking on it. Once that is completed, do the same for the **Material** section, selecting the **RT_Material_00**.
 
 	![](Create_Blueprint_06.png)
 
-1. Now, click on the **Construction Script** tab and drag off the Construction Script node to display the **Executable actions** menu. From this menu locate the **Bake** event that was created in the Event Graph and click on it to add it. 
+1. Now, click on the **Construction Script** tab and drag off the Construction Script node to display the **Executable actions** menu. From this menu, locate the **Bake** event that was created in the Event Graph and click on it to add it. 
 
 	![](Create_Blueprint_07.png)
 
@@ -117,17 +117,17 @@ To create and set up a Blueprint that will work with both Render Targets and Mat
 	![](Create_Blueprint_08.png)
 
 ## Texture Creation
-Now that all of the required parts have been created and set up, in the following section we will take a look at how to create and save the Textures that are created.
+Now that all of the required parts have been created and set up, we will take a look at how to create and save our Textures.
 
-1. Locate the **BP_Texture_Creator** Blueprint in the Content Browser and then drag it into the level.
+1. Locate the **BP_Texture_Creator** Blueprint in the Content Browser and drag it into the level.
 
 	![](Create_Textures_01.png)
 
-1. Locate the **RT_Material_00** Material in the Content Browser and double - click on it to open it up. When opened, select the **Noise** Material Expression node.
+1. Locate the **RT_Material_00** Material in the Content Browser and double-click on it to open it up. When opened, select the **Noise** Material Expression node.
 
 	![](Create_Textures_02.png)
 
-1. With the **Noise** Material Expression selected,  go to the **Details** panel and input the following values into the following parameters and when completed, press the **Apply** button to compile the Material.
+1. With the **Noise** Material Expression selected, go to the **Details** panel and input the following values into the following parameters. When this is complete, press the **Apply** button to compile the Material.
 
 	![](Create_Textures_03.png)
 
@@ -136,21 +136,21 @@ Now that all of the required parts have been created and set up, in the followin
 	* **Output Min:** -0.01
 	* **Level Scale:** 3.0
 
-1. Select the BP_Texture_Creator Blueprint that was placed in the level and move it slightly in any direction to update the RT_00 Render Target that is in the Content Browser. When you do this you should see the RT_00 Render Target update and display exactly what was in the RT_Material_00. 
+1. Select the BP_Texture_Creator Blueprint that was placed in the level and move it slightly in any direction to update the RT_00 Render Target asset. When you do this, you should see RT_00 update and display exactly what was in RT_Material_00. 
 
 	![](Create_Textures_04.png)
 
 	[region:note]
-	Anytime you want to see or update the changes that were made in the Material, you will need to ever so slightly move the Blueprint in order for the changes to propagate. 
+	Anytime you want to see or update the changes that were made in the Material, you will need to move the Blueprint in order for the changes to propagate. 
 	[/region]
 
-1. Inside of the Content Browser, right-click on the RT_00 Render Target and in the **Render Targets Actions** section of the menu, select the **Create Static Texture** option to create a new Texture asset that will display the output of the RT_Material_00.
+1. Inside of the Content Browser, right-click on the RT_00 Render Target. In the **Render Targets Actions** section of the menu, select the **Create Static Texture** option to create a new Texture asset that will display the output of RT_Material_00.
 
 	![](Create_Textures_05.png)
 
-1. When completed, you should now have a new 2D Texture that will display the contents of the RT_Material_00 and is now ready to be used in any other Material of your liking.
+1. When completed, you should have a new 2D Texture that will display the contents of RT_Material_00, and is ready for use in any other Material.
 
 	![](Create_Textures_06.png)
 
 ## Conclusion
-Having the ability to create new Textures inside of UE4 opens up an entirely new world of content creation possibilities. Remember that in order to see your changes you will need to make sure that you compile your Material and slightly move your Blueprint so that your changes will be propagated to your Render Target.
\ No newline at end of file
+Having the ability to create new Textures inside of UE4 opens up an entirely new world of content creation possibilities. Remember that, in order to see your changes, you will need to make sure that you compile your Material and slightly move your Blueprint so that your changes will be propagated to your Render Target.
\ No newline at end of file
diff --git a/Engine/Documentation/Source/Engine/Rendering/RenderTargets/BlueprintRenderTargets/Overview/BlueprintRenderTargetsOverview.INT.udn b/Engine/Documentation/Source/Engine/Rendering/RenderTargets/BlueprintRenderTargets/Overview/BlueprintRenderTargetsOverview.INT.udn
--- a/Engine/Documentation/Source/Engine/Rendering/RenderTargets/BlueprintRenderTargets/Overview/BlueprintRenderTargetsOverview.INT.udn
+++ b/Engine/Documentation/Source/Engine/Rendering/RenderTargets/BlueprintRenderTargets/Overview/BlueprintRenderTargetsOverview.INT.udn
@@ -12,16 +12,18 @@ Navigation:topic
 ![](BP_RT_Overview.png)
 
 ## Overview
-Using Blueprints to interact with render targets opens up a an entirely new way to create and use assets inside Unreal Engine 4 (UE4).  From creating fluid like surfaces to creating Textures that can be used with visual effects and Materials there is almost no limit to the amount of creative freedom that this will offer your UE4 projects. In the following document we will go over some of the things that you need to be aware of when using Blueprints and render targets.  
+Using Blueprints to interact with render targets opens up a an entirely new way to create and use assets inside Unreal Engine 4 (UE4).  From creating fluid like surfaces to Textures that can be used with visual effects and Materials, there is almost no limit to the amount of creative freedom that this will offer your UE4 projects. In the following document, we will go over some of the things that you need to be aware of when using Blueprints and render targets.  
 
 ## Know Limitations
-When using Blueprints and render targets there are few limitations and restrictions that you need to be aware of in order to get the most out of this tool. In the following section we will cover what these limitations are and what can be done to accommodate them in your workflow.
+When using Blueprints and render targets, there are few limitations and restrictions that you need to be aware of in order to get the most out of this tool. In the following section we will cover what these limitations are and what can be done to accommodate them in your workflow.
 
-* You cannot draw to a render target that is being sampled as a Texture by the Material that is input inot the Material input. If you want to accomplish this you will need to either use alpha blending to modify a render target in-place, or ping-pong between two different render targets.
+* You cannot draw to a render target that is being sampled as a Texture by the Material that is input into the Material input. If you want to accomplish this, you will need to either use alpha blending to modify a render target in-place, or ping-pong between two different render targets.
 
 * Only the **Emissive Color** and **Opacity** outputs of the Material are valid when drawing to a render target. 
 
-* Lighting is not supported when using render targets and Material Expression nodes like **World Position** may return unexpected values.
+* Lighting is not supported when using render targets.
+
+* Material Expression nodes like **World Position** may return unexpected values.
 
 * Emissive Color is clamped to be positive by default, but you can output negative values by enabling the material property **AllowNegativeEmissiveColor**.
 
diff --git a/Engine/Documentation/Source/Engine/Rendering/RenderTargets/BlueprintRenderTargets/Reference/BlueprintRenderTargetsReference.INT.udn b/Engine/Documentation/Source/Engine/Rendering/RenderTargets/BlueprintRenderTargets/Reference/BlueprintRenderTargetsReference.INT.udn
--- a/Engine/Documentation/Source/Engine/Rendering/RenderTargets/BlueprintRenderTargets/Reference/BlueprintRenderTargetsReference.INT.udn
+++ b/Engine/Documentation/Source/Engine/Rendering/RenderTargets/BlueprintRenderTargets/Reference/BlueprintRenderTargetsReference.INT.udn
@@ -57,15 +57,15 @@ This page contains a reference listing of properties you will find with the Rend
 
 |Property| Description|
 |----------|-----------------|
-|Width| Render Targets width size.|
-|Height| Render Targets height size.|
+|Width| Render Target's width.|
+|Height| Render Target's height.|
 |Return Value| Texture Render Target 2D reference.|
 
 ## Draw Material to Render Target
 
 ![](Draw_Material_To_Render_Target.png)
 
-* Renders a quad with the Material applied to the specified Render Target. This set the Render Target even if it is already set, which is an expensive operation. Use **BeginDrawCanvasToRenderTarget / EndDrawCanvasToRenderTarget** instead if rendering multiple primitives to the same same Render Target.
+* Renders a quad with the Material applied to the specified Render Target. This set the Render Target even if it is already set, which is an expensive operation. Use **BeginDrawCanvasToRenderTarget / EndDrawCanvasToRenderTarget** instead if rendering multiple primitives to the same Render Target.
 
 |Property| Description|
 |----------|-----------------|
