#UE4 Docs: minor text and formatting updates

diff --git a/Engine/Documentation/Source/Enterprise/Automotive/AutomotiveMaterialPack/AutomotiveMaterialPack.INT.udn b/Engine/Documentation/Source/Enterprise/Automotive/AutomotiveMaterialPack/AutomotiveMaterialPack.INT.udn
--- a/Engine/Documentation/Source/Enterprise/Automotive/AutomotiveMaterialPack/AutomotiveMaterialPack.INT.udn
+++ b/Engine/Documentation/Source/Enterprise/Automotive/AutomotiveMaterialPack/AutomotiveMaterialPack.INT.udn
@@ -2,9 +2,11 @@ Availability:Public
 Title: Automotive Material Pack
 Description:A breakdown of what assets the Automotive Material pack contains and how to use them.
 Crumbs:%ROOT%,
-version:4.12
+version:4.13
 Parent:Enterprise/Automotive
 tags:Enterprise
+tags:Automotive
+tags:Materials
 type:overview
 order:1
 Related:Engine/Rendering/Materials/PhysicallyBased/
@@ -12,12 +14,14 @@ Related:Engine/Rendering/Materials/MaterialProperties/LightingModels
 Related:Engine/Rendering/Materials/HowTo/ClearCoatDualNormal
 Related:Engine/Content/Types/Textures/Cubemaps
 Related:Engine/UI/LevelEditor/Modes/MeshPaintMode/VertexColor
+topic-image:AP_Topic.png
+social-image:automotive_material_pack_social.png
 
-[TOC(start:2 end:3)]
+[TOC(end:2)]
 
 ![](AP_Intro.png)
 
-The Automotive Material Pack is a collection of automotive-themed [PBR Materials](Engine/Rendering/Materials/PhysicallyBased/) and Textures that have been setup for use in Unreal Engine 4(UE4). The pack features 21 Master Materials and 75 Material Instances taken from Epic's 2016 GDC Mclaren 570S demo. In the following document, we will take a look at what Materials the Automotive Material Pack contains and any special information you might need to know in order to use these Materials in your UE4 projects.
+The Automotive Material Pack, available on the Unreal Engine Marketplace, is a collection of automotive-themed [PBR Materials](Engine/Rendering/Materials/PhysicallyBased/) and Textures that have been setup for use in Unreal Engine 4(UE4). The pack features 21 Master Materials and 75 Material Instances taken from Epic's 2016 GDC Mclaren 570S demo. In the following document, we will take a look at what Materials the Automotive Material Pack contains and any special information you might need to know in order to use these Materials in your UE4 projects.
 
 [OBJECT:EmbeddedVideo]
 	[PARAMLITERAL:width]
