#UE4 Docs: converted gifs to async load

diff --git a/Engine/Documentation/Source/Support/Builds/ReleaseNotes/2016/4_13/ReleaseNotes_4_13.INT.udn b/Engine/Documentation/Source/Support/Builds/ReleaseNotes/2016/4_13/ReleaseNotes_4_13.INT.udn
--- a/Engine/Documentation/Source/Support/Builds/ReleaseNotes/2016/4_13/ReleaseNotes_4_13.INT.udn
+++ b/Engine/Documentation/Source/Support/Builds/ReleaseNotes/2016/4_13/ReleaseNotes_4_13.INT.udn
@@ -3,6 +3,24 @@ Title:Unreal Engine 4.13 Release Notes
 Crumbs:%ROOT%, Engine
 Description:
 order:13
+reference-image:image_6.gif
+reference-image:image_7.gif
+reference-image:image_11.gif
+reference-image:image_17.gif
+reference-image:image_18.gif
+reference-image:image_19.gif
+reference-image:image_20.gif
+reference-image:image_23.gif
+reference-image:image_24.gif
+reference-image:image_25.gif
+reference-image:image_26.gif
+reference-image:image_27.gif
+reference-image:image_28.gif
+reference-image:image_29.gif
+reference-image:image_33.gif
+reference-image:image_41.gif
+reference-image:image_45.gif
+reference-image:image_52.gif
 
 This release brings hundreds of updates for Unreal Engine 4, including 145 improvements submitted by the community of Unreal Engine developers on GitHub! Thanks to all of these contributors to Unreal Engine 4.13:
 
@@ -106,11 +124,15 @@ Also, performance has been improved for several of the Noise Material Node featu
 
 Blueprint functions can now be used to draw materials into render targets.  This enables a huge variety of game-specific rendering effects to be implemented without having to modify source code.
 
-![image alt text](image_6.gif)(convert:false)
+[REGION:asyncgif]
+![image alt text](image_6.gif)
+[/REGION]
 
 This is a fluid surface simulation implemented entirely in Blueprint and Material graphs.  Characters and projectiles can push the fluid around!
 
-![image alt text](image_7.gif)(convert:false)
+[REGION:asyncgif]
+![image alt text](image_7.gif)
+[/REGION]
 
 Above shows a simple heightfield painter made entirely in Blueprint, by accumulating a height value at projectile impacts.
 
@@ -142,7 +164,9 @@ Check out the **BlueprintRenderToTarget **map in the ContentExamples project for
 
 Alembic animation import is now supported!  Alembic allows for complex animations to be authored offline, then rendered in real-time inside UE4!  This feature is still considered experimental, but please try it out and send us feedback.
 
-![image alt text](image_11.gif)(convert:false)
+[REGION:asyncgif]
+![image alt text](image_11.gif)
+[/REGION]
 
 We allow importing Alembic caches in different ways:
 
@@ -192,7 +216,9 @@ Unlike deferred decals there is no projection involved, so a typical decal mesh
 
 Using the widget interaction component, you can now simulate hardware input events with widget components in the world.  
 
-![image alt text](image_17.gif)(convert:false)
+[REGION:asyncgif]
+![image alt text](image_17.gif)
+[/REGION]
 
 You can attach it like a laser pointer to any object in the world to interact with widgets; there are also some other options available for more customized behavior.  When standard input comes to the player controller, you’ll instruct the interaction component to simulate a particular hardware input such as Left Mouse Down/Up over whatever widget the user happens to be hovering at that time.
 
@@ -202,7 +228,7 @@ For users who were previously relying on clicking directly on Widget Components
 
 We’ve added a new project template designed for Virtual Reality on desktop and console.
 
-![image alt text](image_18.gif)(convert:false)
+[REGION:asyncgif]![image alt text](image_18.gif)(h:380)[/REGION] [REGION:asyncgif]![image alt text](image_19.gif)(h:380)[/REGION]
 
 To access this, simply choose the new VR option in the New Project template selection window.
 
@@ -212,7 +238,9 @@ This Blueprint project has settings that are optimized to run in VR up to 90 fra
 
 Custom Post-Process materials now can be used on Mobile devices!  Here is a "TV Static" effect displayed using the mobile renderer.
 
-![image alt text](image_20.gif)(convert:false)
+[REGION:asyncgif]
+![image alt text](image_20.gif)
+[/REGION]
 
 * This feature requires the "**Mobile HDR**" option to be enabled in your Project Settings
 
@@ -262,7 +290,7 @@ Mac Metal now has initial Shader Model 5 support enabled by default. This expose
 
 We’ve added a physical animation component that allows you to easily drive skeletal mesh animation through physical motors! 
 
-![image alt text](image_23.gif)(convert:false)
+[REGION:asyncgif]![image alt text](image_23.gif)(h:401)[/REGION] [REGION:asyncgif]![image alt text](image_24.gif)(h:401)[/REGION]
 
 The component allows you to set motor strengths directly, as well as using pre-configured physical animation profiles which can be created and edited inside PhAT.  The new "**Physical Animation Profiles**" feature in PhAT provide a powerful way to customize character physics in different game contexts, as well as fine tune it for special animations. 
 
@@ -272,7 +300,9 @@ You can create and edit different profiles within the PhAT tool, and then change
 
 There is a new utility in Procedural Mesh Component which will ‘slice’ a Procedural Mesh at runtime using a plane.
 
-![image alt text](image_25.gif)(convert:false)
+[REGION:asyncgif]
+![image alt text](image_25.gif)
+[/REGION]
 
 After slicing, we support adding ‘capping’ geometry, and creating a second Procedural Mesh Component for the ‘other half’ if desired. 
 
@@ -282,7 +312,9 @@ Also, Procedural Mesh now supports simple collision, so physics simulation can b
 
 You can now paint on textures and mesh vertices in using the VR Editor.
 
-![image alt text](image_26.gif)(convert:false)
+[REGION:asyncgif]
+![image alt text](image_26.gif)
+[/REGION]
 
 This allows you to use motion controllers to paint on static meshes while immersed in VR.    To use this feature, open the "Modes" window in VR, then click the “Mesh Paint” tab.  Now simply select an object in the world, then aim and pull your trigger to paint!  Pressure sensitivity is supported on your controller’s trigger, and you can hold the ‘Modifier’ button to erase instead of paint.
 
@@ -290,7 +322,9 @@ This allows you to use motion controllers to paint on static meshes while immers
 
 In this release, you can use motion controllers to spray down foliage instances while in VR.
 
-![image alt text](image_27.gif)(convert:false)
+[REGION:asyncgif]
+![image alt text](image_27.gif)
+[/REGION]
 
 Select a foliage type, aim your laser and pull the trigger to paint foliage! You can hold the modifier button to erase foliage. Pressure sensitivity on the trigger is supported. Additionally, the Foliage Editor has been improved to be a bit more functional while in VR, although some features, lasso tool and select tool, are still unavailable.
 
@@ -298,13 +332,17 @@ Select a foliage type, aim your laser and pull the trigger to paint foliage! You
 
 The Color Picker window is now available in VR, so you can change color properties on lights and other Actors in your level. You can also use the Color Picker to select colors for Vertex Painting and Texture Painting in VR. 
 
-![image alt text](image_28.gif)(convert:false)
+[REGION:asyncgif]
+![image alt text](image_28.gif)
+[/REGION]
 
 ## New: Play from VR Editor
 
 To easily prototype your project it is now possible to play your game in VR from within the VR Editor. Press the "Play" button on the quick menu to start playing in VR!  To instantly go back to the VR Editor, *hold the Grip buttons on both controllers and squeeze both trigger buttons*.
 
-![image alt text](image_29.gif)(convert:false)
+[REGION:asyncgif]
+![image alt text](image_29.gif)
+[/REGION]
 
 ## New: Improved VR Transform Gizmo
 
@@ -330,7 +368,9 @@ You can now take screenshots right from VR!
 
 You can now enter and leave VR editing mode when the VR editor is enabled without having to use the VR button or escape manually! As long as the editor is in the foreground, when you wear the headset, you will automatically enter VR editing mode; when you remove the headset, you will leave it.
 
-![image alt text](image_33.gif)(convert:false)
+[REGION:asyncgif]
+![image alt text](image_33.gif)
+[/REGION]
 
 There is a setting under VR in the Experimental section of Editor Settings that will allow you to turn off auto-entry if you prefer.
 
@@ -468,7 +508,9 @@ In the visualization above, the highest the LOD (white) is rendered with tessell
 
 We have added a new type of animation asset called a Pose Asset. This contains a set of named bone poses, which you can blend additively, in a similar manner to blend shapes for vertices. 
 
-![image alt text](image_41.gif)(convert:false)
+[REGION:asyncgif]
+![image alt text](image_41.gif)
+[/REGION]
 
 One use for this is to support facial animation where either FACS (Facial Action Coding System) or viseme curves can drive poses. However you could use this system to create new animation by blending multiple poses.
 
@@ -488,7 +530,9 @@ To support this system, we have improved how we handle scale with additive blend
 
 We have added a new Pose Driver node, which lets you drive a curve value (such as a morph target weight) based on a bone’s movement. 
 
-![image alt text](image_44.gif)(convert:false)
+[REGION:asyncgif]
+![image alt text](image_44.gif)
+[/REGION]
 
 This uses an RBF (Radial Basis Function) to interpolate driven values based on the orientation of a target bone. You use a PoseAsset to define the target poses for the bone, and the desired curve values at each pose. This node can be used as a Pose Space Deformer, to drive corrective blend shapes based on the orientation of a bone.
 
@@ -496,7 +540,9 @@ This uses an RBF (Radial Basis Function) to interpolate driven values based on t
 
 Anim graph nodes can now be "watched" in Persona. 
 
-![image alt text](image_45.gif)(convert:false)
+[REGION:asyncgif]
+![image alt text](image_45.gif)
+[/REGION]
 
 This allows you to see a representation of the pose being generated at any point in the anim graph dynamically. Multiple watches can be active at once allowing you to compare poses at different points and find the exact point at which any errors in your current pose are introduced. This can be very useful for debugging complex Animation Blueprints and previously would only have been achievable by connecting the node you wanted to view directly to the root node and recompiling the blueprint.
 
@@ -552,7 +598,9 @@ Now it's possible to set Spline Component points in the Blueprint Editor using t
 
 The "Reset to Default" context action on the spline visualizer will set an instance back to the Blueprint default.  Any changes made to the Blueprint default will be propagated to any instance whose spline points have not been subsequently edited.
 
-![image alt text](image_52.gif)(convert:false)
+[REGION:asyncgif]
+![image alt text](image_52.gif)
+[/REGION]
 
 **Numerical editing of spline points in the Details panel**
 
@@ -606,7 +654,7 @@ We removed the Skeleton Curve tab from Persona, and moved that functionality int
 
 You can see all curves that belong to current skeleton or currently active curves from preview asset. And you could also filter by specific type of curves if you only want to see active curves. Please note that we named default curve to be called "Attribute", so any animation curves will be by default attribute curves. 
 
-## If you want to modify the curve value, you can either turn off Auto check box option or just type the value. 
+If you want to modify the curve value, you can either turn off Auto check box option or just type the value. 
 
 ## New: Sprites in UMG Widgets
 
@@ -4332,6 +4380,10 @@ TSets can now be used as a UPROPERTY for the purposes of automatic serialization
 
 * FContext (in our internationalization manifest) has been renamed to FManifestContext to better reflect its purpose and avoid naming conflicts with other code
 
+* References to the USplineComponent members SplineInfo, SplineRotInfo and SplineScaleInfo should now be replaced by calls to GetSplinePointsPosition(), GetSplinePointsRotation() and GetSplinePointsScale(), respectively.
+
+* Overridden implementations of FComponentVisualizer::VisProxyHandleClick must now take a FEditorViewportClient* instead of a FLevelEditorViewportClient* as the first parameter.
+
 ### Animation
 
 * Added a validation path to see if mesh hierarchy matches with the animation skeleton when initializing animation.
@@ -4442,7 +4494,7 @@ TSets can now be used as a UPROPERTY for the purposes of automatic serialization
 
 * Any user-created vertex factory that implements the Get Static Batch Element Visibility function will need to set "Requires Per Element Visibility" on its Mesh Batch.
 
-#### Optimizations7
+#### Optimizations
 
 * The console command "r.Tonemapper.ScreenPercentage" has been renamed to “r.Tonemapper.MergeWithUpscale.Mode”, and offers additional modes (see the main release notes or cvar help for more information).
 
