#UE4 Docs: fixed broken link

diff --git a/Engine/Documentation/Source/Resources/Templates/FirstPerson/FirstPersonTemplate.INT.udn b/Engine/Documentation/Source/Resources/Templates/FirstPerson/FirstPersonTemplate.INT.udn
--- a/Engine/Documentation/Source/Resources/Templates/FirstPerson/FirstPersonTemplate.INT.udn
+++ b/Engine/Documentation/Source/Resources/Templates/FirstPerson/FirstPersonTemplate.INT.udn
@@ -53,7 +53,7 @@ So you have started a new game using the **First Person Template**, awesome! If
 	* **Bonus**: Add a Particle System Component to the **FirstPersonProjectile** to make it look like it trails smoke behind it. Look for the **Use Local Space** property [here](Engine/Rendering/ParticleSystems/Reference/Modules/Required).
 * [Create a HUD](Engine/UMG/QuickStart) for your First Person Game using [Unreal Motion Graphics](Engine/UMG).
 * Why not add some Destructible Meshes to the environment? Check out the [Destructible Actor Documentation](Engine/Physics/Destructibles) and the [](Engine/Blueprints/HowTo) for methods on how to do this!
-* You can add **AI Characters** using [Behavior Trees](Gameplay/AI/BehaviorTreeQuickStart) to your game and have them chase, flee, help, or harm the player.
+* You can add **AI Characters** using [Behavior Trees](Engine/AI/BehaviorTrees/QuickStart) to your game and have them chase, flee, help, or harm the player.
 	* **Bonus**: Implement [damage](Engine/Blueprints/UserGuide/Events#eventanydamage) on the **FirstPersonCharacter** and your new **AI Character**, and make it so they can [Destroy](Engine/Blueprints/UserGuide/Events#eventendplay) each other. There is more information on how to do this in the [](Gameplay/HowTo/SpawnAndDestroyActors/Blueprints) guide.
 	* **Double Bonus**: Following the [](Gameplay/HowTo/CharacterMovement/Blueprints) guide, give your AI more movement options!
 * If you end up killing your Player, maybe you will want to be able to respawn them? Check out this how to guide for [Respawning the Player](Gameplay/HowTo/RespawnPlayer/Blueprints).
diff --git a/Engine/Documentation/Source/Resources/Templates/ThirdPerson/ThirdPersonTemplate.INT.udn b/Engine/Documentation/Source/Resources/Templates/ThirdPerson/ThirdPersonTemplate.INT.udn
--- a/Engine/Documentation/Source/Resources/Templates/ThirdPerson/ThirdPersonTemplate.INT.udn
+++ b/Engine/Documentation/Source/Resources/Templates/ThirdPerson/ThirdPersonTemplate.INT.udn
@@ -45,14 +45,14 @@ resources
 [/OBJECT]
 [/VAR]
 
-![](templateMore.png)(convert:false)
+![](templateMore.PNG)(convert:false)
 
 So you have started a new game using the **Third Person Template**? Awesome! If you are looking for guidance on how to do specific things or perhaps just looking at the blank canvas of endless possibilities, this doc contains ideas for things you can implement.
 
 * [Read up](Engine/Blueprints) and [watch some video](Videos) on the **Blueprints Visual Scripting** system. Sometimes it is not about what to do, but how you can do it, and knowing more about Blueprints will enable you to rapidly prototype ideas, even if you will ultimately be working in C++.
 * [Create a HUD](Engine/UMG/QuickStart) for your Third Person Game using [Unreal Motion Graphics](Engine/UMG).
 * Why not add some Destructible Meshes to the environment? Check out the [Destructible Actor Documentation](Engine/Physics/Destructibles) and the [](Engine/Blueprints/HowTo) for methods on how to do this!
-* You can add **AI Characters** using [Behavior Trees](Gameplay/AI/BehaviorTreeQuickStart) to your game and have them chase, flee, help, or harm the player.
+* You can add **AI Characters** using [Behavior Trees](Engine/AI/BehaviorTrees/QuickStart) to your game and have them chase, flee, help, or harm the player.
 	* **Bonus**: Implement [damage](Engine/Blueprints/UserGuide/Events#eventanydamage) on the **ThirdPersonCharacter** and your new **AI Character**, and make it so they can [Destroy](Engine/Blueprints/UserGuide/Events#eventendplay) each other. There is more information on how to do this in the [](Gameplay/HowTo/SpawnAndDestroyActors/Blueprints) guide.
 	* **Double Bonus**: Following the [](Gameplay/HowTo/CharacterMovement/Blueprints) guide, give your AI and **ThirdPersonCharacter** more movement options!
 * If you end up killing your Player, maybe you will want to be able to respawn them? Check out this how to guide for [Respawning the Player](Gameplay/HowTo/RespawnPlayer/Blueprints).
