#UE4 Docs: added Enterprise and Automotive landing pages

diff --git a/Engine/Documentation/Source/Enterprise/Automotive/Automotive.INT.udn b/Engine/Documentation/Source/Enterprise/Automotive/Automotive.INT.udn
new file mode 100644
--- /dev/null
+++ b/Engine/Documentation/Source/Enterprise/Automotive/Automotive.INT.udn
@@ -0,0 +1,15 @@
+Availability:Public
+Title: Automotive Applications
+Description:Using Unreal Engine for automotive visualization applications
+Crumbs:%ROOT%,
+version:4.13
+Parent:Enterprise
+tags:Enterprise
+tags:Automotive
+type:landing
+order:8
+topic-image:Automotive_Topic.png
+
+![](Automotive_Hero.png)
+
+[DIR(output:"topic" org:"hierarchy" parent:"Enterprise/Automotive")]
diff --git a/Engine/Documentation/Source/Enterprise/Enterprise.INT.udn b/Engine/Documentation/Source/Enterprise/Enterprise.INT.udn
new file mode 100644
--- /dev/null
+++ b/Engine/Documentation/Source/Enterprise/Enterprise.INT.udn
@@ -0,0 +1,17 @@
+Availability:Public
+Title: Enterprise Applications
+Description:Using Unreal Engine for automotive, aviation, architecture, consumer electronics, and complex data visualization applications
+Crumbs:%ROOT%,
+version:4.13
+Parent:%ROOT%
+tags:Enterprise
+type:landing
+order:8
+topic-image:Enterprise_Topic.png
+topic-icon:%ROOT%/enterprise_icon.png
+
+![](Enterprise_Hero.png)
+
+Unreal Engine and its real-time interaction and rendering capabilities are well-suited for a huge range of projects across all Enterprise applications, including automotive, aviation, architecture, consumer electronics and complex data visualization.
+
+[DIR(output:"topic" org:"hierarchy" parent:"Enterprise")]
