#DocFix

Made fixes to one of the lines requested by JoshM

diff --git a/Engine/Documentation/Source/Platforms/Android/GettingStarted/1/AndroidQS_1.INT.udn b/Engine/Documentation/Source/Platforms/Android/GettingStarted/1/AndroidQS_1.INT.udn
--- a/Engine/Documentation/Source/Platforms/Android/GettingStarted/1/AndroidQS_1.INT.udn
+++ b/Engine/Documentation/Source/Platforms/Android/GettingStarted/1/AndroidQS_1.INT.udn
@@ -104,7 +104,7 @@ It is highly recommended that you install CodeWorks for Android 1R4 while the ed
 	It is perfectly fine to leave this set to the default location.
 	[/region]
 
-1.	Make sure that your are installing **CodeWorks for Android 1R4** and then click the **Next** button.
+1.	Make sure that your are installing **CodeWorks for Android 1R4u2** and then click the **Next** button.
 	
 	![](AndroidWorks_03.png)
 
