Copying //UE4/Dev-Build to //UE4/Dev-Main (Source: //UE4/Dev-Build @ 3145834)

==========================
MAJOR FEATURES + CHANGES
==========================

Change 3129636 on 2016/09/17 by [REDACTED_PERSON]

	UBT: Add a "-nolink" option which allows compiling object files without linking them into an executable. Useful for non-unity builds, which take a very long time to link or fail on some platforms due to command lines being too long, PDB file having too many records, etc...

Change 3129825 on 2016/09/18 by [REDACTED_PERSON]

	UBT: Don't force Linux to build in unity; it seems to build fine without.

Change 3129965 on 2016/09/19 by [REDACTED_PERSON]

	Duplicating CL#3129960 from Release-4.13
	Exclude NetworkProfiler when building CS tools for Linux

Change 3130653 on 2016/09/19 by [REDACTED_PERSON]

	UHT: Fix missing "Error:" prefix in output log, causing messages to be ignored for failure emails.

Change 3130662 on 2016/09/19 by [REDACTED_PERSON]

	EC: Prevent UHT failures from being reported twice, and remove the need for special case to show UHT summary output.

Change 3131956 on 2016/09/20 by [REDACTED_PERSON]

	Addtional fixes for compiling Editor as a monolithic executable
	Change so monolithic editor is output to Project Binaries directory
	Removed duplicated ReturnContainerIndexFromChannelName function
	Only check for out of date modules for non monolithic editor
	Don't define GIsGameAgnosticExe or PER_MODULE_BOILERPLATE for monolithic editor, done elsewhere
	Correct IMPLEMENT_MODULE for QuadricMeshReduction and AudioCapture modules

Change 3132112 on 2016/09/20 by [REDACTED_PERSON]

	Docs: Remove reference to UBT environment variables from configuration docs.

Change 3132815 on 2016/09/20 by [REDACTED_PERSON]

	AutomationTool: Delete GUBP. Everything now uses BuildGraph!

Change 3132871 on 2016/09/20 by [REDACTED_PERSON]

	UBT: Remove GUBP callbacks from TargetRules instances.

Change 3132987 on 2016/09/20 by [REDACTED_PERSON]

	Allow public distribution of the compiled SimplygonMeshUtilities binaries.

Change 3133974 on 2016/09/21 by [REDACTED_PERSON]

	Allow public distribution of the SimplygonSwarm module. Requires a separate Simplygon DLL (still in a NotForLicensees folder) to function correctly.

Change 3137228 on 2016/09/22 by [REDACTED_PERSON]

	UAT: Merging fix to parallel executor on Linux from 4.13 branch.

Change 3139783 on 2016/09/26 by [REDACTED_PERSON]

	Fixed Xbox support for Installed Builds
	Corrected typo in Xbox+PS4 filter creation and added XboxOnePackageNameUtil.exe
	Added Xbox versions of ThirdParty libs that hadn't been specified until now

Change 3141721 on 2016/09/27 by [REDACTED_PERSON]

	Remove declaration of circular references between FbxAutomationTestBuilder and LevelEditor; causes LevelEditor to be built differently if plugin is enabled, which results in shared build products being invalidated by switching between games.

Change 3141789 on 2016/09/27 by [REDACTED_PERSON]

	UBT: Retain the ".suppressed" part of output file names when building import libraries for circularly referenced modules.

Change 3141805 on 2016/09/27 by [REDACTED_PERSON]

	UBT: Allow reusing build ids in version manifests as long as we aren't modifying any engine binaries (building more or fewer is permitted), and merge manifests together if possible. Allows building the entire solution through Visual Studio, when some projects may build more modules than another.

Change 3141980 on 2016/09/27 by [REDACTED_PERSON]

	EC: Include a "Steps to Reproduce" section in failure emails, which gives the command line to run to execute the step.

Change 3143996 on 2016/09/28 by [REDACTED_PERSON]

	BuildGraph: Fix exported job steps having dependencies on nodes behind triggers, causing jobs to never terminate.

Change 3144696 on 2016/09/29 by [REDACTED_PERSON]

	Adding -NoSharedPCH to NonUnity build steps and split them into separate nodes for UE4Editor and UE4Game
	Fixed redefinition of __WINDOWS_DS__

Change 3144931 on 2016/09/29 by [REDACTED_PERSON]

	Core: Changes to search paths for DLL loading.

	* The default binaries directory is now added to the list of search paths. Normally LoadLibrary finds these without needing to preload them, but plugins can change the global search paths by calling SetDefaultDllDirectories().
	* Only the top entry of the DLL directory stack is searched. There is typically only one directory here anyway, but the intent is more consistent with the operation of SetDllDirectory().
	* Resolved import paths are converted to absolute, so the resulting calls to LoadLibrary will not be influenced by calls to SetDefaultDllDirectories() changing the base directory.
	* Search paths aren't de-duplicated any more. They don't overlap in practice, and it's not expensive even if they do.

Change 3144932 on 2016/09/29 by [REDACTED_PERSON]

	Vulkan: Only add vulkan-1.dll as a delay load dependency from the VulkanRHI module, otherwise it can be added to the linker command-line 20 or more times.

Change 3145011 on 2016/09/29 by [REDACTED_PERSON]

	Core: Include the compatible changelist in version manifests, so local builds after syncing with UGS will have the correct compatible changelist numbers post-hotfix.

diff --git a/Engine/Documentation/Source/Programming/UnrealBuildSystem/Configuration/BuildConfiguration.INT.udn b/Engine/Documentation/Source/Programming/UnrealBuildSystem/Configuration/BuildConfiguration.INT.udn
--- a/Engine/Documentation/Source/Programming/UnrealBuildSystem/Configuration/BuildConfiguration.INT.udn
+++ b/Engine/Documentation/Source/Programming/UnrealBuildSystem/Configuration/BuildConfiguration.INT.udn
@@ -1,10 +1,17 @@
 Availability:Public
 Title:Configuring Unreal Build System
 Crumbs: %ROOT%, Programming, Programming/UnrealBuildSystem
-Description:Configuring Unreal Build System to control  how the engine is built.
-Version: 4.9
+Description:Configuring Unreal Build System to control how the engine is built.
+Version: 4.14
 
-The Unreal Build System can be configured through the BuildConfiguration.cs to modify the build process to suit a variety of needs.
+UnrealBuildTool reads settings from XML configuration files in the following locations:
+
+    * Engine/Programs/NotForLicensees/UnrealBuildTool/BuildConfiguration.xml
+	* Engine/Saved/UnrealBuildTool/BuildConfiguration.xml
+	* *User Folder/AppData*/Roaming/Unreal Engine/UnrealBuildTool/BuildConfiguration.xml
+	* *My Documents*/Unreal Engine/UnrealBuildTool/BuildConfiguration.xml
+
+The default settings are written to Engine/Programs/UnrealBuildTool/BuildConfiguration.xml as a reference. This file is generated whenever UnrealBuildTool runs; any settings you change here will be overwritten.
 
 ## UnrealBuldTool Configuration Properties
 
@@ -15,17 +22,7 @@ $ bStopXGECompilationAfterErrors : If true, compilation with XGE will stop when
 $ bOmitPCDebugInfoInDevelopment : If true, debug info is not included when building in Development configuration.
     * Set to false to always include debug info.
     * Including debug info slows down compiling, but gives you line numbers and stack info in Development builds.
-$ bUseUnityBuild : Set to true when rebuilding the engine as a whole. This should be disabled when using fast iteration mode. See [Build Modes](#BuildModes) for more information.
-$ bEnableFastIteration : If true, compile and link times are dramatically decreased. This should be disabled when creating unity builds. See [Build Modes](#BuildModes) for more information.
+$ bUseUnityBuild : Set to true when rebuilding the engine as a whole. This should be disabled when using fast iteration mode.
+$ bEnableFastIteration : If true, compile and link times are dramatically decreased. This should be disabled when creating unity builds.
     * Setting to _true_ gives fastest iteration times.
     * If you encounter weird "RPC" or "PDB" link errors, then you should set `bUseIncrementalLinking = false` in `ValidateConfiguration()`
-
-
-## Build Modes
-
-There are two build modes: **Unity** and **Fast Iteration**. Unity is good for making changes that require rebuilding the whole engine e.g. changing an Engine header. Fast iteration is good when you are making localized changes - it gets your compile/link time down to less than 4 seconds. To change between the two, open BuildConfiguration.cs and find the lines:
-
-	bEnableFastIteration = Utils.GetEnvironmentVariable( "ue.bEnableFastIteration", false );
-	bUseUnityBuild = Utils.GetEnvironmentVariable("ue.bUseUnityBuild", true);
-
-This can be modified to force the one you want to be true. Unity is the default.
