#DocFix

Cleaned up a couple of lines from AllanB's comments for UEDOC-3218

diff --git a/Engine/Documentation/Source/Platforms/Android/OpenGLES31MobileRenderer/OpenGLES31MobileRenderer.INT.udn b/Engine/Documentation/Source/Platforms/Android/OpenGLES31MobileRenderer/OpenGLES31MobileRenderer.INT.udn
--- a/Engine/Documentation/Source/Platforms/Android/OpenGLES31MobileRenderer/OpenGLES31MobileRenderer.INT.udn
+++ b/Engine/Documentation/Source/Platforms/Android/OpenGLES31MobileRenderer/OpenGLES31MobileRenderer.INT.udn
@@ -51,7 +51,7 @@ In order to preview what OpenGL ES 3.1 will look like in the UE4 Editor, you wil
 	Click for full image.
 	[/REGION]
 
-1. Now if you look at the lower right-hand corner of the UE4 Editor viewport, you should see the text **Feature Level: ES3_1** which means you are now previewing the level using OpenGL ES 3.1.
+1. Now if you look at the lower right-hand corner of the UE4 Editor viewport, you should see the text **Feature Level: ES3_1** which means you are now previewing the level using OpenGL ES 3.1's feature set.
 
 	[REGION:lightbox]
 	[![](OGL3_02.png)(w:404)](OGL3_02.png)
