Copying //UE4/Dev-Framework to //UE4/Dev-Main (Source: //UE4/Dev-Framework @ 3227619)

==========================
MAJOR FEATURES + CHANGES
==========================

Change 3198996 on 2016/11/15 by [REDACTED_PERSON]

	BeginPlay calls will now be dispatched in a consistent order regardless of placed in persistent level, streamed in level, or dynamically spawned
	AActor::BeginPlay is now protected, you should call DispatchBeginPlay instead.

Change 3199019 on 2016/11/15 by [REDACTED_PERSON]

	Mark user-facing BeginPlay calls as protected

Change 3200128 on 2016/11/16 by [REDACTED_PERSON]

	Dont propgate threaded update flag from UAnimBluepint to CDO if we fail thread safety checks

	Also fully deprecated (with _DEPRECATED) older flags in UAnimInstance.

Change 3200133 on 2016/11/16 by [REDACTED_PERSON]

	Fix Set Anim Instance Class not working on the second attempt (InitAnim would not be called)

Change 3200167 on 2016/11/16 by [REDACTED_PERSON]

	Newly added virtual bones are now selected in the skeleton tree

Change 3200255 on 2016/11/16 by [REDACTED_PERSON]

	Stop SkeletalMeshTypes.h being globally included

Change 3200289 on 2016/11/16 by Jurre.deBaare

	Hidden Material References from Mesh Components Fix
	#fix Make sure that in PostEditChangeProp we reset the override material arrays
	#misc changed a property comparison to use GET_MEMBER_NAME_CHECKED instead

Change 3200291 on 2016/11/16 by Jurre.deBaare

	Imported Alembic skeletal anims have cut-off shadow due to moving out of the bounds
	#fix retrieve bounds from alembic archive at various levels (global, transform, meshes) and build archive bounds which is set on the animation sequence

Change 3200293 on 2016/11/16 by Jurre.deBaare

	Overlapping UV's cause merge actor texture baking issues
	#fix Only look for overlapping UVs if vertex data baking is actually expected/enabled

Change 3200294 on 2016/11/16 by Jurre.deBaare

	Scrubbing Playback Speed under Geometry Cache in the details panel is too sensitive
	#fix Make the UIMin/Max smaller than the clamping value for proper user interaction while sliding (thanks James for the tip!)

Change 3200295 on 2016/11/16 by Jurre.deBaare

	Merge Actor Specific LOD level can be set to 8
	#fix Change clamping value and added UI clamp metadata

Change 3200296 on 2016/11/16 by Jurre.deBaare

	In Merge Actors if you select use specific Lod level you have access to all the merge material settings
	#fix Added edit condition to non-grayed out material settings

Change 3200303 on 2016/11/16 by [REDACTED_PERSON]

	Fixed diagonal current scrub value in anim curves

Change 3200304 on 2016/11/16 by [REDACTED_PERSON]

	Rezero is now explicit about what it does (current vs. specified frame)

	Also no longer ingores Z-offset (legacy feature - root motion can have any translation, not just 2D).

Change 3200307 on 2016/11/16 by [REDACTED_PERSON]

	Add curve panel to anim BP editor

	Also improve curve modification message routing. We were needlessly passing delegates up and down the widget hierarchy and conflating smart name edits with curve edits (key addition etc.).

Change 3200313 on 2016/11/16 by Jurre.deBaare

	Animations with materials driven by scalar parameters from curves wont update until persona is closed and reopened
	#fix in debug skeletal mesh component just mark the cached parameters dirty every tick

Change 3200316 on 2016/11/16 by Jurre.deBaare

	Converted Skeletal To Static Mesh Gets Corrupted When Merged
	#fix Assume that the all static meshes will contain valid texture coordinates for channel 0 (which is expect by static mesh code as well)
	#misc Ensure that we set the lightmap index for converted skeletal meshes to either an empty one or the highest one used

Change 3200321 on 2016/11/16 by Jurre.deBaare

	Scrolling/scroll bar are disabled in Alembic Import window if you scroll a certain way down
	#fix change the way the layout is constructed

Change 3200323 on 2016/11/16 by Jurre.deBaare

	Toggling sky in Persona does not effect reflections
	#fix turn of skylight together with the actual environment sphere
	#misc found incorrect copy paste in toggling floor/environment visibility with key stroke

Change 3200324 on 2016/11/16 by Jurre.deBaare

	Open Merge Actor menu on right clicking two selected actors
	#fix Added option 'Merge Actors' to right-click context menu when having selected one or multiple actors in the viewport

Change 3200331 on 2016/11/16 by [REDACTED_PERSON]

	Added support for suspending clothing simulations at runtime, exposed also to blueperints. And aded option in Persona to pause simulations when animations are paused.

Change 3200334 on 2016/11/16 by Jurre.deBaare

	Dynamic light settings in Persona viewport cause edges to appear hardened
	#fix Makeing the directional light stationary to ups the shadowing quality

Change 3200356 on 2016/11/16 by Jurre.deBaare

	Rate scale option for animation nodes in blend spaces
	#added Rate scale variable to blend space samples, these rates are now multiplied with the global rate scale during playback
	#misc bumped framework object version to update all blendspaces on load

Change 3200380 on 2016/11/16 by Jurre.deBaare

	Fix for Mac CIS issues

Change 3200383 on 2016/11/16 by [REDACTED_PERSON]

	Split FAttenuationSettings in to FBaseAttenuationSettings and FSoundAttenuationSettings in preparation for reuse of the base attenuation for force feedback

Change 3200385 on 2016/11/16 by [REDACTED_PERSON]

	Refactor SkeletalMesh to use same color buffer type as StaticMesh

Change 3200407 on 2016/11/16 by [REDACTED_PERSON]

	Fix CIS error in FbxAutomationTests.cpp

Change 3200417 on 2016/11/16 by Jurre.deBaare

	Fix for CIS issues
	#fix Rogue }

Change 3200446 on 2016/11/16 by [REDACTED_PERSON]

	Change fix for Set Anim Instance Class from CL 3200133

Change 3200579 on 2016/11/16 by [REDACTED_PERSON]

	Fix for serialization crash in Odin

	#jir UE-38683

Change 3200659 on 2016/11/16 by [REDACTED_PERSON]

	Fix build errors

Change 3200801 on 2016/11/16 by [REDACTED_PERSON]

	Fix error message

Change 3200873 on 2016/11/16 by [REDACTED_PERSON]

	Test case for Update Rate Optimization

	- LOD_URO_Map.umap - test map
	- LODPawn - pawn that contains mesh with URO setting
	- You can tweak the value in LODPawn

Change 3201017 on 2016/11/16 by [REDACTED_PERSON]

	- Allow slave component to be removed when setting master pose to nullptr
	- licensee reported this issue. https://udn.unrealengine.com/questions/321037/skeletalmeshcomponent.html

Change 3201765 on 2016/11/17 by Jurre.deBaare

	Improved tooltip for FBlendParameter.GridNum

Change 3201817 on 2016/11/17 by [REDACTED_PERSON]

	Added display/edit of bone transforms in details panel

	Added UBoneProxy tickable editor object held by the skeleton tree that updates its internal transforms in Tick().
	Updated various bits of supporting code to allow selection to be properly preserved in cases such as undo/redo. This allows the bone proxy object to be displayed over an undo/redo event. It also fixes some inconsistency with selection between the skeleton tree and the preview scene.
	Breaking change: Updated FOnPreviewMeshChangedMulticaster delegate signature to take both the old and new skeletal mesh. This is to allow clients to skip certain logic if the skeletal mesh hasnt really changed (in this case de-selection).

Change 3201819 on 2016/11/17 by [REDACTED_PERSON]

	Fix CIS error

Change 3201901 on 2016/11/17 by [REDACTED_PERSON]

	With new system, the skeleton curve count is not the one we should check but BoneContainer.GetAnimCurveNameUids().
	- removed GetCurveNumber from skeleton
	- changed curve count to  use BoneContainer's curve list.

Change 3201999 on 2016/11/17 by [REDACTED_PERSON]

	Add local/world transform editing to bone editing

	Added details customization & support code for world-space editing of bone transforms

Change 3202111 on 2016/11/17 by mason.seay

	Potential test assets for HLOD

Change 3202240 on 2016/11/17 by [REDACTED_PERSON]

	Fixed extra whitespace not being removed in front of console commands.

	GitHub #2843

Change 3202259 on 2016/11/17 by Jurre.deBaare

	Readded missing shadows in advanced preview scene

Change 3203180 on 2016/11/17 by mason.seay

	Moved and updated URO Map

Change 3203678 on 2016/11/18 by [REDACTED_PERSON]

	Bug fix for menu extenders in PhAT.

	GitHub #2550

Change 3203679 on 2016/11/18 by [REDACTED_PERSON]

	Fixed LOD hysteresis not being properly converted from the old metric

	This addreses some 'LOD lag' issues seen when just treating as an equivalent fudge factor, as the magnitude needed to have an effect has changed.

Change 3203747 on 2016/11/18 by Jurre.deBaare

	Crash when repeatedly undoing and readding of animation to a AnimOffset 1D - IsValidBlendSampleIndex
	#fix Ensure we reset the hightlighting / dragging / selection state when PostUndo is called, this makes sure we repopulate tooltips if need etc.

Change 3203748 on 2016/11/18 by Jurre.deBaare

	Crash Generating Proxy Meshes after replacing static meshes in the level
	#fix just calculate bounds for the used UVs (old behaviour was wrong)

Change 3203751 on 2016/11/18 by james.cobbett

	Changes to TM-PoseSnapshot and new test assets

Change 3203799 on 2016/11/18 by [REDACTED_PERSON]

	Switched fudged auto-LOD calculations to use a pow() decay instead of a recprocal

	Still a fudge when LOD reduction has not been performed in-engine, but a fudge with similar outcomes to the previous method.
	Also fixed up the naming of some variables that still referred to screen areas & LOD distances.

Change 3203856 on 2016/11/18 by james.cobbett

	TM-PoseSnapshot - Rebuild lighting and updated anims

Change 3203880 on 2016/11/18 by [REDACTED_PERSON]

	Copying //UE4/Dev-Physics-Upgrade to Dev-Framework (//UE4/Dev-Framework)

Change 3203940 on 2016/11/18 by [REDACTED_PERSON]

	Fix missing newline for ps4

Change 3203960 on 2016/11/18 by [REDACTED_PERSON]

	Readd fix for linux macro expansion warning

Change 3203975 on 2016/11/18 by [REDACTED_PERSON]

	Fix for linux toolchain not knowing about no-unused-local-typedef

Change 3203989 on 2016/11/18 by [REDACTED_PERSON]

	Make sure physx automation doesn't try to build html5 APEX.

Change 3204031 on 2016/11/18 by james.cobbett

	Minor update to test level

Change 3204035 on 2016/11/18 by [REDACTED_PERSON]

	Additional Attenuation refactor cleanup

Change 3204044 on 2016/11/18 by [REDACTED_PERSON]

	Fix typo of NV_SIMD_SSE2

Change 3204049 on 2016/11/18 by [REDACTED_PERSON]

	Fix missing newline for PS4 compiler

Change 3204463 on 2016/11/18 by mason.seay

	Finalized URO test map

Change 3204621 on 2016/11/18 by mason.seay

	Small improvements

Change 3204751 on 2016/11/18 by [REDACTED_PERSON]

	Make PhAT highlight selected bodies and constraints in the tree view

Change 3205868 on 2016/11/21 by [REDACTED_PERSON]

	Merging //UE4/Dev-Main to Dev-Framework (//UE4/Dev-Framework) @ 3205744

Change 3205887 on 2016/11/21 by Jurre.deBaare

	Fix for similar crash in blendspace editor like UE-38734

Change 3206121 on 2016/11/21 by [REDACTED_PERSON]

	PR #2935: Minor subtitle issues (Contributed by [REDACTED_PERSON])

Change 3206187 on 2016/11/21 by [REDACTED_PERSON]

	PR #2935: Minor subtitle issues (Contributed by [REDACTED_PERSON])
	Additional bits

Change 3206318 on 2016/11/21 by [REDACTED_PERSON]

	Fix Linux compiler whinging

Change 3206379 on 2016/11/21 by [REDACTED_PERSON]

	Fix crash when streaming in a sublevel with a child actor in it (4.14.1)

Change 3206591 on 2016/11/21 by [REDACTED_PERSON]

	Refactor restrictions to allow hidden and clarify disabled

Change 3206776 on 2016/11/21 by [REDACTED_PERSON]

	ForceFeedback component allows rumble events to be placed or spawned in to the world with attenuation settings that dictate how intensely the rumble pattern will be applied to the player based on their distance to the effect.
	ForceFeedback Attenuation settings can be defined via the content browser or directly on the component.

Change 3206901 on 2016/11/21 by [REDACTED_PERSON]

	Fix compile error in automation tests

Change 3207235 on 2016/11/22 by danny.bouimad

	Updated Map

Change 3207264 on 2016/11/22 by [REDACTED_PERSON]

	Disable bone editing in anim blueprint editor

Change 3207303 on 2016/11/22 by [REDACTED_PERSON]

	Clear material curve by setting it directly because the flag might not exist

Change 3207331 on 2016/11/22 by [REDACTED_PERSON]

	Fix overflow issues in SerializeProperties_DynamicArray_r. Also, fix crash from not ensuring properties were serialized successfully.

Change 3207357 on 2016/11/22 by [REDACTED_PERSON]

	Updating testcontent for pose drivers

Change 3207425 on 2016/11/22 by [REDACTED_PERSON]

	Fix frame count issue with montage

Change 3207478 on 2016/11/22 by [REDACTED_PERSON]

	Fix so that curve warning doesn't happen when your name is same.

Change 3207526 on 2016/11/22 by [REDACTED_PERSON]

	Fix crash when property restriction introduces a hidden entry

Change 3207731 on 2016/11/22 by danny.bouimad

	MoreUpdates

Change 3207764 on 2016/11/22 by [REDACTED_PERSON]

	#fix order of morphtarget to first process animation and then BP for slave component

Change 3207842 on 2016/11/22 by [REDACTED_PERSON]

	Fix it so ActiveStructRedirects are checked in addition to ActiveClassRedirects when serializing a raw UStruct reference, such as in a blueprint UStructProperty. This fixes issue with the attenuation settings struct rename, and should have always been working this way. ActiveClassRedirects will still work.

Change 3208202 on 2016/11/22 by [REDACTED_PERSON]

	The quat->rotator->quat->rotator conversions are introducing more error than in 4.13, so a pitch limit of -89.99 was too precise.

Change 3208510 on 2016/11/23 by [REDACTED_PERSON]

	Disable UBT Telemetry on internal builds #jira AN-1059
	#tests build a few different ways, add more diagnostics to clarify if the provider is being used.

Change 3208734 on 2016/11/23 by [REDACTED_PERSON]

	Change EnsureAllIndicesHaveHandles to try and maintain validity of as many of the handles as possible + Make FRichCurve key member private as it needs to stay in sync with map on base class

Change 3208782 on 2016/11/23 by [REDACTED_PERSON]

	Fixed material and vert count issues with skeletal to static mesh conversion

	Material remapping was not bein gbuilt, so material indices were overwitten inappropriately.
	Vertex tangentY was being recalculated incorrectly (discarding the W component when transformed), so vertices were not correctly re-merged later in the static mesh build phase.

Change 3208798 on 2016/11/23 by [REDACTED_PERSON]

	UE-38478 - Fix collision on procmesh created in BeginPlay in cooked builds

Change 3208801 on 2016/11/23 by Jurre.deBaare

	Hidden Material References from Mesh Components Fix
	#fix forgot to mark the renderstate dirty and wrapped it to only apply when overridematerials actually contain something

Change 3208807 on 2016/11/23 by [REDACTED_PERSON]

	CIS fix

Change 3208824 on 2016/11/23 by danny.bouimad

	More content updates for Testing

Change 3208827 on 2016/11/23 by [REDACTED_PERSON]

	Removing Old Pose driver Testassets I created awhile ago.

Change 3209026 on 2016/11/23 by [REDACTED_PERSON]

	CIS Fix for FRichCurve

Change 3209083 on 2016/11/23 by [REDACTED_PERSON]

	Don't crash if after an undo the previously selected object no longer exists (4.14.1)

Change 3209085 on 2016/11/23 by [REDACTED_PERSON]

	Don't crash if a negative length passed in to UKismetStringLibrary::GetSubstring (4.14.1)

Change 3209124 on 2016/11/23 by [REDACTED_PERSON]

	From PR #2955

Change 3209231 on 2016/11/23 by [REDACTED_PERSON]

	Auto removal

Change 3209232 on 2016/11/23 by [REDACTED_PERSON]

	GetComponents now optionally can include components in Child Actors

Change 3209233 on 2016/11/23 by [REDACTED_PERSON]

	ParseIntoArray resets instead of empty

Change 3209235 on 2016/11/23 by [REDACTED_PERSON]

	Allow child actor components to be selected in viewports
	Fix selection highlight not working on nested child actors

Change 3209247 on 2016/11/23 by [REDACTED_PERSON]

	Merging //UE4/Dev-Main to Dev-Framework (//UE4/Dev-Framework) @ 3209194

Change 3209299 on 2016/11/23 by [REDACTED_PERSON]

	Use MoveTemp to reduce some memory churn in graph schema actions

Change 3209347 on 2016/11/23 by [REDACTED_PERSON]

	Don't dispatch a tick function that had been scheduled but has been disabled before being executed.

Change 3209507 on 2016/11/23 by [REDACTED_PERSON]

	From PR #2908

Change 3209882 on 2016/11/24 by [REDACTED_PERSON]

	Copy-to-array now works with the fast path

	Refactored the copy record generation/validation code to be clearer with better seperation of concerns.
	Made sure we always properly generate a full exec chain for our events, despite some other them potentially using the fast path (this may have been a bug waiting to happen).
	Fixed a potentiual bug with sub anim instances were potentiall fast path non-array properties were skipped.
	Added tests for fast path validity to EditorTests project. Assets to follow.

Change 3209884 on 2016/11/24 by [REDACTED_PERSON]

	File I missed

Change 3209885 on 2016/11/24 by [REDACTED_PERSON]

	Support assets for fast path tests

Change 3209939 on 2016/11/24 by [REDACTED_PERSON]

	Fixed anim blueprint compiler not following reroute nodes when building cached pose fragment list

Change 3209941 on 2016/11/24 by Jurre.deBaare

	Removing and readding a point to the Anim Offset graph results in the animation to not preview correctly.
	#fix make sure that when we delete a sample point we reset the preview base pose
	#misc changed how the preview base pose is determined/updated

Change 3209942 on 2016/11/24 by [REDACTED_PERSON]

	Fixed transactions being made when setting bone space in details panel

	Also added reset to defaults to allow easy removal of bone modifications.

Change 3209945 on 2016/11/24 by james.cobbett

	Test assets for Pose Snapshot Test Case

Change 3210239 on 2016/11/25 by [REDACTED_PERSON]

	Making Navmesh react to changes done to static mesh's collision setup via the SM Editor #UE4

Change 3210279 on 2016/11/25 by [REDACTED_PERSON]

	Fixed anim sub-instances only allowing one pin to work when any pin required a call out to the VM for evaluation

Change 3210288 on 2016/11/25 by danny.bouimad

	Cleaned up Pose Driver Anim BP's

Change 3210334 on 2016/11/25 by [REDACTED_PERSON]

	Fixed preview mesh references getting broken in physics assets when renaming the preview mesh asset. Added explicit reference collection for the TAssetPtr

Change 3210349 on 2016/11/25 by [REDACTED_PERSON]

	UE-35783 Fix scrolling in PoseAsset editor panels

Change 3210356 on 2016/11/25 by [REDACTED_PERSON]

	UE-38420 Disable 'Convert to Static Mesh' option if no MeshComponents selected (e.g. cables)

Change 3210357 on 2016/11/25 by Jurre.deBaare

	Numeric textbox value label incorrect for aimoffset/blendspaces in grid
	#fix change lambda capture type (was referencing local variable)

Change 3210358 on 2016/11/25 by Jurre.deBaare

	Crash Generating Proxy Mesh with Transition Screen Size set to 1
	#fix 1.0 was not included within the possible range

Change 3210364 on 2016/11/25 by [REDACTED_PERSON]

	Improve BuildVertexBuffers to use stride and avoid copying colors

Change 3210371 on 2016/11/25 by Jurre.deBaare

	You can no longer enable tooltip display when using anim offset
	#fix Added back ability to show advanced preview sample weighting to tooltip under CTRL down

	It's not clear that the user has to hold shift to preview in blend spaces
	#fix Preview value is now set by default and has a tooltip state, this will inform the user how to move the preview value

	#misc refactored out some duplicate code :)

Change 3210387 on 2016/11/25 by james.cobbett

	Updating test asset

Change 3210550 on 2016/11/26 by [REDACTED_PERSON]

	Merging //UE4/Dev-Main to Dev-Framework (//UE4/Dev-Framework) @ 3209927

	Brings IWYU in and required substantial fixups

Change 3210551 on 2016/11/26 by [REDACTED_PERSON]

	Delete empty cpp files

Change 3211002 on 2016/11/28 by [REDACTED_PERSON]

	added navigation update on editting volume's brush
	#ue4

Change 3211011 on 2016/11/28 by [REDACTED_PERSON]

	Roll back CL# 3210334 as it is causing deadlocks during GC

Change 3211039 on 2016/11/28 by Jurre.deBaare

	Merge Actors tool is splitting every vertex on spline meshes, causing hard edged vertex colors.
	#fix prevent using the wedge map when propagating spline mesh vertex colours

Change 3211053 on 2016/11/28 by [REDACTED_PERSON]

	Make sure objects without simple collision do not simulate. Fixes crash when two trimesh only objects collide

Change 3211101 on 2016/11/28 by mason.seay

	Adjusting trigger collision so it can't be triggered by projectiles

Change 3211171 on 2016/11/28 by Jurre.deBaare

	Previewing outside of Blendspace Graph points causes unexpected weighting
	Second Animation Sample added to AimOffset or Blendspace swaps with the first sample

	#fix Changed behaviour for calculating blendspace grid weighting for one, two or colinear triangles
	- One: fill grid weights to single sample
	- Two: find closest point on line between the two samples for the grid point, and weight according to the distance on the line
	- Colinear: find two closest samples and apply behaviour above
	#misc rename variables to make the code more clear and correct

Change 3211491 on 2016/11/28 by [REDACTED_PERSON]

	Provide proper tooltip for GetParentActor/Component
	Expose GetAttachParentActor/SocketName to blueprints
	De-virtualize Actor GetAttach... functions

Change 3211570 on 2016/11/28 by [REDACTED_PERSON]

	Title doesn't update when asset is being dropped

Change 3211766 on 2016/11/28 by [REDACTED_PERSON]

	Remove warning when a constraint has two empty components. This can be a valid usecase for when components are determined dynamically.

Change 3211938 on 2016/11/28 by [REDACTED_PERSON]

	CSV's for testing gameplay tags

Change 3212090 on 2016/11/28 by [REDACTED_PERSON]

	Expose angular SLERP drive to blueprints

Change 3212102 on 2016/11/28 by [REDACTED_PERSON]

	Fix shadow variable issue

Change 3212182 on 2016/11/28 by [REDACTED_PERSON]

	PR #2902: Fix last collision preset display (Contributed by [REDACTED_PERSON])

Change 3212196 on 2016/11/28 by dan.reynolds

	AEOverview Update:

	Minor tweaks and fixes

	Added Attenuation Curve Tests

	Renamed SC to SCLA for Sound Class prefix

	WIP SCON (Sound Concurrency)

Change 3212347 on 2016/11/28 by [REDACTED_PERSON]

	Fix chained tag redirectors to work properly
	Const fixes and removed a bad error message spam, and fix rename message

Change 3212385 on 2016/11/28 by [REDACTED_PERSON]

	Avoid duplicate GetWorld() calls

Change 3212386 on 2016/11/28 by [REDACTED_PERSON]

	auto shoo

Change 3213018 on 2016/11/29 by [REDACTED_PERSON]

	Fix shadow variable for real

Change 3213037 on 2016/11/29 by [REDACTED_PERSON]

	Fix deprecation warnings

Change 3213039 on 2016/11/29 by [REDACTED_PERSON]

	Generalize logic for when a component prevents an Actor from auto destroying
	Add forcefeedback component to the components that will hold up the auto destroy of an actor

Change 3213088 on 2016/11/29 by [REDACTED_PERSON]

	Move significance manager out of experimental

Change 3213187 on 2016/11/29 by [REDACTED_PERSON]

	Add InsertDefaulted to mirror options available when Adding

Change 3213254 on 2016/11/29 by [REDACTED_PERSON]

	add auto-complete for showdebug forcefeedback

Change 3213260 on 2016/11/29 by [REDACTED_PERSON]

	Allow systems to inject auto-complete console entries

Change 3213276 on 2016/11/29 by [REDACTED_PERSON]

	add auto-complete entry for showdebug significancemanager

Change 3213331 on 2016/11/29 by [REDACTED_PERSON]

	Split SkeletalMesh skin weights into their own stream
	Remove unused FGPUSkinVertexColor struct
	Remove unused FSkeletalMeshVertexBuffer::bInfluencesByteSwapped bool
	Fix FSkeletalMeshMerge::GenerateLODModel to handle >4 weights
	Update friendly name for FColorVertexBuffer now it's used by skel mesh as well

Change 3213349 on 2016/11/29 by [REDACTED_PERSON]

	Fix tag rename feedback message

Change 3213355 on 2016/11/29 by [REDACTED_PERSON]

Change 3213406 on 2016/11/29 by [REDACTED_PERSON]

	Make sure body transforms are not set while the physx simulation is running.

Change 3213508 on 2016/11/29 by Jurre.deBaare

	When performing a merge actor on an actor merging multiple materials certain maps aren't generated
	#fix Apparently rendering out specular etc now outputs its value only to the red channel, so had to change how we populate the combined metallic/roughness/specular map

Change 3213557 on 2016/11/29 by [REDACTED_PERSON]

Change 3213634 on 2016/11/29 by [REDACTED_PERSON]

	Make sure if no shapes are found for vehicle wheels we create spheres and attach them to the actor.

Change 3213639 on 2016/11/29 by [REDACTED_PERSON]

	Fix from nvidia for vehicle suspension exploding when given a bad normal.

Change 3213812 on 2016/11/29 by [REDACTED_PERSON]

	UE-35925 Remove hard-coded asset<->animnode mapping, add SupportsAssetClass virtual instead

Change 3213824 on 2016/11/29 by [REDACTED_PERSON]

	Fix CIS

Change 3213873 on 2016/11/29 by [REDACTED_PERSON]

	Fix welded bodies not properly computing mass properties.

Change 3213950 on 2016/11/29 by [REDACTED_PERSON]

	Fixed navigation collision being generated wrong for StaticMeshes created from BSP #Orion

Change 3213951 on 2016/11/29 by [REDACTED_PERSON]

	Fixed perception system having issue with registering perception listener spawned in sublevels #UE4

Change 3214005 on 2016/11/29 by [REDACTED_PERSON]

	Fix mass kg override not propagating to blueprint instances.

Change 3214046 on 2016/11/29 by [REDACTED_PERSON]

	Duplicate all instanced subobjects, not just those that are editinlinenew
	Make AABrush.Brush instanced rather than export

Change 3214064 on 2016/11/29 by [REDACTED_PERSON]

	Use GetComponents directly where safe instead of copying in to an array

Change 3214116 on 2016/11/29 by [REDACTED_PERSON]

	Fix tooltip when dragging anim assets onto players

Change 3214136 on 2016/11/29 by [REDACTED_PERSON]

	Make it so moving bodies is immediate when in editor. Useful for editor tools that rely on physx data

Change 3214162 on 2016/11/29 by [REDACTED_PERSON]

	Fixed a bug in EnvQueryGenerator_SimpleGrid resuting in one extra column and row of points being generated #UE4

Change 3214177 on 2016/11/29 by [REDACTED_PERSON]

	Use correct SocketName (broken in CL#2695130)

Change 3214427 on 2016/11/29 by dan.reynolds

	AEOverview Update

	Fixed Attenuation tests when overlapping attenuation ranges between streamed levels

	Added Sound Concurrency Far then Prevent New testmap

	Removed some Sound Concurrency assets

Change 3214469 on 2016/11/29 by dan.reynolds

	AEOverview Update

	Added Sound Concurrency Test for Stop Farthest then Oldest

Change 3214842 on 2016/11/30 by Jurre.deBaare

	LookAt AimOffset in the Anim Graph causes character to explode
	#fix ensure that the source socket exists on the skeleton during compilation (as far as we can), and skip blendspace evaluation in case of it not being valid during runtime

Change 3214866 on 2016/11/30 by james.cobbett

	Updating Pose Snapshot test assets

Change 3214964 on 2016/11/30 by thomas.sarkanen

	Added test data for facial animtion curves

Change 3215015 on 2016/11/30 by Jurre.deBaare

	When a Aim Offset axis value is edited drastically the preview mesh will be deformed
	#fix change the way we change data when axis values are changed, simply remap normalized samples to new axis range
	#misc marked some data/functions editor only (not needed during runtime so reduces footprint a little bit)

Change 3215029 on 2016/11/30 by [REDACTED_PERSON]

	Fix CIS

Change 3215033 on 2016/11/30 by [REDACTED_PERSON]

	Add a delegate for when new classes are added via hotreload
	Change existing hotload class reinstancing delegates to be multicast

Change 3215048 on 2016/11/30 by [REDACTED_PERSON]

	Use getKinematicTarget whenever a body is kinematic.
	This should fix some edge cases in FBodyInstance where stale transforms may be used when operations are run in PrePhysics.

Change 3215052 on 2016/11/30 by [REDACTED_PERSON]

	Generalize the volume actor factory logic
	Create volume factories when hotreload adds a new volume class

Change 3215055 on 2016/11/30 by [REDACTED_PERSON]

	Probable fix for IOS CIS failure

Change 3215091 on 2016/11/30 by [REDACTED_PERSON]

	Easy alternative fix for blending two curves per bone. For now we just combine.

	To fix this properly - i.e. per bone to affect curve - it is very expensive process, so opting into this for 4.15.

Change 3215179 on 2016/11/30 by Jurre.deBaare

	Preview viewport should only use rendering features supported in project
	#fix replace the skylight with a sphere reflection component, this will not give image based lighting but does supply the user with a reflection map + intensity

Change 3215189 on 2016/11/30 by Jurre.deBaare

	CIS fix

Change 3215326 on 2016/11/30 by [REDACTED_PERSON]

	Regression caused by CL #3104976

Change 3215523 on 2016/11/30 by [REDACTED_PERSON]

	Fix cooking old skel meshes in commandlet - vertex buffer was not recreated so UpdateUVChannelData would crash

Change 3215539 on 2016/11/30 by [REDACTED_PERSON]

	Fix failure to cleanup objects in a hidden always loaded sub-level

Change 3215568 on 2016/11/30 by [REDACTED_PERSON]

	UE-39197 Delay node of 0.0 causes crash

Change 3215719 on 2016/11/30 by [REDACTED_PERSON]

	UE-39074 Audio related Client crash experienced on latest live build ++UT+Release-Next-CL-3193528

Change 3215773 on 2016/11/30 by [REDACTED_PERSON]

	PR #2819 : Fixed typo in SoundWave.h

Change 3215828 on 2016/11/30 by [REDACTED_PERSON]

	PR #2900: fixed a former change that overlooked the 2 character difference between 16 and 32. (Contributed by [REDACTED_PERSON])

Change 3215831 on 2016/11/30 by [REDACTED_PERSON]

	UE-36688 Add BlendOption (with CustomCurve) to PoseBlendNode

Change 3215904 on 2016/11/30 by [REDACTED_PERSON]

	Fix significance calculations

Change 3215955 on 2016/11/30 by [REDACTED_PERSON]

	UE-36791 Fix scaling of rotated convex elements, by baking element transform into cooked convex data.

Change 3215959 on 2016/11/30 by [REDACTED_PERSON]

	Remove LogTemp warning from FAnimBlueprintCompiler::FinishCompilingClass

Change 3216057 on 2016/11/30 by [REDACTED_PERSON]

	Don't reset expose on spawn properties when in a PIE world

Change 3216114 on 2016/11/30 by [REDACTED_PERSON]

	Move SkeletalMeshComponent and SkinnedMeshComponent functions out of SkeletalMesh.cpp into correct cpp files

Change 3216144 on 2016/11/30 by [REDACTED_PERSON]

	Fix FConstraintInstance scaling issues in FSkeletalMeshComponent::InitArticulated.

	InitArticulated uses the default Constraint Template from the Physics Asset a skeletal mesh is associated with.
	This caused issues if a skeletal mesh had bone scales that differed from those in the physics asset.

Change 3216148 on 2016/11/30 by [REDACTED_PERSON]

	Create test map and asset for Skeletal Mesh Component Scaling and Skeletal Mesh Uniform Import Scaling.

Change 3216160 on 2016/11/30 by [REDACTED_PERSON]

	Fixing a memory leak in concurrency management

Change 3216164 on 2016/11/30 by [REDACTED_PERSON]

	Move SkeletalMeshActor code into its own cpp file
	Fix CIS for SkeletalMeshComponent.cpp

Change 3216371 on 2016/11/30 by dan.reynolds

	AEOverview Update

	Minor tweaks

	Completed Sound Concurrency Rule Test Maps

	Added additional test files

Change 3216509 on 2016/11/30 by [REDACTED_PERSON]

	Fix missing include

Change 3216510 on 2016/11/30 by [REDACTED_PERSON]

	Code cleanup

Change 3216723 on 2016/12/01 by Jurre.deBaare

	When clearing a blend sample animation the animation will try and blend to the ref pose
	#fix do not delete sample when animation == nullptr but mark it as invalid, it then will be rendered in red on the grid and discarded during triangle/line generation
	#fix indice mapping for 2d blend spaces was incorrect before (luckily never caused an error)
	#misc weird whitespace changes

Change 3216745 on 2016/12/01 by Jurre.deBaare

	- Blend space triangulation was incorrect in some cases, due to refactor some data was not initialised.
	- UDN user was hitting a check within the triangle flipping behaviour

	#fix Revisited the conditions to determine whether or not a point lies within a triangles circumcircle
	#fix In case we cannot flip the current triangle we skip it and move onto the next one instead of putting in a hard check
	#misc refactored triangle flipping code to make it smaller (more readible)

Change 3216903 on 2016/12/01 by mason.seay

	Imported mesh for quick test

Change 3216904 on 2016/12/01 by Jurre.deBaare

	CIS Fix
	#fix replaced condition by both non-editor as editor valid one

Change 3216998 on 2016/12/01 by [REDACTED_PERSON]

	fixed AI slowing down on ramps due to 3D input vector being constrained by movement component
	#2998

Change 3217012 on 2016/12/01 by [REDACTED_PERSON]

	Checking in James' fix on drag/drop to replace assets

Change 3217031 on 2016/12/01 by james.cobbett

	Updating Pose Snapshot Assets. Again.

Change 3217033 on 2016/12/01 by [REDACTED_PERSON]

	Update bounds on all skel meshes when physics asset is changed

Change 3217181 on 2016/12/01 by [REDACTED_PERSON]

	Fix imported animations containing a black thumbnail

Change 3217183 on 2016/12/01 by [REDACTED_PERSON]

	Add some extra debugging code for future animation compression / ddc issues

Change 3217184 on 2016/12/01 by james.cobbett

	Fixing a test asset by checking a check box. Sigh.

Change 3217216 on 2016/12/01 by [REDACTED_PERSON]

	Undo part of CL 3217183. Will need to add this back differently.

Change 3217274 on 2016/12/01 by [REDACTED_PERSON]

	When serializing in an enum tagged property follow redirects

Change 3217419 on 2016/12/01 by james.cobbett

	Changes to test assets for more Pose Snapshot tests

Change 3217449 on 2016/12/01 by [REDACTED_PERSON]

	Adding new audio setting to disable EQ and reverb.

	Hooked up to XAudio2 (for now).

Change 3217513 on 2016/12/01 by [REDACTED_PERSON]

	Improve bWantsBeginPlay deprecation message

Change 3217620 on 2016/12/01 by mason.seay

	Updated test assets for HLOD

Change 3217872 on 2016/12/01 by [REDACTED_PERSON]

	UEFW-113 Adding master reverb to audio mixer

	- Added new submix editor to create new submixes
	- Created new default master submixes for reverb and EQ and master submixes
	- Fixed a number of minor issues found in auido mixer while working on feature

Change 3218053 on 2016/12/01 by [REDACTED_PERSON]

	Added mass debug rendering

Change 3218143 on 2016/12/01 by [REDACTED_PERSON]

	Fixing up reverb to support multi-channel (5.1 and 7.1) configurations.

	- Added default reverb send amount

Change 3218440 on 2016/12/01 by [REDACTED_PERSON]

	#ue4 - Made some static FNames const.

Change 3218715 on 2016/12/02 by james.cobbett

	Fixed bug in test asset.

Change 3218836 on 2016/12/02 by james.cobbett

	Fixing up test asset

Change 3218884 on 2016/12/02 by james.cobbett

	Moar test asset changes

Change 3218943 on 2016/12/02 by [REDACTED_PERSON]

	Make sure welded bodies include the center of mass offset. Note this also changes the COM nudge to be world space instead of local space

Change 3218955 on 2016/12/02 by [REDACTED_PERSON]

	Fix initialization order issues
	Remove monolithic includes
	Change signature to pass string by const ref

Change 3219149 on 2016/12/02 by [REDACTED_PERSON]

	Fix SetCollisionObjectType not working on skeletal mesh components

Change 3219162 on 2016/12/02 by [REDACTED_PERSON]

	Fix compile error when blend space on aim offset nodes is exposed as pin

Change 3219198 on 2016/12/02 by [REDACTED_PERSON]

	UEnum::FindValue/IndexByName will now correctly follow redirects

Change 3219340 on 2016/12/02 by [REDACTED_PERSON]

	#ue4 - Optimized and cleaned up some Actor methods related to location and rotation.

	- Inlined GetActorForwardVector(), GetActorUpVector(), GetActorRightVector(). Wrapped them to simply call the methods on USceneComponent rather than using a different approach to computing these vectors.
	- Inlined blueprint versions: K2_GetActorLocation(), K2_GetActorRotation(), K2_GetRootComponent().
	- Cleaned up template methods that are used to delay compilation of USceneComponent calls to make them private and prefix "Template" to their names so they don't show up in autocomplete for calls to the public methods.

Change 3219482 on 2016/12/02 by [REDACTED_PERSON]

	Fix crash when double deleting a clothing actor due to destroying USkeletalMesh before USkeletalMeshComponent.

Change 3219676 on 2016/12/02 by [REDACTED_PERSON]

	Make clearer that ref pose is from skeleton

Change 3219687 on 2016/12/02 by [REDACTED_PERSON]

	Supporting multi-channel reverb with automatic downmixing of input to stereo

Change 3219688 on 2016/12/02 by [REDACTED_PERSON]

	Fix crash when remapping additive animations after skeleton hierarchy change

Change 3219699 on 2016/12/02 by [REDACTED_PERSON]

	#ue4 - Fix template's use of old GetActorRotation() function.

Change 3219969 on 2016/12/02 by [REDACTED_PERSON]

	Modified from shelve Zak.Middleton made of PR #1885, after some more testing

Change 3220010 on 2016/12/02 by [REDACTED_PERSON]

	Fixing up sound class editor

Change 3220013 on 2016/12/02 by [REDACTED_PERSON]

	Deleting monolithic file

Change 3220249 on 2016/12/02 by [REDACTED_PERSON]

	Changing reverb settings parameter thread sync method

	- Switching to a simple ring buffer rather than using a crit sect

Change 3220251 on 2016/12/02 by [REDACTED_PERSON]

	Removing hard-coded audio mixer module name for the case when using -audiomixer argument,

	-added new entry to ini file that allows you to specify the audio mixer module name used for the platform.

Change 3221118 on 2016/12/05 by Jurre.deBaare

	Back out changelist 3220249 to fix CIS

Change 3221363 on 2016/12/05 by [REDACTED_PERSON]

	Change slot node category from Blends to Montage

Change 3221375 on 2016/12/05 by [REDACTED_PERSON]

	Change AGameModeBase::GetGameSessionClass to return GameSessionClass when set.

Change 3221402 on 2016/12/05 by [REDACTED_PERSON]

	Add sanitization code around PhsyX flags and refactor the ways flags are managed through a single code path.

Change 3221441 on 2016/12/05 by [REDACTED_PERSON]

	Fixed crash when reimporting a mesh when a different animation was open

Change 3221473 on 2016/12/05 by [REDACTED_PERSON]

	Get rid of auto.
	Use GetComponents directly instead of copying in to temporary arrays

Change 3221584 on 2016/12/05 by [REDACTED_PERSON]

	Fix CIS for Mac builds from CL-3221375

Change 3221631 on 2016/12/05 by [REDACTED_PERSON]

	Possible fix for rare marker sync crash on live servers

	#test ai match, animation seemed fine, no crashes

Change 3221660 on 2016/12/05 by mason.seay

	Resubmitting to add Viewport Bookmark

Change 3221683 on 2016/12/05 by [REDACTED_PERSON]

	Temp (but decent) fix to ARecastNavMesh::GetRandomPointInNavigableRadius sometimes retrieving invalid locations even if there's a valid piece of navmesh in the area #UE4

Change 3221750 on 2016/12/05 by [REDACTED_PERSON]

	Real CIS fix.

Change 3221917 on 2016/12/05 by [REDACTED_PERSON]

	Fix CIS for real this time.

Change 3222370 on 2016/12/05 by mason.seay

	Start of Gameplay Tag testmap

Change 3222396 on 2016/12/05 by [REDACTED_PERSON]

	UEFW-44 Implementing EQ master submix effect for audio mixer

	- New thread safe param setting temlate class (for setting EQ and Reverb params)
	- Hook up reverb submix effect to source voices
	- Implementation of FBiquad for biquad filter coefficients and audioprocessing
	- Implementation of Filter class which hold FBiquad instance per channel, computes coefficents once
	- Implementation of equalizer class which is a serial bank of filters set to ParametricEQ filter type

Change 3222425 on 2016/12/05 by [REDACTED_PERSON]

	Checking in missing files

Change 3222429 on 2016/12/05 by [REDACTED_PERSON]

	Last missing file!

Change 3222783 on 2016/12/05 by [REDACTED_PERSON]

	Update SkelMeshScaling map.

Change 3223173 on 2016/12/06 by [REDACTED_PERSON]

	Fix crash in thumbnail rendering when creating a new montage

Change 3223179 on 2016/12/06 by [REDACTED_PERSON]

	auto/NULL cleanup

Change 3223329 on 2016/12/06 by [REDACTED_PERSON]

	Fix (hard to explain) memory corruption

Change 3223334 on 2016/12/06 by [REDACTED_PERSON]

	Add HasBeenInitialized check inside AActor::InitializeComponents

Change 3223340 on 2016/12/06 by [REDACTED_PERSON]

	Refactor SkeletalMesh constraint scaling fixes. Add a check on bodies to ensure they are valid.

Change 3223372 on 2016/12/06 by [REDACTED_PERSON]

	Probably fix HTML5 CIS failure

Change 3223511 on 2016/12/06 by [REDACTED_PERSON]

	Fix Mac CIS shadow warning

Change 3223541 on 2016/12/06 by [REDACTED_PERSON]

	fixed missing NavCollision data in static meshes

Change 3223672 on 2016/12/06 by [REDACTED_PERSON]

	Remove unnecessary code from both customizations

Change 3223751 on 2016/12/06 by [REDACTED_PERSON]

	Properly remove components from their owner when manipulating through editinlinenew properties

Change 3223831 on 2016/12/06 by [REDACTED_PERSON]

	Don't allow deleting a leaf explicit tag whose implicit parent tags are still referenced and it is the only thing keeping them alive
	Add Tag Source to tooltip in management mode
	Fix RequestGameplayTagChildrenInDictionary to work properly

Change 3223862 on 2016/12/06 by [REDACTED_PERSON]

	Hide deprecated attach functions for all games not just Paragon

Change 3224003 on 2016/12/06 by [REDACTED_PERSON]

	Put behavior of player camera back to how it was prior to Ansel plugin support changes. Make photography only work a different way.

Change 3224602 on 2016/12/07 by Jurre.deBaare

	Crash on creating LODs with Medic
	#fix Added clamp for UVs -1024 to 1024

Change 3224604 on 2016/12/07 by Jurre.deBaare

	Fix for incorrect normal calculation in certain circumstances
	#fix Make sure we propagate the matrices to samples after we (re)calculated normals
	#fix Conditionally swap/inverse the vertex data buffers instead of always
	#fix Set preview mesh for alembic import animation sequences
	#misc removed commented out code and added debug code

Change 3224609 on 2016/12/07 by Jurre.deBaare

	Alembic Import Issues (skeletal) w. UVs and smoothing groups
	#fix Changed the way we populate smoothing group indices for alembic caches
	#misc removed commented out code, set base preview pose for alembic imported skeletal meshes / anim sequences

Change 3224783 on 2016/12/07 by [REDACTED_PERSON]

	Support per-instance skeletal mesh vertex color override

Change 3224784 on 2016/12/07 by [REDACTED_PERSON]

	Add skelmesh vert color override map. Fix my vert color material to work on skel mesh.

Change 3225131 on 2016/12/07 by Jurre.deBaare

	Crash when baking matrix animation when importing an alembic file as skeletal
	#fix condition whether or not to apply matrices had not been moved over in previous change

Change 3225491 on 2016/12/07 by [REDACTED_PERSON]

	- Morphtarget fix on the first frame

Change 3225597 on 2016/12/07 by mason.seay

	Updated materials on meshes to ones that don't have physical materials, also rebuilt lighting

Change 3225758 on 2016/12/07 by [REDACTED_PERSON]

	UE-39421 Fix for sound class graph bug

Change 3225957 on 2016/12/07 by [REDACTED_PERSON]

Change 3225967 on 2016/12/07 by [REDACTED_PERSON]

	Fix not removing link up cache when removed.

Change 3225990 on 2016/12/07 by [REDACTED_PERSON]

Change 3226123 on 2016/12/07 by [REDACTED_PERSON]

	Fix for sound class asset creation from within the sound class graph

Change 3226165 on 2016/12/07 by mason.seay

	Replaced skelmesh gun with static mesh cube

Change 3226336 on 2016/12/07 by [REDACTED_PERSON]

	Fixing up sound class replacement code.

	If you delete a sound class but replace with another, now it properly replaces sound classes in the sound class graphs without totally destroying them

Change 3226701 on 2016/12/08 by [REDACTED_PERSON]

	Merging //UE4/Dev-Main to Dev-Framework (//UE4/Dev-Framework) @ CL 3226613

Change 3226710 on 2016/12/08 by Jurre.deBaare

	Fix for alembic import crash
	#misc update num mesh samples and take into account user set start frame in case of skipping preroll frames

Change 3226834 on 2016/12/08 by Jurre.deBaare

	Fix for incorrect matrix samples being applied during Alembic cache importing
	#fix Change way we loop through samples and determine correct matrix and mesh sample indices

Change 3227330 on 2016/12/08 by Jurre.deBaare

	Temporary fix for animBP compilation error, underlying issue is causing the skeleton to not be fully loaded when we are validating the animation node. This makes the socket name check fail and consequently output a compilation error
	#UE-39499

	#fix Ensure that the skeleton is loaded by checking for RF_NeedPostLoad
	#misc corrected socket name output, removed unnecessary nullptr check

Change 3227575 on 2016/12/08 by [REDACTED_PERSON]

	Merging //UE4/Dev-Main to Dev-Framework (//UE4/Dev-Framework) @ 3227387

Change 3227602 on 2016/12/08 by [REDACTED_PERSON]

	Copyright 2016 to 2017 updates for new Framework files

diff --git a/Engine/Documentation/Source/Programming/Introduction/IntroToProgramming.INT.udn b/Engine/Documentation/Source/Programming/Introduction/IntroToProgramming.INT.udn
--- a/Engine/Documentation/Source/Programming/Introduction/IntroToProgramming.INT.udn
+++ b/Engine/Documentation/Source/Programming/Introduction/IntroToProgramming.INT.udn
@@ -46,11 +46,13 @@ Once you choose to create the class, the wizard will generate the files and open
 	public:	
 		// Sets default values for this actor's properties
 		AMyActor();
-		// Called when the game starts or when spawned
-		virtual void BeginPlay() override;
-		
+	
 		// Called every frame
 		virtual void Tick( float DeltaSeconds ) override;
+
+	protected:
+		// Called when the game starts or when spawned
+		virtual void BeginPlay() override;
 	};
 
 The class wizard generates your class with **BeginPlay()** and **Tick()** specified as overloads. **BeginPlay()** is an event that lets you know the Actor has entered the game in a playable state. This is a good place to initiate gameplay logic for your class. **Tick()** is called once per frame with the amount of elapsed time since the last call passed in. There you can do any recurring logic. However if you do not need that functionality, it is best to remove it to save yourself a small amount of performance. If you remove it, make sure to remove the line in the constructor that indicated ticking should occur. The constructor below contains the line in question.
diff --git a/Engine/Documentation/Source/Programming/QuickStart/3/QuickStart_3.INT.udn b/Engine/Documentation/Source/Programming/QuickStart/3/QuickStart_3.INT.udn
--- a/Engine/Documentation/Source/Programming/QuickStart/3/QuickStart_3.INT.udn
+++ b/Engine/Documentation/Source/Programming/QuickStart/3/QuickStart_3.INT.udn
@@ -132,9 +132,11 @@ We're now ready to create objects in **Unreal Editor** based on our code. For re
 				// Sets default values for this actor's properties
 				AFloatingActor();
 
+			protected:
 				// Called when the game starts or when spawned
 				virtual void BeginPlay() override;
 	
+			public:
 				// Called every frame
 				virtual void Tick( float DeltaSeconds ) override;
 	
diff --git a/Engine/Documentation/Source/Programming/QuickStart/4/QuickStart_4.INT.udn b/Engine/Documentation/Source/Programming/QuickStart/4/QuickStart_4.INT.udn
--- a/Engine/Documentation/Source/Programming/QuickStart/4/QuickStart_4.INT.udn
+++ b/Engine/Documentation/Source/Programming/QuickStart/4/QuickStart_4.INT.udn
@@ -68,9 +68,11 @@ Version: 4.9
 				// Sets default values for this actor's properties
 				AFloatingActor();
 
+			protected:
 				// Called when the game starts or when spawned
 				virtual void BeginPlay() override;
 	
+			public:
 				// Called every frame
 				virtual void Tick( float DeltaSeconds ) override;
 	
diff --git a/Engine/Documentation/Source/Programming/Slate/InGameUI/QuickStart/SlateInGameUIQuickStart.INT.udn b/Engine/Documentation/Source/Programming/Slate/InGameUI/QuickStart/SlateInGameUIQuickStart.INT.udn
--- a/Engine/Documentation/Source/Programming/Slate/InGameUI/QuickStart/SlateInGameUIQuickStart.INT.udn
+++ b/Engine/Documentation/Source/Programming/Slate/InGameUI/QuickStart/SlateInGameUIQuickStart.INT.udn
@@ -101,8 +101,12 @@ Next I added a new GameMode to represent the state of the player being at the ma
 	{
 		GENERATED_UCLASS_BODY()
 
+	protected:
+
 		virtual void BeginPlay() override;
 
+	public:
+
 		TSharedPtr<MainMenu> MainMenuPtr;
 	
 	};
diff --git a/Engine/Documentation/Source/Programming/Tutorials/AutoCamera/2/AutoCamera_2.INT.udn b/Engine/Documentation/Source/Programming/Tutorials/AutoCamera/2/AutoCamera_2.INT.udn
--- a/Engine/Documentation/Source/Programming/Tutorials/AutoCamera/2/AutoCamera_2.INT.udn
+++ b/Engine/Documentation/Source/Programming/Tutorials/AutoCamera/2/AutoCamera_2.INT.udn
@@ -101,9 +101,11 @@ No further code is needed. We can set up our CameraDirector in our world now.
 				// Sets default values for this actor's properties
 				ACameraDirector();
 
+			protected:
 				// Called when the game starts or when spawned
 				virtual void BeginPlay() override;
 	
+			public:
 				// Called every frame
 				virtual void Tick( float DeltaSeconds ) override;
 
diff --git a/Engine/Documentation/Source/Programming/Tutorials/AutoCamera/3/AutoCamera_3.INT.udn b/Engine/Documentation/Source/Programming/Tutorials/AutoCamera/3/AutoCamera_3.INT.udn
--- a/Engine/Documentation/Source/Programming/Tutorials/AutoCamera/3/AutoCamera_3.INT.udn
+++ b/Engine/Documentation/Source/Programming/Tutorials/AutoCamera/3/AutoCamera_3.INT.udn
@@ -73,9 +73,11 @@ We now have a system that moves the player's camera based purely on game logic.
 				// Sets default values for this actor's properties
 				ACameraDirector();
 
+			protected:
 				// Called when the game starts or when spawned
 				virtual void BeginPlay() override;
 	
+			public:
 				// Called every frame
 				virtual void Tick( float DeltaSeconds ) override;
 
diff --git a/Engine/Documentation/Source/Programming/Tutorials/Components/1/Components_1.INT.udn b/Engine/Documentation/Source/Programming/Tutorials/Components/1/Components_1.INT.udn
--- a/Engine/Documentation/Source/Programming/Tutorials/Components/1/Components_1.INT.udn
+++ b/Engine/Documentation/Source/Programming/Tutorials/Components/1/Components_1.INT.udn
@@ -116,9 +116,11 @@ Our new Pawn now has a collection of useful Components attached to it, and is re
 				// Sets default values for this pawn's properties
 				ACollidingPawn();
 
+			protected:
 				// Called when the game starts or when spawned
 				virtual void BeginPlay() override;
 	
+			public:
 				// Called every frame
 				virtual void Tick( float DeltaSeconds ) override;
 
diff --git a/Engine/Documentation/Source/Programming/Tutorials/Components/2/Components_2.INT.udn b/Engine/Documentation/Source/Programming/Tutorials/Components/2/Components_2.INT.udn
--- a/Engine/Documentation/Source/Programming/Tutorials/Components/2/Components_2.INT.udn
+++ b/Engine/Documentation/Source/Programming/Tutorials/Components/2/Components_2.INT.udn
@@ -89,9 +89,11 @@ We have just defined our input configuration and created a custom Pawn Movement
 				// Sets default values for this pawn's properties
 				ACollidingPawn();
 
+			protected:
 				// Called when the game starts or when spawned
 				virtual void BeginPlay() override;
 	
+			public:
 				// Called every frame
 				virtual void Tick( float DeltaSeconds ) override;
 
diff --git a/Engine/Documentation/Source/Programming/Tutorials/Components/3/Components_3.INT.udn b/Engine/Documentation/Source/Programming/Tutorials/Components/3/Components_3.INT.udn
--- a/Engine/Documentation/Source/Programming/Tutorials/Components/3/Components_3.INT.udn
+++ b/Engine/Documentation/Source/Programming/Tutorials/Components/3/Components_3.INT.udn
@@ -90,9 +90,11 @@ With our Pawn Movement Component's behavior defined, we can now write the code t
 				// Sets default values for this pawn's properties
 				ACollidingPawn();
 
+			protected:
 				// Called when the game starts or when spawned
 				virtual void BeginPlay() override;
 	
+			public:
 				// Called every frame
 				virtual void Tick( float DeltaSeconds ) override;
 
diff --git a/Engine/Documentation/Source/Programming/Tutorials/Components/4/Components_4.INT.udn b/Engine/Documentation/Source/Programming/Tutorials/Components/4/Components_4.INT.udn
--- a/Engine/Documentation/Source/Programming/Tutorials/Components/4/Components_4.INT.udn
+++ b/Engine/Documentation/Source/Programming/Tutorials/Components/4/Components_4.INT.udn
@@ -131,9 +131,11 @@ Our programming work is done, and we can now place our custom Pawn in the world
 				// Sets default values for this pawn's properties
 				ACollidingPawn();
 
+			protected:
 				// Called when the game starts or when spawned
 				virtual void BeginPlay() override;
 	
+			public:
 				// Called every frame
 				virtual void Tick( float DeltaSeconds ) override;
 
diff --git a/Engine/Documentation/Source/Programming/Tutorials/Components/5/Components_5.INT.udn b/Engine/Documentation/Source/Programming/Tutorials/Components/5/Components_5.INT.udn
--- a/Engine/Documentation/Source/Programming/Tutorials/Components/5/Components_5.INT.udn
+++ b/Engine/Documentation/Source/Programming/Tutorials/Components/5/Components_5.INT.udn
@@ -65,9 +65,11 @@ version:4.9
 				// Sets default values for this pawn's properties
 				ACollidingPawn();
 
+			protected:
 				// Called when the game starts or when spawned
 				virtual void BeginPlay() override;
 	
+			public:
 				// Called every frame
 				virtual void Tick( float DeltaSeconds ) override;
 
diff --git a/Engine/Documentation/Source/Programming/Tutorials/FirstPersonShooter/2/3/FirstPersonShooter_2_3.INT.udn b/Engine/Documentation/Source/Programming/Tutorials/FirstPersonShooter/2/3/FirstPersonShooter_2_3.INT.udn
--- a/Engine/Documentation/Source/Programming/Tutorials/FirstPersonShooter/2/3/FirstPersonShooter_2_3.INT.udn
+++ b/Engine/Documentation/Source/Programming/Tutorials/FirstPersonShooter/2/3/FirstPersonShooter_2_3.INT.udn
@@ -62,9 +62,11 @@ Now that you've set-up Axis mappings for your `FPSCharacter`, you can switch to
 			// Sets default values for this character's properties
 			AFPSCharacter();
 		
+		protected:
 			// Called when the game starts or when spawned
 			virtual void BeginPlay() override;
 			
+		public:
 			// Called every frame
 			virtual void Tick( float DeltaSeconds ) override;
 		
diff --git a/Engine/Documentation/Source/Programming/Tutorials/FirstPersonShooter/2/5/FirstPersonShooter_2_5.INT.udn b/Engine/Documentation/Source/Programming/Tutorials/FirstPersonShooter/2/5/FirstPersonShooter_2_5.INT.udn
--- a/Engine/Documentation/Source/Programming/Tutorials/FirstPersonShooter/2/5/FirstPersonShooter_2_5.INT.udn
+++ b/Engine/Documentation/Source/Programming/Tutorials/FirstPersonShooter/2/5/FirstPersonShooter_2_5.INT.udn
@@ -87,9 +87,11 @@ Go back to Visual Studio to add code to your FPSCharacter class.
 			// Sets default values for this character's properties
 			AFPSCharacter();
 		
+		protected:
 			// Called when the game starts or when spawned
 			virtual void BeginPlay() override;
 			
+		public:
 			// Called every frame
 			virtual void Tick( float DeltaSeconds ) override;
 		
diff --git a/Engine/Documentation/Source/Programming/Tutorials/FirstPersonShooter/2/7/FirstPersonShooter_2_7.INT.udn b/Engine/Documentation/Source/Programming/Tutorials/FirstPersonShooter/2/7/FirstPersonShooter_2_7.INT.udn
--- a/Engine/Documentation/Source/Programming/Tutorials/FirstPersonShooter/2/7/FirstPersonShooter_2_7.INT.udn
+++ b/Engine/Documentation/Source/Programming/Tutorials/FirstPersonShooter/2/7/FirstPersonShooter_2_7.INT.udn
@@ -51,9 +51,11 @@ At the end of the previous step, the default camera was positioned inside the me
 			// Sets default values for this character's properties
 			AFPSCharacter();
 		
+		protected:
 			// Called when the game starts or when spawned
 			virtual void BeginPlay() override;
 			
+		public:		
 			// Called every frame
 			virtual void Tick( float DeltaSeconds ) override;
 		
diff --git a/Engine/Documentation/Source/Programming/Tutorials/FirstPersonShooter/2/8/FirstPersonShooter_2_8.INT.udn b/Engine/Documentation/Source/Programming/Tutorials/FirstPersonShooter/2/8/FirstPersonShooter_2_8.INT.udn
--- a/Engine/Documentation/Source/Programming/Tutorials/FirstPersonShooter/2/8/FirstPersonShooter_2_8.INT.udn
+++ b/Engine/Documentation/Source/Programming/Tutorials/FirstPersonShooter/2/8/FirstPersonShooter_2_8.INT.udn
@@ -160,10 +160,12 @@ Before continuing, download and extract our sample mesh from the following link:
 			public:
 				// Sets default values for this character's properties
 				AFPSCharacter();
-			
+
+			protected:			
 				// Called when the game starts or when spawned
 				virtual void BeginPlay() override;
-				
+			
+			public:	
 				// Called every frame
 				virtual void Tick( float DeltaSeconds ) override;
 			
diff --git a/Engine/Documentation/Source/Programming/Tutorials/FirstPersonShooter/3/1/FirstPersonShooter_3_1.INT.udn b/Engine/Documentation/Source/Programming/Tutorials/FirstPersonShooter/3/1/FirstPersonShooter_3_1.INT.udn
--- a/Engine/Documentation/Source/Programming/Tutorials/FirstPersonShooter/3/1/FirstPersonShooter_3_1.INT.udn
+++ b/Engine/Documentation/Source/Programming/Tutorials/FirstPersonShooter/3/1/FirstPersonShooter_3_1.INT.udn
@@ -84,9 +84,11 @@ Now that we've set-up our character, it's time to implement a projectile weapon
 			// Sets default values for this actor's properties
 			AFPSProjectile();
 		
+		protected:
 			// Called when the game starts or when spawned
 			virtual void BeginPlay() override;
 			
+		public:
 			// Called every frame
 			virtual void Tick( float DeltaSeconds ) override;
 		
@@ -173,9 +175,11 @@ Now that we've set-up our character, it's time to implement a projectile weapon
 			// Sets default values for this actor's properties
 			AFPSProjectile();
 		
+		protected:
 			// Called when the game starts or when spawned
 			virtual void BeginPlay() override;
 			
+		public:
 			// Called every frame
 			virtual void Tick( float DeltaSeconds ) override;
 		
@@ -274,9 +278,11 @@ Now that we've set-up our character, it's time to implement a projectile weapon
 			// Sets default values for this actor's properties
 			AFPSProjectile();
 		
+		protected:
 			// Called when the game starts or when spawned
 			virtual void BeginPlay() override;
 			
+		public:
 			// Called every frame
 			virtual void Tick( float DeltaSeconds ) override;
 		
@@ -384,9 +390,11 @@ Now that we've set-up our character, it's time to implement a projectile weapon
 			// Sets default values for this character's properties
 			AFPSCharacter();
 		
+		protected:
 			// Called when the game starts or when spawned
 			virtual void BeginPlay() override;
 			
+		public:
 			// Called every frame
 			virtual void Tick( float DeltaSeconds ) override;
 		
@@ -590,9 +598,11 @@ Ultimately, you'll be able to calculate an initial location for the projectile b
 			// Sets default values for this character's properties
 			AFPSCharacter();
 		
+		protected:
 			// Called when the game starts or when spawned
 			virtual void BeginPlay() override;
 			
+		public:
 			// Called every frame
 			virtual void Tick( float DeltaSeconds ) override;
 		
diff --git a/Engine/Documentation/Source/Programming/Tutorials/FirstPersonShooter/3/4/FirstPersonShooter_3_4.INT.udn b/Engine/Documentation/Source/Programming/Tutorials/FirstPersonShooter/3/4/FirstPersonShooter_3_4.INT.udn
--- a/Engine/Documentation/Source/Programming/Tutorials/FirstPersonShooter/3/4/FirstPersonShooter_3_4.INT.udn
+++ b/Engine/Documentation/Source/Programming/Tutorials/FirstPersonShooter/3/4/FirstPersonShooter_3_4.INT.udn
@@ -53,9 +53,11 @@ Now that we can detect the projectile's collision interactions, we can determine
 			// Sets default values for this actor's properties
 			AFPSProjectile();
 		
+		protected:
 			// Called when the game starts or when spawned
 			virtual void BeginPlay() override;
 			
+		public:
 			// Called every frame
 			virtual void Tick( float DeltaSeconds ) override;
 		
diff --git a/Engine/Documentation/Source/Programming/Tutorials/FirstPersonShooter/3/5/FirstPersonShooter_3_5.INT.udn b/Engine/Documentation/Source/Programming/Tutorials/FirstPersonShooter/3/5/FirstPersonShooter_3_5.INT.udn
--- a/Engine/Documentation/Source/Programming/Tutorials/FirstPersonShooter/3/5/FirstPersonShooter_3_5.INT.udn
+++ b/Engine/Documentation/Source/Programming/Tutorials/FirstPersonShooter/3/5/FirstPersonShooter_3_5.INT.udn
@@ -174,9 +174,11 @@ C++ classes to Blueprints.
 				// Sets default values for this actor's properties.
 				AFPSProjectile();
 			
-				// Called when the game starts or when spawned.
+			protected:
+				// Called when the game starts or when spawned
 				virtual void BeginPlay() override;
-				
+			
+			public:
 				// Called every frame.
 				virtual void Tick( float DeltaSeconds ) override;
 			
@@ -276,9 +278,11 @@ C++ classes to Blueprints.
 				// Sets default values for this character's properties.
 				AFPSCharacter();
 			
-				// Called when the game starts or when spawned.
+			protected:
+				// Called when the game starts or when spawned
 				virtual void BeginPlay() override;
-				
+			
+			public:
 				// Called every frame.
 				virtual void Tick( float DeltaSeconds ) override;
 			
diff --git a/Engine/Documentation/Source/Programming/Tutorials/PlayerCamera/1/PlayerCamera_1.INT.udn b/Engine/Documentation/Source/Programming/Tutorials/PlayerCamera/1/PlayerCamera_1.INT.udn
--- a/Engine/Documentation/Source/Programming/Tutorials/PlayerCamera/1/PlayerCamera_1.INT.udn
+++ b/Engine/Documentation/Source/Programming/Tutorials/PlayerCamera/1/PlayerCamera_1.INT.udn
@@ -86,9 +86,11 @@ We now have a simple Pawn that will allow us to control our camera comfortably.
 				// Sets default values for this pawn's properties
 				APawnWithCamera();
 
+			protected:
 				// Called when the game starts or when spawned
 				virtual void BeginPlay() override;
-	
+			
+			public:
 				// Called every frame
 				virtual void Tick( float DeltaSeconds ) override;
 
diff --git a/Engine/Documentation/Source/Programming/Tutorials/PlayerCamera/2/PlayerCamera_2.INT.udn b/Engine/Documentation/Source/Programming/Tutorials/PlayerCamera/2/PlayerCamera_2.INT.udn
--- a/Engine/Documentation/Source/Programming/Tutorials/PlayerCamera/2/PlayerCamera_2.INT.udn
+++ b/Engine/Documentation/Source/Programming/Tutorials/PlayerCamera/2/PlayerCamera_2.INT.udn
@@ -75,9 +75,11 @@ Now that we have defined our input, we need to write some code to react to it. W
 				// Sets default values for this pawn's properties
 				APawnWithCamera();
 
+			protected:
 				// Called when the game starts or when spawned
 				virtual void BeginPlay() override;
-	
+			
+			public:
 				// Called every frame
 				virtual void Tick( float DeltaSeconds ) override;
 
diff --git a/Engine/Documentation/Source/Programming/Tutorials/PlayerCamera/3/PlayerCamera_3.INT.udn b/Engine/Documentation/Source/Programming/Tutorials/PlayerCamera/3/PlayerCamera_3.INT.udn
--- a/Engine/Documentation/Source/Programming/Tutorials/PlayerCamera/3/PlayerCamera_3.INT.udn
+++ b/Engine/Documentation/Source/Programming/Tutorials/PlayerCamera/3/PlayerCamera_3.INT.udn
@@ -174,9 +174,11 @@ version:4.9
 				// Sets default values for this pawn's properties
 				APawnWithCamera();
 
+			protected:
 				// Called when the game starts or when spawned
 				virtual void BeginPlay() override;
-	
+			
+			public:
 				// Called every frame
 				virtual void Tick( float DeltaSeconds ) override;
 
diff --git a/Engine/Documentation/Source/Programming/Tutorials/PlayerInput/1/PlayerInput_1.INT.udn b/Engine/Documentation/Source/Programming/Tutorials/PlayerInput/1/PlayerInput_1.INT.udn
--- a/Engine/Documentation/Source/Programming/Tutorials/PlayerInput/1/PlayerInput_1.INT.udn
+++ b/Engine/Documentation/Source/Programming/Tutorials/PlayerInput/1/PlayerInput_1.INT.udn
@@ -87,9 +87,11 @@ Now that we have a customized Pawn to react to our game's input, we'll need to d
 				// Sets default values
 				AMyPawn();
 
+			protected:
 				// Called when the game starts or when spawned
 				virtual void BeginPlay() override;
-	
+			
+			public:
 				// Called every frame
 				virtual void Tick( float DeltaSeconds ) override;
 
diff --git a/Engine/Documentation/Source/Programming/Tutorials/PlayerInput/2/PlayerInput_2.INT.udn b/Engine/Documentation/Source/Programming/Tutorials/PlayerInput/2/PlayerInput_2.INT.udn
--- a/Engine/Documentation/Source/Programming/Tutorials/PlayerInput/2/PlayerInput_2.INT.udn
+++ b/Engine/Documentation/Source/Programming/Tutorials/PlayerInput/2/PlayerInput_2.INT.udn
@@ -90,9 +90,11 @@ We're now ready to finish coding our MyPawn class in **Visual Studio**.
 				// Sets default values
 				AMyPawn();
 
+			protected:
 				// Called when the game starts or when spawned
 				virtual void BeginPlay() override;
-	
+			
+			public:
 				// Called every frame
 				virtual void Tick( float DeltaSeconds ) override;
 
diff --git a/Engine/Documentation/Source/Programming/Tutorials/PlayerInput/3/PlayerInput_3.INT.udn b/Engine/Documentation/Source/Programming/Tutorials/PlayerInput/3/PlayerInput_3.INT.udn
--- a/Engine/Documentation/Source/Programming/Tutorials/PlayerInput/3/PlayerInput_3.INT.udn
+++ b/Engine/Documentation/Source/Programming/Tutorials/PlayerInput/3/PlayerInput_3.INT.udn
@@ -133,9 +133,11 @@ version:4.9
 				// Sets default values
 				AMyPawn();
 
+			protected:
 				// Called when the game starts or when spawned
 				virtual void BeginPlay() override;
-
+			
+			public:
 				// Called every frame
 				virtual void Tick(float DeltaSeconds) override;
 
diff --git a/Engine/Documentation/Source/Programming/Tutorials/VariablesTimersEvents/1/VariablesTimersEvents_1.INT.udn b/Engine/Documentation/Source/Programming/Tutorials/VariablesTimersEvents/1/VariablesTimersEvents_1.INT.udn
--- a/Engine/Documentation/Source/Programming/Tutorials/VariablesTimersEvents/1/VariablesTimersEvents_1.INT.udn
+++ b/Engine/Documentation/Source/Programming/Tutorials/VariablesTimersEvents/1/VariablesTimersEvents_1.INT.udn
@@ -132,9 +132,11 @@ At this point, we've already created a simple class that uses a timer. Non-progr
 				// Sets default values for this actor's properties
 				ACountdown();
 
+			protected:
 				// Called when the game starts or when spawned
 				virtual void BeginPlay() override;
 	
+			public:
 				// Called every frame
 				virtual void Tick( float DeltaSeconds ) override;
 
diff --git a/Engine/Documentation/Source/Programming/Tutorials/VariablesTimersEvents/2/VariablesTimersEvents_2.INT.udn b/Engine/Documentation/Source/Programming/Tutorials/VariablesTimersEvents/2/VariablesTimersEvents_2.INT.udn
--- a/Engine/Documentation/Source/Programming/Tutorials/VariablesTimersEvents/2/VariablesTimersEvents_2.INT.udn
+++ b/Engine/Documentation/Source/Programming/Tutorials/VariablesTimersEvents/2/VariablesTimersEvents_2.INT.udn
@@ -111,9 +111,11 @@ We have now made a variable and a function accessible to, and alterable by, non-
 				// Sets default values for this actor's properties
 				ACountdown();
 
+			protected:
 				// Called when the game starts or when spawned
 				virtual void BeginPlay() override;
 	
+			public:
 				// Called every frame
 				virtual void Tick( float DeltaSeconds ) override;
 
diff --git a/Engine/Documentation/Source/Programming/Tutorials/VariablesTimersEvents/3/VariablesTimersEvents_3.INT.udn b/Engine/Documentation/Source/Programming/Tutorials/VariablesTimersEvents/3/VariablesTimersEvents_3.INT.udn
--- a/Engine/Documentation/Source/Programming/Tutorials/VariablesTimersEvents/3/VariablesTimersEvents_3.INT.udn
+++ b/Engine/Documentation/Source/Programming/Tutorials/VariablesTimersEvents/3/VariablesTimersEvents_3.INT.udn
@@ -108,9 +108,11 @@ version:4.9
 				// Sets default values for this actor's properties
 				ACountdown();
 
+			protected:
 				// Called when the game starts or when spawned
 				virtual void BeginPlay() override;
 	
+			public:
 				// Called every frame
 				virtual void Tick( float DeltaSeconds ) override;
 
