Copying //UE4/Dev-Documentation to Samples-Main (//UE4/Samples-Main)

diff --git a/Engine/Documentation/Source/Engine/AI/BehaviorTrees/BehaviorTrees.INT.udn b/Engine/Documentation/Source/Engine/AI/BehaviorTrees/BehaviorTrees.INT.udn
--- a/Engine/Documentation/Source/Engine/AI/BehaviorTrees/BehaviorTrees.INT.udn
+++ b/Engine/Documentation/Source/Engine/AI/BehaviorTrees/BehaviorTrees.INT.udn
@@ -3,7 +3,7 @@ Title: Behavior Trees
 Crumbs:%ROOT%, Engine, Gameplay/AI
 Description:Landing Page for all the documentation relating to Behavior Trees
 parent:Engine
-order:5
+order:6
 
 [VAR:Topic]
 [OBJECT:Topic]
diff --git a/Engine/Documentation/Source/Engine/AI/EnvironmentQuerySystem/EnvironmentQuerySystem.INT.udn b/Engine/Documentation/Source/Engine/AI/EnvironmentQuerySystem/EnvironmentQuerySystem.INT.udn
--- a/Engine/Documentation/Source/Engine/AI/EnvironmentQuerySystem/EnvironmentQuerySystem.INT.udn
+++ b/Engine/Documentation/Source/Engine/AI/EnvironmentQuerySystem/EnvironmentQuerySystem.INT.udn
@@ -3,7 +3,7 @@ Title: Environment Query System
 Crumbs:%ROOT%, Engine, Gameplay/AI
 Description:Landing Page for all the documentation relating to the Environment Query System.
 parent:Engine
-order:6
+order:7
 
 [VAR:Topic]
 [OBJECT:Topic]
diff --git a/Engine/Documentation/Source/Engine/Actors/FogEffects/FogEffects.INT.udn b/Engine/Documentation/Source/Engine/Actors/FogEffects/FogEffects.INT.udn
--- a/Engine/Documentation/Source/Engine/Actors/FogEffects/FogEffects.INT.udn
+++ b/Engine/Documentation/Source/Engine/Actors/FogEffects/FogEffects.INT.udn
@@ -2,6 +2,9 @@ Availability: Public
 Title:Fog Effects
 Crumbs: %ROOT%, Engine
 Description:Using fog to add ambiance to the world and set the mood of environments.
+parent:Engine/Rendering
+order:6
+type:landing
 
 [VAR:Topic]
 [OBJECT:Topic]
diff --git a/Engine/Documentation/Source/Engine/Actors/Merging/ActorMerging.INT.udn b/Engine/Documentation/Source/Engine/Actors/Merging/ActorMerging.INT.udn
--- a/Engine/Documentation/Source/Engine/Actors/Merging/ActorMerging.INT.udn
+++ b/Engine/Documentation/Source/Engine/Actors/Merging/ActorMerging.INT.udn
@@ -1,4 +1,4 @@
-Availability: Docs
+Availability: Public
 Title:Actor Merging
 Crumbs: %ROOT%, Engine, Engine/Actors, Engine/Actors/Merging
 Description:How to use the Actor Merge tool to combine multiple Static Meshes into one in UE4.
diff --git a/Engine/Documentation/Source/Engine/Animation/AnimBlueprints/AnimGraph/AnimGraph.INT.udn b/Engine/Documentation/Source/Engine/Animation/AnimBlueprints/AnimGraph/AnimGraph.INT.udn
--- a/Engine/Documentation/Source/Engine/Animation/AnimBlueprints/AnimGraph/AnimGraph.INT.udn
+++ b/Engine/Documentation/Source/Engine/Animation/AnimBlueprints/AnimGraph/AnimGraph.INT.udn
@@ -1,113 +1,90 @@
 Availability:Public
 Title:AnimGraph
 Crumbs: %ROOT%, Engine, Engine/Animation, Engine/Animation/AnimBlueprints
-Description:Graph used to sample, blend, and manipulate poses to be applied to SkeletalMeshes by the Animation Blueprint.
-version: 4.12
+Description:Graph used to sample, blend, and manipulate poses to be applied to Skeletal Meshes by the Animation Blueprint.
 related:Engine\Animation\AnimBlueprints
 related:Engine\Animation\AnimBlueprints\EventGraph
 related:Engine\Animation\AnimBlueprints\Interface
 related:Engine\Animation\AnimBlueprints\Creation
-
-[VAR:Topic]
-[OBJECT:Topic]
-	[PARAM:image]
-		![%Engine\Animation\AnimBlueprints\AnimGraph:title%](AnimGraphTopic.png)
-	[/PARAM]
-	[PARAM:icon]
-		![](%ROOT%/build_icon.png)(convert:false)
-	[/PARAM]
-	[PARAM:title]
-		%Engine\Animation\AnimBlueprints\AnimGraph:title%
-	[/PARAM]
-	[PARAM:description]
-		%Engine\Animation\AnimBlueprints\AnimGraph:description%
-	[/PARAM]
-	[PARAM:path]
-		[RELATIVE:Engine\Animation\AnimBlueprints\AnimGraph]
-	[/PARAM]
-[/OBJECT]
-[/VAR]
-[VAR:TopicCompact]
-[OBJECT:TopicCompact]
-	[PARAM:image]
-		![%Engine\Animation\AnimBlueprints\AnimGraph:title%](AnimGraphTopic.png)
-	[/PARAM]
-	[PARAM:icon]
-		![](%ROOT%/build_icon.png)(convert:false)
-	[/PARAM]
-	[PARAM:title]
-		%Engine\Animation\AnimBlueprints\AnimGraph:title%
-	[/PARAM]
-	[PARAM:description]
-		%Engine\Animation\AnimBlueprints\AnimGraph:description%
-	[/PARAM]
-	[PARAM:path]
-		[RELATIVE:Engine\Animation\AnimBlueprints\AnimGraph]
-	[/PARAM]
-[/OBJECT]
-[/VAR]
-
+version: 4.14
+type:reference
+topic-image:AnimGraphTopic.png
+tags:Animation
+tags:Animation Blueprint
+tags:Blueprints
 
 [TOC(start:2 end:3)]
 
 [EXCERPT:Intro]
-The **AnimGraph** is used to evaluate a final pose for the SkeletalMesh for the current frame. By default, 
-each Animation Blueprint has an AnimGraph which can have [**Animation Nodes**](Engine/Animation/AnimBlueprints/AnimGraph#AnimationNodes) placed within it 
-to sample AnimationSequences, perform animation blends, or control bone transforms using SkeletalControls. 
-The resulting pose is then applied to the SkeletalMesh each frame.
+The **AnimGraph** is used to evaluate a final pose for the Skeletal Mesh for the current frame. By default, 
+each Animation Blueprint has an AnimGraph which can have [Animation Nodes](Engine/Animation/NodeReference) placed within it 
+to sample Animation Sequences, perform animation blends or control bone transforms using Skeletal Controls. 
+The resulting pose is then applied to the Skeletal Mesh each frame.
 [/EXCERPT:Intro]
 
-![Anim Graph](animgraph.png)(w:720)
+![AnimGraph](animgraph.png)
 
-Inside the AnimGraph you can use values calculated in the **EventGraph** or elsewhere, likely in code, 
-and use those as inputs for BlendSpaces or other blend nodes. Of course, you can also just sample 
-AnimationSequences directly without any special blending. The AnimGraph also allows for the use of 
-SkeletalControls to modify bone positions, rotations, and scales directly.
+Inside the AnimGraph, you can use values calculated in the **EventGraph** or elsewhere, likely in code, 
+and use those as inputs for Blend Spaces or other blend nodes. Of course, you can also just sample 
+Animation Sequences directly without any special blending. The AnimGraph also allows for the use of 
+Skeletal Controls to modify bone positions, rotations, and scales directly.
 
 ## Execution Flow
 
 All graphs have a flow of execution shown as pulsing links between exec pins on the graph's nodes. For 
-regular graphs, such as the **EventGraph** of a Blueprint, this flow is visualized during play as it is 
-dependent on the firing of events. The **AnimGraph** is different in that it displays the flow of execution 
+regular graphs, such as the EventGraph of a Blueprint, this flow is visualized during play as it is 
+dependent on the firing of events. The AnimGraph is different in that it displays the flow of execution 
 at all times since it is not event-based and is being evaluated each frame.
 
 ![Blend - Alpha 0.0](exec_flow_0.png)(w:720)
 
-In the **AnimGraph**, the flow of execution represents poses being passed from one node to another. Some 
-nodes, such as [**Blends**](Engine/Animation/NodeReference/Blend), have multiple inputs and make a 
+In the AnimGraph, the flow of execution represents poses being passed from one node to another. Some 
+nodes, such as [Blends](Engine/Animation/NodeReference/Blend), have multiple inputs and make a 
 decision internally on which input is currently part of the flow of execution. This determination is usually 
 dependent on some external input, like the value passed to a data pin.
 
 In the example below, the same blend node from above switched the execution flow between its inputs based on 
-the value of its Alpha pin (the value is being manually set for simplicity):
+the value of its Alpha pin (which is being manually set):
 
 ![Blend - Alpha 1.0](exec_flow_1.png)(w:720)
 [REGION:caption]Blend with Alpha of 1.0 - Input B is used[/REGION]
 
 ## Working with the AnimGraph
 
-The **AnimGraph** works by evaluating a graph of nodes. Some of these nodes perform some specific operation 
-on one or more input poses, while others are used to access or sample other types of assets, such as AnimationSequences, 
-BlendSpaces, or AnimMontages. These assets and nodes, as well as how they fit into the **AnimGraph**, are described below.
+The AnimGraph works by evaluating a graph of nodes. Some of these nodes perform some specific operation 
+on one or more input poses, while others are used to access or sample other types of assets, such as Animation Sequences, 
+Blend Spaces or Anim Montages. State Machines, which contain their own graph networks can exist within the AnimGraph, and multiple State Machines can be blended together to produce a final output pose. 
+Please refer to each link below for information on how each works within the context of the AnimGraph. 
+
+* [Animation Blending](Engine/Animation/NodeReference/Blend) - can be accomplished through the use of various Blend Nodes. 
+* [Animation Nodes](Engine/Animation/NodeReference) - are used to perform operations on input poses, such as blending, direct bone manipulation and more. 
+* [Animation Sequences](Engine/Animation/Sequences/UserGuide) - can be placed directly within the AnimGraph to drive a Final Animation Pose. 
+* [Blend Spaces](Engine/Animation/Blendspaces/UserGuide) - can be used within the AnimGraph to blend poses based on variables. 
+* [Skeletal Controls](Engine/Animation/NodeReference/SkeletalControls) - are nodes which can be used to directly drive the bones within a Skeleton. 
+* [Space Conversion](Engine/Animation/NodeReference/SpaceConversion) - are nodes that can change whether a pose is in local-space or component-space. 
+* [State Machines](Engine/Animation/StateMachines) - are a series of graphs, rules and variables that determine an animation state for a character to enter. 
+* [Sync Groups](Engine\Animation\SyncGroups) - can be used to keep related animation that may have different lengths synchronized. 
 
-### AnimationSequences
+
+<!-- includes were just duplicating information when arriving at the linked page, commenting out. 
+### Animation Sequences
 
 [INCLUDE:Engine/Animation/Sequences/UserGuide#Intro]
 
 ![Animation Sequences](Engine/Animation/Sequences/UserGuide/CE_anim_sample.png)
 
 [REGION:note]
-See the [**Animation Sequence User Guide**](Engine/Animation/Sequences/UserGuide) for more details.
+See the [Animation Sequence User Guide](Engine/Animation/Sequences/UserGuide) for more details.
 [/REGION]
 
-### Blendspaces
+### Blend Spaces
 
 [INCLUDE:Engine/Animation/Blendspaces/UserGuide#Intro]
 
 ![Output Pose - Result](Engine/Animation/Blendspaces/UserGuide/bs_output_result.png)
 
 [REGION:note]
-See [**Using Blend Spaces**](Engine/Animation/Blendspaces/UserGuide) for more details.
+See [Using Blend Spaces](Engine/Animation/Blendspaces/UserGuide) for more details.
 [/REGION]
 
 ### Animation Nodes
@@ -115,7 +92,7 @@ See [**Using Blend Spaces**](Engine/Animation/Blendspaces/UserGuide) for more de
 [INCLUDE:Engine/Animation/NodeReference#Intro]
 
 [REGION:note]
-See [**Animation Node Reference**](Engine/Animation/NodeReference) for more descriptions of all available nodes.
+See [Animation Node Reference](Engine/Animation/NodeReference) for more descriptions of all available nodes.
 [/REGION]
 
 #### Animation Blending
@@ -125,7 +102,7 @@ See [**Animation Node Reference**](Engine/Animation/NodeReference) for more desc
 ![Two-way Blend](Engine/Animation/NodeReference/Blend/CE_blend_twoway.png)
 
 [REGION:note]
-See [**Blend Nodes**](Engine/Animation/NodeReference/Blend) for more information on the space conversion nodes.
+See [Blend Nodes](Engine/Animation/NodeReference/Blend) for more information on the space conversion nodes.
 [/REGION]
 
 #### Skeletal Controls
@@ -133,7 +110,7 @@ See [**Blend Nodes**](Engine/Animation/NodeReference/Blend) for more information
 [INCLUDE:Engine/Animation/NodeReference/SkeletalControls#intro]
 
 [REGION:note]
-See [**Skeletal Controls**](Engine/Animation/NodeReference/SkeletalControls) for more information on the various types of SkeletalControls available.
+See [Skeletal Controls](Engine/Animation/NodeReference/SkeletalControls) for more information on the various types of Skeletal Controls available.
 [/REGION]
 
 #### Space Conversion
@@ -145,7 +122,7 @@ See [**Skeletal Controls**](Engine/Animation/NodeReference/SkeletalControls) for
 [/REGION]
 
 [REGION:note]
-See [**Convert Space Nodes**](Engine/Animation/NodeReference/SpaceConversion) for more information on the space conversion nodes.
+See [Convert Space Nodes](Engine/Animation/NodeReference/SpaceConversion) for more information on the space conversion nodes.
 [/REGION]
 
 #### Sync Groups
@@ -155,18 +132,24 @@ See [**Convert Space Nodes**](Engine/Animation/NodeReference/SpaceConversion) fo
 ![Usage in graph](Engine\Animation\SyncGroups\AnimationSyncGroup.png)(w:480)
 
 [REGION:note]
-See [**Sync Groups**](Engine\Animation\SyncGroups) for more information on Sync Group usage and additional details.
+See [Sync Groups](Engine\Animation\SyncGroups) for more information on Sync Group usage and additional details.
 [/REGION]
+-->
 
 ## Thread Safety
 
-AnimGraph logic can be run on worker threads using the "Use Multi Threaded Animation Update" setting found in "Class Defaults". This can
-significantly improve the perfomance of games on multi-core machines by allowing animation work to complete in parallel with other
+AnimGraph logic can be run on worker threads using the "Use Multi-Threaded Animation Update" setting found in "Class Defaults". This can
+significantly improve the performance of games on multi-core machines by allowing animation work to complete in parallel with other
 updates.
 
 If this setting is enabled (it is by default) then the compiler will warn if 'unsafe' operations are performed in the AnimGraph. Generally,
 safe operations are:
 
-* Using most AnimGRaph nodes.
+* Using most AnimGraph nodes.
 * Directly accessing member variables of your Animation Blueprint (including members of structures).
-* Calling Blueprint 'pure' functions that dont modify state (such as most math functions).
+* Calling Blueprint 'pure' functions that don't modify state (such as most math functions).
+
+
+
+
+
diff --git a/Engine/Documentation/Source/Engine/Animation/AnimBlueprints/Creation/CreatingVimBlueprints.INT.udn b/Engine/Documentation/Source/Engine/Animation/AnimBlueprints/Creation/CreatingVimBlueprints.INT.udn
--- a/Engine/Documentation/Source/Engine/Animation/AnimBlueprints/Creation/CreatingVimBlueprints.INT.udn
+++ b/Engine/Documentation/Source/Engine/Animation/AnimBlueprints/Creation/CreatingVimBlueprints.INT.udn
@@ -2,54 +2,21 @@ Availability:Public
 Title:Animation Blueprint Creation
 Crumbs: %ROOT%, Engine, Engine/Animation, Engine/Animation/AnimBlueprints
 Description:Guide to creating new Animation Blueprints from Skeleton assets
-version: 4.12
 related:Engine\Animation\AnimBlueprints
 related:Engine\Animation\AnimBlueprints\AnimGraph
 related:Engine\Animation\AnimBlueprints\EventGraph
 related:Engine\Animation\AnimBlueprints\Interface
+version: 4.14
+type:reference
+topic-image:CreateAnimBP.png
+tags:Animation
+tags:Animation Blueprint
+tags:Basics
 
-[VAR:Topic]
-[OBJECT:Topic]
-	[PARAM:image]
-		![%Engine\Animation\AnimBlueprints\Creation:title%](CreateAnimBP.png)
-	[/PARAM]
-	[PARAM:icon]
-		![](%ROOT%/build_icon.png)(convert:false)
-	[/PARAM]
-	[PARAM:title]
-		%Engine\Animation\AnimBlueprints\Creation:title%
-	[/PARAM]
-	[PARAM:description]
-		%Engine\Animation\AnimBlueprints\Creation:description%
-	[/PARAM]
-	[PARAM:path]
-		[RELATIVE:Engine\Animation\AnimBlueprints\Creation]
-	[/PARAM]
-[/OBJECT]
-[/VAR]
-[VAR:TopicCompact]
-[OBJECT:TopicCompact]
-	[PARAM:image]
-		![%Engine\Animation\AnimBlueprints\Creation:title%](CreateAnimBP.png)
-	[/PARAM]
-	[PARAM:icon]
-		![](%ROOT%/build_icon.png)(convert:false)
-	[/PARAM]
-	[PARAM:title]
-		%Engine\Animation\AnimBlueprints\Creation:title%
-	[/PARAM]
-	[PARAM:description]
-		%Engine\Animation\AnimBlueprints\Creation:description%
-	[/PARAM]
-	[PARAM:path]
-		[RELATIVE:Engine\Animation\AnimBlueprints\Creation]
-	[/PARAM]
-[/OBJECT]
-[/VAR]
 
 ## Creating an Animation Blueprint Asset from a Skeleton
 
-When creating a new Animation Blueprint, you will be asked to pick a valid [**Skeleton**](Engine/Animation/Skeleton) for it to target.
+When creating a new Animation Blueprint, you will be asked to pick a valid [Skeleton](Engine/Animation/Skeleton) for it to target.
 
 1. Either **Right-click** in an empty space in the **Content Browser** or click the **Add New** button.
 
@@ -57,7 +24,7 @@ When creating a new Animation Blueprint, you will be asked to pick a valid [**Sk
 
 1. On the context menu under **Create Advanced Asset**, click **Animation**.
 
-1. Click **Animation Blueprint**.
+1. In the expanded menu, select the **Animation Blueprint** option.
 
 	![](CE_create_NoSkel.png)
 
@@ -65,7 +32,7 @@ When creating a new Animation Blueprint, you will be asked to pick a valid [**Sk
 
 	![](PickSkeleton.png)
 
-	You are also given the ability to select a [**Parent Class**](Engine/Blueprints/UserGuide/Types/ClassBlueprint) for the Animation Blueprint to use. 
+	You are also given the ability to select a [Parent Class](Engine/Blueprints/UserGuide/Types/ClassBlueprint) for the Animation Blueprint to use. 
 	Most of the time you will not need to select one and can use the default class. However, based on your circumstances, you can choose one if you need to. 
 
 1. The new asset will automatically be created in the **Content Browser**. 
@@ -73,14 +40,12 @@ When creating a new Animation Blueprint, you will be asked to pick a valid [**Sk
 	![](NewAnimBPCreated.png)
 
 
+	
+## Context Menu Creation of an Animation Blueprint
 
+You can also create an Animation Blueprint that targets a specific [Skeletal Mesh](Engine/Content/Types/SkeletalMeshes) or Skeleton by **Right-clicking** on it in the **Content Browser**. 
 
-
-## Creating an Animation Blueprint from a Skeletal Mesh
-
-You can also create an Animation Blueprint that targets a specific [**Skeletal Mesh**](Engine/Content/Types/SkeletalMeshes) asset. 
-
-1. **Right-click** on the Skeletal Mesh you want the Animation Blueprint to target and click **Create** > **Anim Assets** > **Anim Blueprint**.
+1. **Right-click** on the Skeletal Mesh (or Skeleton Asset) you want the Animation Blueprint to target and click **Create** > **Anim Assets** > **Anim Blueprint**.
 
 	![Create Animation Blueprint - Menu](CE_create_menu.png)
 
diff --git a/Engine/Documentation/Source/Engine/Animation/AnimBlueprints/EventGraph/EventGraph.INT.udn b/Engine/Documentation/Source/Engine/Animation/AnimBlueprints/EventGraph/EventGraph.INT.udn
--- a/Engine/Documentation/Source/Engine/Animation/AnimBlueprints/EventGraph/EventGraph.INT.udn
+++ b/Engine/Documentation/Source/Engine/Animation/AnimBlueprints/EventGraph/EventGraph.INT.udn
@@ -2,69 +2,33 @@ Availability:Public
 Title:EventGraph
 Crumbs: %ROOT%, Engine, Engine/Animation, Engine/Animation/AnimBlueprints
 Description:Event-based graph used to update the Animation Blueprint and calculate values for use in the AnimGraph.
-version: 4.12
 related:Engine\Animation\AnimBlueprints
 related:Engine\Animation\AnimBlueprints\AnimGraph
 related:Engine\Animation\AnimBlueprints\Interface
 related:Engine\Animation\AnimBlueprints\Creation
-
-[VAR:Topic]
-[OBJECT:Topic]
-	[PARAM:image]
-		![%Engine\Animation\AnimBlueprints\EventGraph:title%](EventGraphTopic.png)
-	[/PARAM]
-	[PARAM:icon]
-		![](%ROOT%/build_icon.png)(convert:false)
-	[/PARAM]
-	[PARAM:title]
-		%Engine\Animation\AnimBlueprints\EventGraph:title%
-	[/PARAM]
-	[PARAM:description]
-		%Engine\Animation\AnimBlueprints\EventGraph:description%
-	[/PARAM]
-	[PARAM:path]
-		[RELATIVE:Engine\Animation\AnimBlueprints\EventGraph]
-	[/PARAM]
-[/OBJECT]
-[/VAR]
-[VAR:TopicCompact]
-[OBJECT:TopicCompact]
-	[PARAM:image]
-		![%Engine\Animation\AnimBlueprints\EventGraph:title%](EventGraphTopic.png)
-	[/PARAM]
-	[PARAM:icon]
-		![](%ROOT%/build_icon.png)(convert:false)
-	[/PARAM]
-	[PARAM:title]
-		%Engine\Animation\AnimBlueprints\EventGraph:title%
-	[/PARAM]
-	[PARAM:description]
-		%Engine\Animation\AnimBlueprints\EventGraph:description%
-	[/PARAM]
-	[PARAM:path]
-		[RELATIVE:Engine\Animation\AnimBlueprints\EventGraph]
-	[/PARAM]
-[/OBJECT]
-[/VAR]
-
+version: 4.14
+type:reference
+topic-image:EventGraphTopic.png
+tags:Animation
+tags:Animation Blueprint
+tags:Blueprints
 
 [TOC(start:2 end:3)]
 
 [EXCERPT:Intro]
-Every Animation Blueprint has a single **EventGraph**, which is a standard [**Graph**](Engine/Blueprints/Editor/UIComponents/GraphEditor) 
+Every Animation Blueprint has a single **EventGraph**, which is a standard [Graph](Engine/Blueprints/Editor/UIComponents/GraphEditor) 
 that uses a collection of special animation-related events to initiate sequences of nodes. The most common use of the EventGraph 
-is to update values used by [**Blend Spaces**](Engine/Animation/Blendspaces) and other blend nodes to drive animations within the AnimGraph.
+is to update values used by [Blend Spaces](Engine/Animation/Blendspaces) and other blend nodes to drive animations within the AnimGraph.
 [/EXCERPT:Intro]
-
-![Event graph](eventgraph.png)(w:720)
 		
-For example, the EventGraph of an Animation Blueprint could use the velocity of the owning Pawn to determine direction and speed. These two values can then be used as inputs into a BlendSpace within the AnimGraph to blend between an idle animation and various directional movement animations. 
+![EventGraph](eventgraph.png)
 
-Similarly, the pitch and yaw of the owning Pawn can be used as inputs into another BlendSpace to blend between various aiming poses applied additively in order to cause the mesh to aim where it is looking while still maintaining the previously mentioned movement animations.
+As an example use case, the EventGraph of an Animation Blueprint could use the velocity of the owning Pawn to determine direction and speed. These values can then be used as inputs in a Blend Space within the AnimGraph to blend between an idle animation and various directional movement animations. 
+Similarly, the pitch and yaw of the owning Pawn can be used as inputs in another Blend Space to blend between various aiming poses applied additively in order to cause the mesh to aim where it is looking while still maintaining the previously mentioned movement animations.
 
 ## Animation Events
 
-Various events are called by the animation system to initialize and update the Animation Blueprint. These can be placed in the Event Graph to implement or override functionality.
+Various events are called by the animation system to initialize and update the Animation Blueprint. These can be placed in the EventGraph to implement or override functionality.
 
 
 ### Blueprint Initialize Animation
@@ -77,5 +41,7 @@ Various events are called by the animation system to initialize and update the A
 [INCLUDE:Engine/Animation/NodeReference/Events#Update]
 
 [REGION:note]
-For general information on working with node graphs, see the [**Working with Nodes**](Engine/Blueprints/Editor/UIComponents/GraphEditor) documentation.
+Please see the [](Engine/Blueprints/Scripting) page for general Blueprint Scripting information or [Working with Nodes](Engine/Blueprints/Editor/UIComponents/GraphEditor) for node based graph commands. 
 [/REGION]
+
+
diff --git a/Engine/Documentation/Source/Engine/Animation/AnimBlueprints/Interface/AnimationBlueprintUI.INT.udn b/Engine/Documentation/Source/Engine/Animation/AnimBlueprints/Interface/AnimationBlueprintUI.INT.udn
--- a/Engine/Documentation/Source/Engine/Animation/AnimBlueprints/Interface/AnimationBlueprintUI.INT.udn
+++ b/Engine/Documentation/Source/Engine/Animation/AnimBlueprints/Interface/AnimationBlueprintUI.INT.udn
@@ -1,184 +1,50 @@
 Availability:Public
-Title:Animation Blueprint UI
+Title:Animation Blueprint Editor
 Crumbs: %ROOT%, Engine, Engine/Animation, Engine/Animation/AnimBlueprints
-Description:An overview of the Animation Blueprint UI.
-version: 4.12
+Description:Highlights the core features of the Animation Blueprint Editor in Unreal Engine 4. 
 related:Engine\Animation\AnimBlueprints
 related:Engine\Animation\AnimBlueprints\Creation
 related:Engine\Animation\AnimBlueprints\EventGraph
 related:Engine\Animation\AnimBlueprints\AnimGraph
+version: 4.14
+type:reference
+topic-image:AnimBPUITopic.png
+tags:Animation
+tags:Animation Blueprint
+tags:Blueprints
 
-[VAR:Topic]
-[OBJECT:Topic]
-	[PARAM:image]
-		![%Engine\Animation\AnimBlueprints\Interface:title%](AnimBPUITopic.png)
-	[/PARAM]
-	[PARAM:icon]
-		![](%ROOT%/build_icon.png)(convert:false)
-	[/PARAM]
-	[PARAM:title]
-		%Engine\Animation\AnimBlueprints\Interface:title%
-	[/PARAM]
-	[PARAM:description]
-		%Engine\Animation\AnimBlueprints\Interface:description%
-	[/PARAM]
-	[PARAM:path]
-		[RELATIVE:Engine\Animation\AnimBlueprints\Interface]
-	[/PARAM]
-[/OBJECT]
-[/VAR]
-[VAR:TopicCompact]
-[OBJECT:TopicCompact]
-	[PARAM:image]
-		![%Engine\Animation\AnimBlueprints\Interface:title%](AnimBPUITopic.png)
-	[/PARAM]
-	[PARAM:icon]
-		![](%ROOT%/build_icon.png)(convert:false)
-	[/PARAM]
-	[PARAM:title]
-		%Engine\Animation\AnimBlueprints\Interface:title%
-	[/PARAM]
-	[PARAM:description]
-		%Engine\Animation\AnimBlueprints\Interface:description%
-	[/PARAM]
-	[PARAM:path]
-		[RELATIVE:Engine\Animation\AnimBlueprints\Interface]
-	[/PARAM]
-[/OBJECT]
-[/VAR]
-
-The following sections describe the areas of the Animation Blueprint UI.
+The **Animation Blueprint Editor** shares similar functionality to the normal [Blueprint Editor](Engine/Blueprints) but is geared more towards character animation scripting. 
 
 ![](AnimGraphUI.png)
 
-1. [Menu Bar](#menubar)
-1. [Toolbar](#toolbar)
-1. [Viewport](#viewport)
-1. [Graph](#graph)
-1. [Anim Preview Editor](#animprevieweditor)
-1. [Asset Override Editor](#assetoverrideeditor)
-1. [Asset Browser tab](#assetbrowsertab)
-1. [My Blueprint tab](#myblueprinttab)
-1. [Details panel](#detailspanel)
-1. [Palette](#palette)
+Please refer to each section below for a breakdown of the Animation Blueprint Editor user interface: 
 
-## Menu Bar
+## 1. Toolbar
 
-The **Menu Bar** gives you access to common editor actions such as saving, undoing and redoing, as well as several Persona-specific menus. For instance, the **Asset Browser** tab allows
-you to change the default preview mesh, remove unused bones from the skeleton, and record animations. 
+The **Toolbar** inside the Animation Blueprint Editor allows you to compile your Blueprints, Save, locate the Animation Blueprint asset in the **Content Browser**, and define **Class Settings** and **Class Defaults** settings like the [Blueprint Editor Toolbar.](Engine/Blueprints/Editor/UIComponents/Toolbar) 
+On the far right of the Toolbar is the [Editor Toolbar](Engine/Animation/Persona/Toolbar) which allows you to switch between the different Animation Tools in Unreal Engine 4. 
 
+## 2. Viewport
 
-## Toolbar
+The [Viewport](Engine/Animation/Persona/Viewport) window allows you to preview playback of animation assets on your selected Skeletal Mesh and provides information about your assets. You can change lighting modes, show or hide the bones of your skeleton, adjust animation playback speeds, even set your Skeletal Mesh to automatically rotate around on a turntable allowing you view it from all angles. 
 
-The Animation Blueprint **Toolbar** contains many of the same options as other Blueprint Toolbars.
+## 3. Graph
 
-![](AnimGraphUI_Toolbar.png)
+The **Graph** panel includes two main graph types: the [Event Graph](Engine/Animation/AnimBlueprints/EventGraph) which contains the animation event nodes used to trigger the updating of Skeletal Mesh poses and the [Anim Graph,](Engine/Animation/AnimBlueprints/AnimGraph) which is used to evaluate a final pose for the Skeletal Mesh to actually enter for the current frame. 
+These two graphs work together and drive the logic and poses for your character to enter during gameplay. 
 
-[EXCERPT:Toolbar]
-| Item | Description |
-| ---- | ----------- |
-| ![Compile Successful Button](Engine/Blueprints/Editor/UIComponents/Toolbar/bp_compile_successful.png) | Compiles the Blueprint. | 
-| ![Save Button](Engine/Blueprints/Editor/UIComponents/Toolbar/bp_save.png) | [INCLUDE:Engine/Blueprints/Editor/UIComponents/Toolbar#Save] |
-| ![Find in Content Browser Button](Engine/Blueprints/Editor/UIComponents/Toolbar/bp_find_in_cb.png) | [INCLUDE:Engine/Blueprints/Editor/UIComponents/Toolbar#FindInCB] | 
-| ![](PoseButtons.png) | Updates the viewport to display the Preview or Reference pose. | 
-|  ![](Record.png)| Allows you to record a new animation based off the current preview pose animation. | 
-| ![Search Button](Engine/Blueprints/Editor/UIComponents/Toolbar/bp_search.png) | [INCLUDE:Engine/Blueprints/Editor/UIComponents/Toolbar#Search] | 
-| ![Blueprint Properties Button](Engine/Blueprints/Editor/UIComponents/Toolbar/bp_properties.png) | [INCLUDE:Engine/Blueprints/Editor/UIComponents/Toolbar#BlueprintProps] | 
-| ![Simulate Button](Engine/Blueprints/Editor/UIComponents/Toolbar/bp_simulate.png) | Starts the game in simulation mode. |
-| ![Play In Editor Button](Engine/Blueprints/Editor/UIComponents/Toolbar/bp_play_in.png) | Starts the game in normal play mode in the Level Editor. Clicking the arrow displays the **Launch Options** menu. | 
-| ![Debug Dropdown](Engine/Blueprints/Editor/UIComponents/Toolbar/bp_debug_dropdown.png) | [INCLUDE:Engine/Blueprints/Editor/UIComponents/Toolbar#DebugObject]  | 
-[/EXCERPT:Toolbar]
+## 4. Details
 
-[REGION:note]
-For more information about Blueprint toolbars, see the **[](Engine/Blueprints\Editor/UIComponents\Toolbar)** documentation.
-[/REGION]
+The **Details** panel in the Animation Blueprint Editor is the same as the [](Engine/Blueprints/Editor/UIComponents/Details) in the Blueprint Editor and is where you can access and edit properties related to any variables you have created as well as the nodes placed within your graph. 
 
-In addition to those options, the Animation Blueprint toolbar also contains tabs for the various modes of Persona.
 
-| Item | Description |
-| ---- | ----------- |
-|  ![](button_AnimGraphUI_Skeleton.png) | Opens the Persona tab that contains Skeleton properties. |
-|  ![](button_AnimGraphUI_Mesh.png) | Opens the Persona tab that contains Skeletal Mesh properties. |
-|  ![](button_AnimGraphUI_Animation.png) | Opens the Persona tab that contains properties of animation assets. |
-|  ![](button_AnimGraphUI_Graph.png) | Opens the Animation Blueprint tab. |
+## 5. My Blueprint
 
+The [My Blueprint](Engine/Blueprints/Editor/UIComponents/MyBlueprint) panel, also found in the Blueprint Editor, contains a list of your graphs, functions, variables and other related properties contained within the Animation Blueprint. 
 
+## 6. Anim Preview Editor / Asset Browser
 
-## Viewport
-The **Viewport** displays the Skeletal Mesh with any animations activated by current variable values.
+The [](Engine/Animation/Persona/BlueprintDefaults) allows you to make changes to your variables which will update the Skeletal Mesh in the viewport. 
+You can also switch to the **Edit Defaults** mode to make changes to variables and then apply those as the default values. 
+Docked on a separate tab is the [](Engine\Animation\Persona\AssetBrowser) which allows you to view all the animation assets that are associated with and can be used by the associated Skeleton asset. 
 
-![](AnimGraphUI_Viewport.png)(w:500)
-
-## Graph
-
-The tabbed **Graph** panel displays whichever graph is currently selected.
-
-
-### EventGraph Tab
-
-The **EventGraph** tab contains any animation event nodes that trigger the updates of the Skeletal Mesh poses.  It can also contain animation nodes that set variable values that determine which animation asset is played when.
-
-![](AnimGraphUI_Graph.png)(w:900)
-
-### AnimGraph Tab
-
-The **AnimGraph** tab contains the State Machine that hooks up all the actual animations and States, along with any other animation nodes you want.
-
-![](AnimGraphUI_AnimGraph.png)
-
-Although you could put all your animation nodes directly into the AnimGraph tab, it is more efficient to organize them into a State Machine, States, and Transitions (Rules), for ease of upkeep and simplicity.
-
-[REGION:note]
-For more information about State Machines, see **[](Engine/Animation/StateMachines)**.
-[/REGION]
-
-
-
-## Anim Preview Editor
-
-The **Anim Preview Editor** panel displays all of your Animation Blueprint's variables, in their assigned categories. You can manually modify the variable values here to preview and test different animations in the viewport.
-
-![](AnimGraphUI_BPDefaults.png)
-
-[REGION:note]
-For more information, see **[](Engine/Animation/Persona\BlueprintDefaults)**.
-[/REGION]
-
-
-## Asset Override Editor
-
-The **Asset Override Editor** panel displays all of your animation assets that are referenced by the parent animation graph. If there are no nodes referencing assets in the parent class or this is not a derived blueprint there will be nothing displayed in this panel.
-
-## Asset Browser Tab
-
-The Asset Browser tab contains a list of all the animation assets in your Animation Blueprint.
-
-![](AnimGraphUI_AssetBrowser.png)
-
-By default, you can find the Asset Browser tab next to the **My Blueprint** tab, underneath the Graph.
-
-
-## My Blueprint Tab
-
-The My Blueprint tab contains a list of all the graphs and variables in your Animation Blueprint.
-
-![](AnimGraphUI_MyBP.png)
-
-By default, you can find the **My Blueprint** tab next to the **Asset Browser** tab, underneath the Graph.
-
-
-## Details Panel
-
-The **Details** panel displays the properties of selected elements.
-
-![](AnimGraphUI_Details.png)
-
-## Palette
-
-The **Palette** panel contains all the nodes available for use in the Animation Blueprint's various graphs. In the top panel any nodes marked as a "favorite" will be displayed here.
-
-![](AnimGraphUI_Palette.png)
-
-[REGION:note]
-The Palette may not be visible by default.  If it is not, you can open it by clicking **Palette** on the **Window** menu.
-[/REGION]
diff --git a/Engine/Documentation/Source/Engine/Animation/AnimBlueprints/VimBlueprints.INT.udn b/Engine/Documentation/Source/Engine/Animation/AnimBlueprints/VimBlueprints.INT.udn
--- a/Engine/Documentation/Source/Engine/Animation/AnimBlueprints/VimBlueprints.INT.udn
+++ b/Engine/Documentation/Source/Engine/Animation/AnimBlueprints/VimBlueprints.INT.udn
@@ -9,27 +9,13 @@ Related: Engine/Animation/PhysicallyDrivenAnimation
 Related: Engine/Animation/StateMachines
 Related: Gameplay/AI
 Related: Engine/Content/Tools/MayaRiggingTool
-version: 4.12
-
-[VAR:Topic]
-[OBJECT:Topic]
-	[PARAM:image]
-		![%Engine/Animation/AnimBlueprints:title%](AnimBP_topic.png)
-	[/PARAM]
-	[PARAM:icon]
-		![](%ROOT%/anim_icon.png)(convert:false)
-	[/PARAM]
-	[PARAM:title]
-		%Engine/Animation/AnimBlueprints:title%
-	[/PARAM]
-	[PARAM:description]
-		%Engine/Animation/AnimBlueprints:description%
-	[/PARAM]
-	[PARAM:path]
-		[RELATIVE:Engine/Animation/AnimBlueprints]
-	[/PARAM]
-[/OBJECT]
-[/VAR]
+version: 4.14
+type:landing
+topic-image:animbptopic.png
+tags:Animation
+tags:Animation Blueprint
+tags:Blueprints
+tags:State Machine
 
 [REGION:banner]
 ![](AnimationBlueprint.png)
@@ -37,26 +23,17 @@ version: 4.12
 
 
 [EXCERPT:Intro]
-An **Animation Blueprint** is essentially a specialized [**Blueprint**](Engine/Blueprints) whose graphs 
-control the animation of a SkeletalMesh. It can perform blending of animations, directly control the 
-bones of the skeleton, and output a final pose for a SkeletalMesh each frame.
-
-The Controller directs its Pawn (or Character) to move based on player input or decisions made based on the gameplay 
-environment. Each Pawn has a SkeletalMeshComponent that references the SkeletalMesh to animate and 
-has an instance of a Animation Blueprint. The Animation Blueprint, through its two graphs, can access properties 
-of the owning Pawn, compute values used for blending or state transitions or driving AnimMontages, and 
-calculate the current pose of the SkeletalMesh based on blending of AnimationSequences and direct 
-transformations of the Skeleton from SkeletalControls.
+An **Animation Blueprint** is essentially a specialized [Blueprint](Engine/Blueprints) whose graphs control the animation of a [Skeletal Mesh.](Engine/Content/Types/SkeletalMeshes) 
+Graphs are edited inside the **Animation Blueprint Editor** where you can perform blending of animations, directly control the bones of a Skeleton, or set up logic that will ultimately define the final animation pose for a Skeletal Mesh to use each frame. 
 [/EXCERPT:Intro]
 
 [REGION:note]
-An example of an **Animation Blueprint** is demonstrated on the [**Animation Content Examples**](Resources\ContentExamples\Animation) page under section 1.2.
+An example of an **Animation Blueprint** is demonstrated on the [Animation Content Examples](Resources\ContentExamples\Animation) page under section 1.2.
 [/REGION]
 
 ## Creating & Navigating Animation Blueprints
 
-Before jumping into how Animation Blueprints work, it’s best to understand how to create them and start to familiarize yourself with the Animation Blueprint Interface. 
-Check out each of the links below for more information. 
+Before jumping into how Animation Blueprints work, it's best to understand how to create them and start to familiarize yourself with the Animation Blueprint Interface. 
 
 [REGION:topics half]
 %Engine\Animation\AnimBlueprints\Creation:TopicCompact%
@@ -65,7 +42,7 @@ Check out each of the links below for more information.
 
 ## Working with Animation Blueprints
 
-Once you have created your Animation Blueprint and can navigate the UI, it's time to consider which aspect of it you want to work on first.
+Once you have created your Animation Blueprint and can navigate the UI, it's time to consider which aspect you want to work on first.
 
 [EXCERPT:Graphs]
 There are two main components in an Animation Blueprint that work in conjunction to create the final animation 
@@ -74,22 +51,41 @@ to drive State Machines, BlendSpaces, or other nodes allowing blending between m
 that can fire off notifications to other systems to enable dynamic animation-driven effects to occur.
 [/EXCERPT:Graphs]
 
-Check out the links below for more information on each topic. 
-
 [REGION:topics half]
 %Engine\Animation\AnimBlueprints\EventGraph:TopicCompact%
 %Engine\Animation\AnimBlueprints\AnimGraph:TopicCompact%
 [/REGION]
 
 
+## Working with Sub Animation Blueprint Instances
+
+When creating Animation Blueprints for various characters, it is not uncommon to share animation logic between those characters. 
+For example, suppose you have physical animation or [AnimDynamics](Engine\Animation\AnimHowTo\AnimDynamics) applied to one character and you want to apply it to another. 
+With the **Sub Anim Instance** node, you can create a reference to a separate Animation Blueprint to access and include all of its logic within another Animation Blueprint (sort of like a function call). 
+
+This can be used to break up large Animation Blueprints into separate, smaller Animation Blueprints, that can be included within other Animation Blueprints. 
+For example, you could have one Animation Blueprint for purely "Locomotion" and a Sub Anim Instance that handles any physics-based animation. 
+Member variables of a Sub Animation Blueprint can also be exposed as input pins inside their parent Animation Blueprint and **Sub-Graph Input** nodes can be used to expose **In Poses** on the Sub Anim Instance (shown below).
+
+![](SubAnimInstance.png)
+
+[REGION:warning]
+When working with Sub Anim Instances, they must use the same Skeleton as the outer Animation Blueprint it resides in.
+[/REGION]
+
+[REGION:note]
+See [Using Sub Anim Instances](https://docs.unrealengine.com/latest/INT/Engine/Animation/AnimHowTo/SubAnimInstance/index.html) for an example walk-through of setting up a Sub Animation Blueprint Instance to affect a character's movement. 
+[/REGION]
+
+
 ## Working with Child Animation Blueprints
 
-When creating and setting up your characters and Animation Blueprints, there may be instances where you have multiple characters that are similar but have entirely different animations for different states. Rather than create Animation Blueprints for each of those characters, you may be able to use a **Child Animation Blueprint** of a pre-existing Animation Blueprint and override the animation assets that you want to replace. The Child Animation Blueprint will inherit everything from the parent, but it will also allow you to specify which animations to override via the **Asset Override Editor**. 
+Also when creating and setting up your characters and Animation Blueprints, there may be instances where you have multiple characters that are similar but have entirely different animations for different states. Rather than create Animation Blueprints for each of those characters, you may be able to use a **Child Animation Blueprint** of a pre-existing Animation Blueprint and override the animation assets that you want to replace. The Child Animation Blueprint will inherit everything from the parent, but it will also allow you to specify which animations to override via the **Asset Override Editor**. 
 
-For example, suppose you have two humanoid type characters where one is a human and another is a werewolf. Both characters can walk, run, jump, crawl, etc. However, their motions may be entirely different. The logic of moving between those motions is the same, but the motions themselves are different. It's perfectly fine to create two different Animation Blueprints, but to save some time. you can create a Child Animation Blueprint and selectively choose which animations to override. 
+For example, suppose you have two humanoid type characters where one is a human and another is a werewolf. Both characters can walk, run, jump, crawl, etc. However, their motions may be entirely different. The logic of moving between those motions is the same, but the motions themselves are different. It's perfectly fine to create two different Animation Blueprints but to save some time. you can create a Child Animation Blueprint and selectively choose which animations to override. 
 
 [REGION:note]
-See [**Animation Blueprint Override**](Engine/Animation/AnimHowTo/AnimBPOverride) for an example walk-through of creating and overriding an Animation Blueprint with a Child Animation Blueprint.
+See [Animation Blueprint Override](Engine/Animation/AnimHowTo/AnimBPOverride) for an example walk-through of creating and overriding an Animation Blueprint with a Child Animation Blueprint.
 [/REGION]
 
 
diff --git a/Engine/Documentation/Source/Engine/Animation/AnimHowTo/AdditiveAnimations/01_BlendSetup/BlendSetup.INT.udn b/Engine/Documentation/Source/Engine/Animation/AnimHowTo/AdditiveAnimations/01_BlendSetup/BlendSetup.INT.udn
new file mode 100644
--- /dev/null
+++ b/Engine/Documentation/Source/Engine/Animation/AnimHowTo/AdditiveAnimations/01_BlendSetup/BlendSetup.INT.udn
@@ -0,0 +1,78 @@
+Availability: Public
+Title:1 - Setting up the Animation to Blend
+Crumbs: %ROOT%, Engine, Engine\Animation, Engine\Animation\AnimHowTo
+Description:In this step, we update our player character and create an AnimMontage for the animation we want to layer in on top of our movement animations.
+version: 4.14
+skilllevel: Advanced
+type:howto
+tags:Animation
+tags:How-To
+
+[VAR:Steps]
+[OBJECT:Navigation]
+	[PARAM:previous]
+	[/PARAM]
+	[PARAM:current]
+	[/PARAM]
+	[PARAM:home]
+		[Layered Animations - Home ](Engine\Animation\AnimHowTo\AdditiveAnimations)
+	[/PARAM]
+	[PARAM:next]
+		[Next](Engine\Animation\AnimHowTo\AdditiveAnimations\02_CharacterBlueprintSetup)
+	[/PARAM]
+[/OBJECT]
+[/VAR]
+
+%Steps%
+
+In this step, we change our default player character and create an AnimMontage for the animation we want to layer in on top of our movement animations. 
+
+[EXCERPT:INTRO]
+[REGION:note]
+For this guide we are using the **Blueprint Third Person Template** project and have added the **Animation Starter Pack** to the project: 
+
+![](AnimationAssetPack.png)(w:820)
+
+You can download the Animation Starter Pack for free through the **Marketplace** from the Epic Launcher. 
+[/REGION]
+[/EXCERPT:INTRO]
+
+1.  With your project open, inside the **Content Browser** under the **Content/ThirdPersonBP/Blueprints** folder, open the **ThirdPersonGameMode** Blueprint. 
+
+1.	Under **Default Pawn Class**, click the drop-down menu and select the **Ue4ASP_Character**. 
+
+	![](Animations1.png)
+
+1.	**Compile** and **Save** then close the Blueprint. 
+
+1.	In the **Content Browser**, open the **Content/AnimStarterPack** folder. 
+
+1.	**Right-click** on the **Fire_Shotgun_Ironsights** animation, select **Create** then **Create AnimMontage**.
+  
+	![](Animations2.png)
+	[REGION:note]
+	You can filter the window by typing "Fire" inside the search window as shown above. 
+	[/REGION]
+1.	Open the new AnimMontage and click the magnifying glass icon. 
+
+	![](Animations3.png)
+
+	This will open the **Anim Slot Manager** which we will use to create a **Slot** that we can use to blend to when calling the Slot by name. 
+
+1.	Click the **Add Slot** button, then name the slot **UpperBody** then click **Save**. 
+
+	![](Animations4.png)
+
+	[REGION:note]
+	Although we chose the name "UpperBody" here, we are not targeting the upperbody just yet but have assigned the name so we know where we are targeting.
+	[/REGION]
+
+1.	Click the drop-down menu inside the **Montage Group** section and change it to **DefaultGroup.UpperBody**, then close the window. 
+
+	![](Animations5.png)
+
+	We have now associated this Montage with the **DefaultGroup.UpperBody**, allowing it to play when the Slot is called.  
+
+In the next section, we will change our default player character and make some adjustments to determine when we are firing a weapon. 
+
+%Steps%
diff --git a/Engine/Documentation/Source/Engine/Animation/AnimHowTo/AdditiveAnimations/02_CharacterBlueprintSetup/CharacterBPSetup.INT.udn b/Engine/Documentation/Source/Engine/Animation/AnimHowTo/AdditiveAnimations/02_CharacterBlueprintSetup/CharacterBPSetup.INT.udn
new file mode 100644
--- /dev/null
+++ b/Engine/Documentation/Source/Engine/Animation/AnimHowTo/AdditiveAnimations/02_CharacterBlueprintSetup/CharacterBPSetup.INT.udn
@@ -0,0 +1,70 @@
+Availability: Public
+Title:2 - Updating the Character Blueprint
+Crumbs: %ROOT%, Engine, Engine\Animation, Engine\Animation\AnimHowTo
+Description:In this step, we create some script inside our Character Blueprint to allow us to simulate firing a weapon. 
+version: 4.14
+skilllevel: Advanced
+type:howto
+tags:Animation
+tags:How-To
+
+[VAR:Steps]
+[OBJECT:Navigation]
+	[PARAM:previous]
+		[Previous](Engine\Animation\AnimHowTo\AdditiveAnimations\01_BlendSetup)
+	[/PARAM]
+	[PARAM:current]
+	[/PARAM]
+	[PARAM:home]
+		[Layered Animations - Home ](Engine\Animation\AnimHowTo\AdditiveAnimations)
+	[/PARAM]
+	[PARAM:next]
+		[Next](Engine\Animation\AnimHowTo\AdditiveAnimations\03_AnimGraphSetup)
+	[/PARAM]
+[/OBJECT]
+[/VAR]
+
+%Steps%
+
+In this step, we create some script inside our Character Blueprint to allow us to signal our Animation Blueprint that we are firing a weapon. 
+
+1.	Inside the **Content Browser** under **Content/AnimStarterPack**, open the **Ue4ASP_Character** Blueprint. 
+
+1.	Locate the **Crouching** section, and delete the **InputAction Crouch** Event (with the Warning!) sign. 
+
+	![](Animations6.png)
+
+	There is a warning on this node because our project is not set up to use the InputAction event "Crouch". 
+
+1.	**Right-click** in the graph and add a **C** Key Event and connect as shown below. 
+
+	![](Animations7.png)
+
+	We have just associated crouch now with the **C** key press, you can change this input to whatever you would like. 
+
+1.	In the **MyBlueprint** window, click the **Add Variable** button to create a **Boolean** variable and call it **Fire Button Down**. 
+
+	![](Animations8.png)
+
+1.	**Right-click** in the graph and add the **Left Mouse Button** Key Event. 
+
+1.	Holding **Alt**, drag the **Fire Button Down** variable in the graph to create a **Set** node. 
+
+1.	Repeat the previous step (or copy the **Set** node) and connect as shown below. 
+
+	![](Animations9.png)
+
+	Make sure that **Fire Button Down** is **True** when **Left Mouse Button** is **Pressed** and **False** when **Released**. 
+
+1.	Select **CapsuleComponent** in the **Components** window, then in the **Details** panel check the **Hidden in Game** option. 
+
+	![](HideCapsule.png)
+
+	This will turn off the debug display of capsule's collision. 
+
+1.	**Compile** and **Save** then close the Blueprint. 
+
+Our character is now set up to send signals to our Animation Blueprint when we are crouching or firing a weapon. 
+In the next step, we will set up our Animation Graph inside our Animation Blueprint to handle the blending of our Anim Montage with our normal movement animations so that we can blend between the two. 
+
+%Steps%
diff --git a/Engine/Documentation/Source/Engine/Animation/AnimHowTo/AdditiveAnimations/03_AnimGraphSetup/AnimGraphSetup.INT.udn b/Engine/Documentation/Source/Engine/Animation/AnimHowTo/AdditiveAnimations/03_AnimGraphSetup/AnimGraphSetup.INT.udn
new file mode 100644
--- /dev/null
+++ b/Engine/Documentation/Source/Engine/Animation/AnimHowTo/AdditiveAnimations/03_AnimGraphSetup/AnimGraphSetup.INT.udn
@@ -0,0 +1,75 @@
+Availability: Public
+Title:3 - Setting up the Anim Graph
+Crumbs: %ROOT%, Engine, Engine\Animation, Engine\Animation\AnimHowTo
+Description:In this step, we set up our AnimGraph to handle the blending of our AnimMontage and the movement pose from our State Machine. 
+version: 4.14
+skilllevel: Advanced
+type:howto
+tags:Animation
+tags:How-To
+
+[VAR:Steps]
+[OBJECT:Navigation]
+	[PARAM:previous]
+		[Previous](Engine\Animation\AnimHowTo\AdditiveAnimations\02_CharacterBlueprintSetup)
+	[/PARAM]
+	[PARAM:current]
+	[/PARAM]
+	[PARAM:home]
+		[Layered Animations - Home ](Engine\Animation\AnimHowTo\AdditiveAnimations)
+	[/PARAM]
+	[PARAM:next]
+		[Next](Engine\Animation\AnimHowTo\AdditiveAnimations\04_BlendSettings)
+	[/PARAM]
+[/OBJECT]
+[/VAR]
+
+%Steps%
+
+In this step, we set up our AnimGraph to handle the blending of our AnimMontage and the movement pose from our State Machine. 
+
+1.	In the **Content/AnimStarterPack**, open the **UE4ASP_HeroTPP_AnimBlueprint** Animation Blueprint. 
+
+1.	Inside the **MyBlueprint** window, open the **AnimGraph**. 
+
+	![](Animations9b.png)
+
+1.	Drag off the **Locomotion** State Machine and search for and add **New Save cached pose**. 
+
+	![](Animations10.png)
+
+	Here we are getting the resulting pose from our Locomotion State Machine and storing it in a cached pose that we can use elsewhere. 
+
+1.	Give the cached pose a name by pressing **F2** on the node such as **LocomotionCache**. 
+
+	![](Animations11.png)
+
+1.	**Right-click** in the graph and add the **Use cached pose 'LocomotionCache'** (or whatever you named the cached pose). 
+
+	![](Animations12.png)
+
+1.	Drag off the **Use cached pose** node and search for and add the **Layered blend per bone** node. 
+
+	![](Animations13.png)
+
+	The **Layered Blend Per Bone** node will allow us to blend animations together from a specified bone on the Skeleton.
+
+1.	Copy/Paste the **Use cached pose 'LocomotionCache'** node, then drag off it and add the **Slot 'DefaultSlot'** node. 
+
+	![](Animations14.png)
+
+1.	In the **Settings** for the **Slot** node, click the dropdown menu and select **DefaultGroup.UpperBody**. 
+
+	![](Animations15.png)
+
+	When an Animation Montage that uses this slot is called to play, it will now be called inside our AnimGraph. 
+
+1.	Connect the **Slot** to the **Blend Poses 0** pin on the **Layered blend per bone** node, then connect the out to the **Result** and **Compile**. 
+
+	![](Animations16.png)
+
+Our AnimGraph is complete, but we still need to modify some settings on the **Layered blend per bone** node to finish up which we will do next.  
+
+
+
+%Steps%
diff --git a/Engine/Documentation/Source/Engine/Animation/AnimHowTo/AdditiveAnimations/04_BlendSettings/BlendSettings.INT.udn b/Engine/Documentation/Source/Engine/Animation/AnimHowTo/AdditiveAnimations/04_BlendSettings/BlendSettings.INT.udn
new file mode 100644
--- /dev/null
+++ b/Engine/Documentation/Source/Engine/Animation/AnimHowTo/AdditiveAnimations/04_BlendSettings/BlendSettings.INT.udn
@@ -0,0 +1,112 @@
+Availability: Public
+Title:4 - Finalizing the Blend Settings
+Crumbs: %ROOT%, Engine, Engine\Animation, Engine\Animation\AnimHowTo
+Description:In this final step, we define the method in which to layer our blended animation and which bone to start the blend from. 
+version: 4.14
+skilllevel: Advanced
+type:howto
+tags:Animation
+tags:How-To
+
+[VAR:Steps]
+[OBJECT:Navigation]
+	[PARAM:previous]
+		[Previous](Engine\Animation\AnimHowTo\AdditiveAnimations\03_AnimGraphSetup)
+	[/PARAM]
+	[PARAM:current]
+	[/PARAM]
+	[PARAM:home]
+		[Layered Animations - Home ](Engine\Animation\AnimHowTo\AdditiveAnimations)
+	[/PARAM]
+	[PARAM:next]
+	[/PARAM]
+[/OBJECT]
+[/VAR]
+
+%Steps%
+
+In this final step, we define the method in which to layer our blended animation and which bone to start the blend from. 
+
+1.	Inside the **AnimGraph** of **UE4ASP_HeroTPP_AnimBlueprint**, click on the **Layered blend per bone** node, then expand the **Layer Setup** section and click the **+** sign. 
+
+	![](Animations17.png)
+
+1.	Under **Bone Name** enter **spine_01**, then set **Blend Depth** to **1** and check **Mesh Space Rotation Blend**. 
+
+	![](Animations18.png)
+
+	With these settings, we can blend the shooting AnimMontage onto our Skeleton starting from the **spine_01** bone. 
+
+1.	Inside the **MyBlueprint** window, open the **EventGraph**. 
+
+	![](Animations18b.png)
+
+1.	Drag off the **As Ue4ASP Character** pin of the **Cast To Ue4ASP_Character** node and add the **Get Fire Button Down** node. 
+
+	![](Animations19.png)
+
+	This will tell us if the player is pressing the fire button or not. 
+
+1.	Hold **B** and **Left-click** in the graph to create a **Branch** node and connect it following **Set Crouching**. 
+
+	![](Animations20.png)
+
+	For the **Condition**, plug in the **FireButtonDown** pin. 
+
+1.	**Right-click** in the graph and add the **Montage Is Playing** node and assign **Fire_Shotgun_Ironsights_Montage**. 
+
+	![](Animations22.png)
+
+1.	Off the **Return Value** of the **Montage Is Playing** node, add a **Branch**. 
+
+1.	Connect the **Branch** from Step 5 to the new **Branch** as shown below and add a **Montage Play** node with **Montage to Play** set to your Montage. 
+
+	![](Animations23.png)
+
+	This will check now if the Montage is currently playing and if it is, do not play it (and if it is not, play the Montage). 
+
+1.	**Compile** and **Save**, close the **Blueprint** and delete the existing **ThirdPersonCharacter** in the level.
+
+	![](Animations24.png)
+
+	This will allow us to use our selected character defined in our Game Mode and not the one placed in the level. 
+
+1. **Play** in the editor. 
+
+
+## End Result
+
+[OBJECT:EmbeddedVideo]
+[PARAMLITERAL:width]
+640
+[/PARAMLITERAL]
+[PARAMLITERAL:height]
+360
+[/PARAMLITERAL]
+[PARAMLITERAL:videoid]
+WxYiLbnbQ_A
+[/PARAMLITERAL]
+[/OBJECT]
+
+You will see that when you hold the **Left Mouse Button** down, the character will play the shooting AnimMontage. 
+You can walk around while shooting, you can even crouch (by holding **C**) and shoot as well. 
+You could use this same guide to play a melee attack, reload animation or any other animation that you want to be able to do while moving around. 
+
+The **UpperBody** slot that we created and assigned in our Animation Montage can be assigned in other Animation Montages as well. 
+Since our AnimGraph has the Slot node assigned, whenever we call to play an Animation Montage that has that slot assigned, it will get played with our defined settings. 
+ 
+%Steps%
+
+
+
+
+
+
+
+
+
+
+
+
+
+
diff --git a/Engine/Documentation/Source/Engine/Animation/AnimHowTo/AdditiveAnimations/AdditiveHowTo.INT.udn b/Engine/Documentation/Source/Engine/Animation/AnimHowTo/AdditiveAnimations/AdditiveHowTo.INT.udn
--- a/Engine/Documentation/Source/Engine/Animation/AnimHowTo/AdditiveAnimations/AdditiveHowTo.INT.udn
+++ b/Engine/Documentation/Source/Engine/Animation/AnimHowTo/AdditiveAnimations/AdditiveHowTo.INT.udn
@@ -1,21 +1,35 @@
 Availability: Public
-Title:Using Additive Animations
+Title:Using Layered Animations
 Crumbs: %ROOT%, Engine, Engine\Animation, Engine\Animation\AnimHowTo
-Description:Demonstrates how to blend animations together, in this case a character that can move and fire a weapon. 
-Related: Engine/Animation/AnimBlueprints
-Related: Engine/Animation/AnimMontage
-Related: Engine/Animation/NodeReference/Blend
-Related: Engine/Animation/AnimationBlending
-version: 4.9
+Description:Demonstrates how to blend animations together, in this case, a character that can move and fire a weapon at the same time.
+version: 4.14
 skilllevel: Advanced
+type:howto
+topic-image:LayerAnimationHowTo.png
+tags:Animation
+tags:How-To
 
 Animation blending, as a concept, simply means making a smooth transition between two or more animations on a single character or Skeletal Mesh. In Unreal Engine 4, there are various ways in which such blending can be applied either through [](Engine/Animation/Blendspaces), through additive methods by literally combining two animations based on a weighted bias or alpha value, or even applying direct overrides of the existing pose. 
 
-You can also send animation directly to a specific bone in the Skeleton and all of its children. For instance, you can start with an animation in which the character is running, but then selectively apply a shooting animation to the upper body of the character. The final result would be a character that is running and shooting at the same time, which is what this page will teach you to do in the steps below. 
+You can also send animation directly to specific bones within the Skeleton and all of its children. For instance, you can start with an animation in which the character is running, but then selectively apply a shooting animation to the upper body of the character. The final result would be a character that can run around and shoot at the same time, similar to the video below. 
 
-[TOC(start:2 end:2)]
+[OBJECT:EmbeddedVideo]
+[PARAMLITERAL:width]
+640
+[/PARAMLITERAL]
+[PARAMLITERAL:height]
+360
+[/PARAMLITERAL]
+[PARAMLITERAL:videoid]
+WxYiLbnbQ_A
+[/PARAMLITERAL]
+[/OBJECT]
 
+In this guide we will set up a character similar to the video above, however, you can apply this to any of your own characters to create layered animations. 
+
+<!--
 [EXCERPT:Intro]
+[REGION:note]
 For this guide, we have created a new project using the **Blueprint Third Person** template with **Starter Content** enabled. 
 
 We are also using the **Animation Starter Pack** which is available for download for free via the **Marketplace**. 
@@ -23,200 +37,26 @@ We are also using the **Animation Starter Pack** which is available for download
 ![](AnimationAssetPack.png)
 
 Once downloaded, you can add the Animation Starter Pack to your project via the **Add To Project** button from the Launcher.
+[/REGION]
 [/EXCERPT:Intro]
-
-## Setting up the Animation to Blend
-
-In this section, we assign our player character and make modifications to the animation we will blend in with our movement animations. 
-
-1.  With your project open, inside the **Content Browser** under the **Game/ThirdPersonBP/Blueprints** folder, open the **ThirdPersonGameMode** Blueprint. 
-
-1.	Under **Default Pawn Class**, click the dropdown menu and select the **Ue4ASP_Character**. 
-
-	![](Animations1.png)
-
-1.	**Compile** and **Save** then close the Blueprint. 
-
-1.	In the **Content Browser**, open the **Game/AnimStarterPack** folder. 
-
-1.	**Right-click** on the **Fire_Shotgun_Ironsights** animation, select **Create** then **Create AnimMontage**.
-  
-	![](Animations2.png)
-	[REGION:note]
-	You can filter the window by typing "Fire" inside the search window as shown above. 
-	[/REGION]
-1.	Rename the Montage if you wish, then open it up and click the magnifying glass icon. 
-
-	![](Animations3.png)
-
-	This will open the **Anim Slot Manager** which we will use to create a **Slot** that we can use to blend to when calling the Slot by name. 
-
-1.	Click the **Add Slot** button, then name the slot **UpperBody** then **Save** and close the window. 
-
-	![](Animations4.png)
-
-1.	Click the dropdown menu inside the **Montage Group** section and change it to **DefaultGroup.UpperBody**, then close the window. 
-
-	![](Animations5.png)
-
-	We have now associated this Montage with the **DefaultGroup.UpperBody**, allowing it to play when the Slot is called.  
-
-## Updating the Character Blueprint
-
-Next we will make some updates to our Character Blueprint so that we can determine when they are pressing the shoot button. 
-
-1.	Inside the **Content Browser** under **Game/AnimStarterPack**, open the **Ue4ASP_Character** Blueprint. 
-
-1.	Locate the **Crouching** section, and delete the **InputAction Crouch** Event (with the Warning!) sign. 
-
-	![](Animations6.png)
-
-	There is a warning on this node because our project is not set up to use the InputAction event "Crouch". 
-
-1.	**Right-click** in the graph and add a **C** Key Event and connect as shown below. 
-
-	![](Animations7.png)
-
-	We have just associated crouch now with the **C** key press, you can change this input to whatever you would like. 
-
-1.	In the **MyBlueprint** window, click the **Add Variable** button to create a **Boolean** variable. 
-
-1.	Name the variable **Fire Button Down** and click the eye icon to make the variable public. 
-
-	![](Animations8.png)
-
-1.	**Right-click** in the graph and add the **Left Mouse Button** Key Event. 
-
-1.	Holding **Alt**, drag the **Fire Button Down** variable in the graph to create a **Set** node. 
-
-1.	Repeat the previous step (or copy the **Set** node) and connect as shown below. 
-
-	![](Animations9.png)
-
-	Make sure that **Fire Button Down** is **True** when **Left Mouse Button** is **Pressed** and **False** when **Released**. 
-
-1.	**Compile** and **Save** then close the Blueprint. 
-
-## Setting up the Animation Graph
-
-Here we will set up the Animation Graph and place the nodes needed to blend between our current pose and the fire weapon animation. 
-
-1.	Also under **Game/AnimStarterPack**, open the **UE4ASP_HeroTPP_AnimBlueprint** Animation Blueprint. 
-
-1.	Inside the **MyBlueprint** window, open the **AnimGraph**. 
-
-	![](Animations9b.png)
-
-1.	Drag off the **Locomotion** State Machine and search for and add **New Save cached pose**. 
-
-	![](Animations10.png)
-
-	Here we are getting the resulting pose from our Locomotion State Machine and storing it in a cached pose that we can use elsewhere. 
-
-1.	Give the cached pose a name such as **LocomotionCache**. 
-
-	![](Animations11.png)
-
-1.	**Right-click** in the graph and add the **Use cached pose 'LocomotionCache'** (or whatever you named the cached pose). 
-
-	![](Animations12.png)
-
-1.	Drag off the **Use cached pose** node and search for and add the **Layered blend per bone** node. 
-
-	![](Animations13.png)
-
-	The **Layered Blend Per Bone** node will allow us to blend animations together from a specified bone on the Skeleton.
-
-1.	Copy/Paste the **Use cached pose 'LocomotionCache'** node, then drag off it and add the **Slot 'DefaultSlot'** node. 
-
-	![](Animations14.png)
-
-1.	In the **Settings** for the **Slot** node, click the dropdown menu and select **DefaultGroup.UpperBody**. 
-
-	![](Animations15.png)
-
-1.	Connect the **Slot** to the **Blend Poses 0** pin on the **Layered blend per bone** node, then connect its out to the **Result**. 
-
-	![](Animations16.png)
-
-	Our AnimGraph is complete, but we still need to modify some settings on the **Layered blend per bone** node. 
-
-## Adjusting the Blend Settings
-
-Finally, we update the way we blend between animations so that our upperbody plays one animation while our lowerbody plays another. 
-
-1.	Click on the **Layered blend per bone** node, then expand the **Layer Setup** section and click the **+** sign. 
-
-	![](Animations17.png)
-
-1.	Under **Bone Name** enter **spine_01**, then set **Blend Depth** to **1** and check **Mesh Space Rotation Blend**. 
-
-	![](Animations18.png)
-
-	With these settings we can blend the shooting AnimMontage onto our Skeleton starting from the **spine_01** bone. 
-
-1.	Inside the **MyBlueprint** window, open the **EventGraph**. 
-
-	![](Animations18b.png)
-
-1.	Drag off the **As Ue 4ASP Character** pin of the **Cast To Ue4ASP_Character** node and add the **Get Fire Button Down** node. 
-
-	![](Animations19.png)
-
-	This will tell us if the player is pressing the fire button or not. 
-
-1.	Hold **B** and **Left-click** in the graph to create a **Branch** node and connect it following **Set Crouching**. 
-
-	![](Animations20.png)
-
-	For the **Condition**, plug in the **FireButtonDown** pin. 
-
-1.	**Right-click** in the graph and add the **Montage Is Playing** node. 
-
-	![](Animations21.png)
-
-1.	Connect the **Montage Is Playing** node to the **True** of the **Branch**. 
-
-	![](Animations22.png)
-
-	For the **Montage**, select your Montage (**Fire_Shotgun_Ironsights_Montage**). 
-
-1.	Off the **Return Value** of the **Montage Is Playing** node, add a **Branch**. 
-
-1.	Off the **False** of the new **Branch** node, add the **Montage Play** node with **Montage to Play** set to your Montage. 
-
-	![](Animations23.png)
-
-	This will check now if the Montage is currently playing and if it is, do not play it (and if it is not, play the Montage). 
-
-1.	**Compile** and **Save**, close the **Blueprint** and delete the existing **ThirdPersonCharacter** in the scene. 
-
-1. **Play** in the editor. 
-
-	[OBJECT:EmbeddedVideo]
-	[PARAMLITERAL:width]
-	640
-	[/PARAMLITERAL]
-	[PARAMLITERAL:height]
-	360
-	[/PARAMLITERAL]
-	[PARAMLITERAL:videoid]
-	N96shYlF0Aw
-	[/PARAMLITERAL]
-	[/OBJECT]
-
-You should now see that when you hold the **Left Mouse Button** down, the character will play the shooting AnimMontage. You can walk around while shooting, you can even crouch (by holding **C**) and shoot as well. 
-
-You could use this same method to play a melee attack or reload animation, by creating an **AnimMontage** for the animation you want to blend with the movement animation set and assigning it to the **UpperBody** slot as done in Step 8 above. Because the Slot was set up in the AnimGraph of the character's Animation Blueprint, any animation that is set to use that particular Slot is passed to the Animation Graph and blended in when called. 
-
-
-
-
-
-
-
-
-
+-->
+
+[OBJECT:TopicButtonList]
+	[PARAM:icon]
+		![](%ROOT%/build_icon.png)(convert:false)
+	[/PARAM]
+	[PARAM:title]
+		Steps
+	[/PARAM]
+	[PARAM:description]
+	[/PARAM]
+	[PARAM:links]
+	* [](Engine\Animation\AnimHowTo\AdditiveAnimations\01_BlendSetup "%Engine\Animation\AnimHowTo\AdditiveAnimations\01_BlendSetup:description%")
+	* [](Engine\Animation\AnimHowTo\AdditiveAnimations\02_CharacterBlueprintSetup "%Engine\Animation\AnimHowTo\AdditiveAnimations\02_CharacterBlueprintSetup:description%")
+	* [](Engine\Animation\AnimHowTo\AdditiveAnimations\03_AnimGraphSetup "%Engine\Animation\AnimHowTo\AdditiveAnimations\03_AnimGraphSetup:description%")
+	* [](Engine\Animation\AnimHowTo\AdditiveAnimations\04_BlendSettings "%Engine\Animation\AnimHowTo\AdditiveAnimations\04_BlendSettings:description%")
+	[/PARAM]
+[/OBJECT]
 
 
 
diff --git a/Engine/Documentation/Source/Engine/Animation/AnimHowTo/AimOffset/01_CreatingAimPoses/CreatingAimPoses.INT.udn b/Engine/Documentation/Source/Engine/Animation/AnimHowTo/AimOffset/01_CreatingAimPoses/CreatingAimPoses.INT.udn
new file mode 100644
--- /dev/null
+++ b/Engine/Documentation/Source/Engine/Animation/AnimHowTo/AimOffset/01_CreatingAimPoses/CreatingAimPoses.INT.udn
@@ -0,0 +1,141 @@
+Availability: Public
+Title:1 - Creating Aim Poses
+Crumbs: %ROOT%, Engine, Engine\Animation, Engine\Animation\AnimHowTo
+Description:In this step, we update our default playable character and create all the poses that will be needed for the Aim Offset. 
+version: 4.14
+skilllevel: Advanced
+type:howto
+tags:Animation
+tags:How-To
+tags:Aim Offset
+
+[VAR:Steps]
+[OBJECT:Navigation]
+	[PARAM:previous]
+	[/PARAM]
+	[PARAM:current]
+	[/PARAM]
+	[PARAM:home]
+		[Creating an Aim Offset - Home ](Engine\Animation\AnimHowTo\AimOffset)
+	[/PARAM]
+	[PARAM:next]
+		[Next](Engine\Animation\AnimHowTo\AimOffset\02_CreatingAimOffset)
+	[/PARAM]
+[/OBJECT]
+[/VAR]
+
+%Steps%
+
+In this step, we update our default playable character and create all the poses that will be needed for the Aim Offset. 
+
+[INCLUDE:Engine\Animation\AnimHowTo\AdditiveAnimations\01_BlendSetup#intro]
+
+## Steps
+
+1.  In the **Content Browser** under the **Content/ThirdPersonBP/Blueprints** folder, open the **ThirdPersonGameMode**.
+
+1.	Under **Default Pawn Class**, click the drop-down menu and select the **Ue4ASP_Character** then **Save** and close the Blueprint. 
+
+	![](Animations1.png)
+
+	We are changing the default playable character to use the one provided with the Animation Starter Pack. 
+
+1.	In the **Content Browser** open the **Content/AnimStarterPack** folder, then click **Add New** and create a folder called **AimPoses**.  
+
+	![](AimOffset1.png)
+
+1.	Drag the **Aim_Space_Hip** asset onto the **AimPoses** folder and select **Copy**. 
+
+	![](AimOffset2.png)
+
+	This is the animation sequence that contains a range of motions for aiming a weapon that we will chop up into poses. 
+
+1.	Inside the **AimPoses** folder, with **Aim_Space_Hip** selected, press **Ctrl+W** to duplicate it and name it **Aim_Center**. 
+
+	![](AimOffset3.png)
+
+1.	Open the **Aim_Center** asset, click the **Pause** button from the playback controls, then click the **ToFront** button. 
+
+	![](AimOffset4.png)
+
+1.	**Right-click** on the scrub bar in the timeline, then choose the **Remove from frame 1 to frame 87** option. 
+
+	![](AimOffset5.png)
+
+	With this, we now have a single frame containing a pose we can use in our Aim Offset. 
+
+1.	In the **AimPoses** folder with **Aim_Space_Hip** selected, press **Ctrl+W** to duplicate it and name it **Aim_Center_Down**.  
+
+1.	Open **Aim_Center_Down**, make sure it is on frame **0**, then click the **ToNext** button and go to frame **20**. 
+
+	![](AimOffset9.png)
+
+	The character will now be aiming downward in the viewport, which is the pose we want to use for aiming down. 
+
+1.	**Right-click** on the scrub bar in the timeline, then choose the **Remove frame 0 to frame 20** option. 
+
+	![](AimOffset10.png)
+
+	Depending on where you **Right-click** on the scrub bar, the frame range may be off by a frame or two, just make sure the character is aiming down. 
+	
+1.	**Right-click** on the scrub bar in the timeline again, then choose the **Remove from frame 1 to frame 68** option. 
+
+	![](AimOffset11.png)
+
+1.	In the **AimPoses** folder with **Aim_Space_Hip** selected, press **Ctrl+W** to duplicate it and name it **Aim_Center_Up**.  
+
+1.	Open **Aim_Center_Up**, make sure it is on frame **0**, then click the **ToNext** button and go to frame **10**. 
+
+1.	**Right-click** on the scrub bar in the timeline, then choose the **Remove frame 0 to frame 10** option. 
+
+	![](AimOffset14.png)
+
+1.	**Right-click** on the scrub bar in the timeline again, then choose the **Remove from frame 1 to frame 78** option. 
+
+	![](AimOffset15.png)
+
+1.	**Duplicate** the **Aim_Space_Hip** asset each time and create each of the remaining poses based on the table below.
+	
+	| Anim Name | Start At Keyframe | Remove Frames 1 | Remove Frames 2 |
+	| ------- | ------ | ------ | ------ |
+	| **Aim_Left_Center** | 30 | 0 - 30 | 1 - 57 |
+	| **Aim_Left_Up** | 40 | 0 - 40 | 1 - 48 |
+	| **Aim_Left_Down** | 50 | 0 - 50 | 1 - 37 |
+	| **Aim_Right_Center** | 60 | 0 - 60 | 1 - 27 |
+	| **Aim_Right_Up** | 70 | 0 - 70 | 1 - 17 |
+	| **Aim_Right_Down** | 80 | 0 - 80 | 1 - 8 |
+
+	For each animation, be sure to start at the suggested **Start At Keyframe** frame, then **Right-click** on the scrub bar and **Remove Frames 1**, **Right-click** on the scrub bar again and **Remove Frames 2**. Each of your animations should be a single frame with the character aiming in the direction of its corresponding name. 
+	When you have created each, your **AimPoses** folder should have several poses in it for each aiming direction. 
+
+	![](AimOffset17.png)
+
+1.	In the **Content Browser**, **Shift** and select each of the aim poses. 
+
+	![](AimOffset18.png)
+
+1.	**Right-click** and under **Asset Actions**, select the **Bulk Edit via Property Matrix** option. 
+
+	![](AimOffset19.png)
+
+1.	In the **Property Matrix**, under **Additive Settings**, change the **Additive Anim Type** to **Mesh Space**. 
+
+	![](AimOffset20.png)
+
+	[REGION:note]
+	For an animation to be compatible with an Aim Offset, it must be set to use an Additive Anim Type of **Mesh Space**. 
+	[/REGION:note]
+
+1.	Change the **Base Pose Type** to **Selected animation frame**, then under **Base Pose Animation**, click the pick asset icon and choose **Idle_Rifle_Hip**. 
+
+	![](add_01.png)
+
+	This Base Pose Type will define how the additive delta is to be calculated. 
+
+1.	Once you have done so, return to the **Content Browser** and click the **Save All** button. 
+
+	![](AimOffset22.png)
+
+With each of our pose assets created, in the next step, we will the poses to create our Aim Offset. 
+
+%Steps%
diff --git a/Engine/Documentation/Source/Engine/Animation/AnimHowTo/AimOffset/02_CreatingAimOffset/CreatingAimOffset.INT.udn b/Engine/Documentation/Source/Engine/Animation/AnimHowTo/AimOffset/02_CreatingAimOffset/CreatingAimOffset.INT.udn
new file mode 100644
--- /dev/null
+++ b/Engine/Documentation/Source/Engine/Animation/AnimHowTo/AimOffset/02_CreatingAimOffset/CreatingAimOffset.INT.udn
@@ -0,0 +1,128 @@
+Availability: Public
+Title:2 - Creating the Aim Offset
+Crumbs: %ROOT%, Engine, Engine\Animation, Engine\Animation\AnimHowTo
+Description:In this step, we take each of our aim poses and create our Aim Offset which will blend between each of the poses. 
+version: 4.14
+skilllevel: Advanced
+type:howto
+tags:Animation
+tags:How-To
+tags:Aim Offset
+
+[VAR:Steps]
+[OBJECT:Navigation]
+	[PARAM:previous]
+		[Previous](Engine\Animation\AnimHowTo\AimOffset\01_CreatingAimPoses)
+	[/PARAM]
+	[PARAM:current]
+	[/PARAM]
+	[PARAM:home]
+		[Creating an Aim Offset - Home ](Engine\Animation\AnimHowTo\AimOffset)
+	[/PARAM]
+	[PARAM:next]
+		[Next](Engine\Animation\AnimHowTo\AimOffset\03_ImplementingAimOffset)
+	[/PARAM]
+[/OBJECT]
+[/VAR]
+
+%Steps%
+
+In this step, we take each of our aim poses and create our Aim Offset which will blend between each of the poses. 
+
+## Steps
+
+1.	In the **Content/AnimStarterPack/UE4_Mannequin/Mesh** folder, **Right-click** on the **UE4_Mannequin** and select **Create Aim Offset**. 
+
+	![](AimOffset23.png)
+
+	This will create the Aim Offset asset based on this Skeleton. 
+
+1.	In the Aim Offset window (center window), under **Parameters**, enter the parameters shown below. 
+
+	![](AimOffset26.png)
+
+	Set the **X Axis Label** to **Yaw** and **Y Axis Label** to **Pitch**, then set both Axis Ranges to **-90 to 90** and click **Apply Parameter Changes**. 
+
+	[REGION:note]
+	Aim Offsets work very similar to [Blend Spaces](Engine/Animation/Blendspaces) and will allow us to blend between our poses based on our parameters. 
+	[/REGION]
+
+1.	In the **Asset Browser**, search for **Aim**, then drag the **Aim_Center** animation onto the graph in the center position as shown below. 
+	
+	![](AimOffset27.png)
+
+1.	Drag the **Aim_Center_Up** to position **1** and **Aim_Center_Down** to position **2** as shown below. 
+	
+	![](AimOffset28.png)
+
+	Optionally, you can turn off the tooltip in the grid by unchecking the **Enable Tooltip Display** option. 
+
+1.	Drag the **Aim_Left_Center** to position **1** and **Aim_Right_Center** to position **2** as shown below. 
+
+	![](AimOffset29.png)
+
+1.	Add the **Aim_Left_Up** (1), **Aim_Right_Up** (2), **Aim_Left_Down** (3) and **Aim_Right_Down** (4) poses to finish the Aim Offset. 
+
+	![](AimOffset30.png)
+
+1.	In the **Asset Details** panel, set the **Preview Base Pose** option under **Additive Settings** to **Idle_Rifle_Hip**. 
+
+	![](Add01.png)
+
+	You can preview the blending of poses by moving your mouse within the grid, the viewport mesh will update its pose based on your mouse position. 
+	
+The Aim Offset is now setup, in the next step we will hook it up for use within our **Animation Blueprint** so that gameplay can drive the blending. 
+
+%Steps%
+
+
+
+
+
+
+
+
+
+
+
+
+
+
+
+
+
+
+
+
+
+
+
+
+
+
+
+
+	
+
+
+
+
+	
+
+
+
+
+
+
+
+
+
+
+
+
+
+
+
+
+
+
diff --git a/Engine/Documentation/Source/Engine/Animation/AnimHowTo/AimOffset/03_ImplementingAimOffset/ImplementAimOffset.INT.udn b/Engine/Documentation/Source/Engine/Animation/AnimHowTo/AimOffset/03_ImplementingAimOffset/ImplementAimOffset.INT.udn
new file mode 100644
--- /dev/null
+++ b/Engine/Documentation/Source/Engine/Animation/AnimHowTo/AimOffset/03_ImplementingAimOffset/ImplementAimOffset.INT.udn
@@ -0,0 +1,149 @@
+Availability: Public
+Title:3 - Implementing the Aim Offset
+Crumbs: %ROOT%, Engine, Engine\Animation, Engine\Animation\AnimHowTo
+Description:In this step, we add the Aim Offset to our Animation Blueprint and hook it up to our existing Anim Graph. 
+version: 4.14
+skilllevel: Advanced
+type:howto
+tags:Animation
+tags:How-To
+tags:Aim Offset
+
+[VAR:Steps]
+[OBJECT:Navigation]
+	[PARAM:previous]
+		[Previous](Engine\Animation\AnimHowTo\AimOffset\02_CreatingAimOffset\)
+	[/PARAM]
+	[PARAM:current]
+	[/PARAM]
+	[PARAM:home]
+		[Creating an Aim Offset - Home ](Engine\Animation\AnimHowTo\AimOffset)
+	[/PARAM]
+	[PARAM:next]
+		[Next](Engine\Animation\AnimHowTo\AimOffset\04_FinishingUp)
+	[/PARAM]
+[/OBJECT]
+[/VAR]
+
+%Steps%
+
+In this step, we add the Aim Offset to our Animation Blueprint and hook it up to our existing Anim Graph. 
+
+## Steps
+
+1. In the **Content/AnimStarterPack** folder, open the **UE4ASP_HeroTPP_AnimBlueprint** and **Double-click** the **AnimGraph** in the **MyBlueprint** panel. 
+
+	![](step3_01.png)
+
+1. From the **Asset Browser** tab, drag in your Aim Offset asset. 
+
+	![](step3_02.png)
+
+1. Connect the Aim Offset as shown below, then **Right-click** on the **Yaw** and **Pitch** pins and **Promote to Variable** and name them **Aim Yaw** and **Aim Pitch**. 
+
+	![](step3_03.png)
+
+	These two variables will be populated by where the player is aiming using their mouse and will drive the pose in our Aim Offset. 
+
+1. In the **MyBlueprint** panel, jump over to the **EventGraph**. 
+
+	![](step3_04.png)
+
+1. In the **EventGraph**, locate the **Sequence** node within the Movement portion of the script. 
+
+	![](step3_05.png)
+
+	Click the **Add pin** button on the **Sequence** node. 
+
+	![](step3_05b.png)
+
+1. From the **MyBlueprint** panel, hold down **Alt** and drag in both the **Aim Yaw** and **Aim Pitch** and connect to **Sequence** node. 
+
+	![](step3_06.png)
+
+1. Drag off the **Cast To Ue4ASP_Character** node and add the **Get Control Rotation** and **Get Actor Rotation** nodes. 
+
+	![](step3_07.png)
+
+1. Drag off the **Get Control Rotation** node and add a **Delta(Rotator)** node. 
+
+	![](step3_08.png)
+
+1. Connect the **Get Actor Rotation** to the **Delta(Rotator) B Pin**, then off the **Return Value** add an **RInterp To** node (and swap the connection from Current to Target). 
+
+	![](step3_09.png)
+
+	[REGION:tip]
+	You can hold **Ctrl** and **Left-click** on the **Current** pin to drag its connection to the **Target** pin. 
+	[/REGION]
+
+1. Drag off the **Current** pin of the **RInterp To** node and select **Make Rotator**. 
+
+	![](step3_10.png)
+
+1. From the **MyBlueprint** panel, hold **Ctrl** and drag in the **Aim Pitch** and **Aim Yaw** variables and connect them to the **Make Rotator** node's **Pitch** and **Yaw**. 
+
+	![](step3_11.png)
+
+1. At the start of the **Movement** script, find the **Event Blueprint Update Animation** node and promote **Delta Time X** to a variable called **Time** and connect as shown. 
+
+	![](step3_12.png)
+
+1. Back on the **RInterp To** node, connect the new **Time** variable and set the **Interp Speed** to **15**. 
+
+	![](step3_13.png)
+
+	We are going to use the rotation of the character and the rotation of the player's input to create a new rotator that will limit how much the character can turn in a direction. 
+
+1. Off the **RInterp To** node, add a **Break Rotator** and **Clamp Angle** of **Pitch** and **Yaw** with min/max angles set to **-90** and **90** respectively. 
+
+	![](step3_14.png)
+
+1. Connect the **Clamp Angle** from the **Pitch** to **Aim Pitch** and the **Yaw** to **Aim Yaw**. 
+
+	![](step3_15.png)
+
+	The Aim Yaw and Aim Pitch values that drive our Aim Offset now take into account the character's rotation, the player's current controller rotation and are clamped to prevent the character from turning too far in one direction. 
+	Without clamping the angle degrees, the character's legs could be pointed forwards while the character turns all the way around and points in the opposite direction which is probably something we want to avoid. 
+
+We are just about done with our setup, in the next step, we will tie everything together and test out our character's ability to aim. 
+
+
+%Steps%
+
+
+
+
+
+
+
+
+
+
+
+
+	
+
+
+
+
+	
+
+
+
+
+
+
+
+
+
+
+
+
+
+
+
+
+
+
+
diff --git a/Engine/Documentation/Source/Engine/Animation/AnimHowTo/AimOffset/04_FinishingUp/FinishingUp.INT.udn b/Engine/Documentation/Source/Engine/Animation/AnimHowTo/AimOffset/04_FinishingUp/FinishingUp.INT.udn
new file mode 100644
--- /dev/null
+++ b/Engine/Documentation/Source/Engine/Animation/AnimHowTo/AimOffset/04_FinishingUp/FinishingUp.INT.udn
@@ -0,0 +1,132 @@
+Availability: Public
+Title:4 - Finishing Up
+Crumbs: %ROOT%, Engine, Engine\Animation, Engine\Animation\AnimHowTo
+Description:In this final section, we will update our Character Blueprint to change how controller rotation is handled as well as fix some minor issues before testing. 
+version: 4.14
+skilllevel: Advanced
+type:howto
+tags:Animation
+tags:How-To
+tags:Aim Offset
+
+[VAR:Steps]
+[OBJECT:Navigation]
+	[PARAM:previous]
+		[Previous](Engine\Animation\AnimHowTo\AimOffset\03_ImplementingAimOffset)
+	[/PARAM]
+	[PARAM:current]
+	[/PARAM]
+	[PARAM:home]
+		[Creating an Aim Offset - Home ](Engine\Animation\AnimHowTo\AimOffset)
+	[/PARAM]
+	[PARAM:next]
+	[/PARAM]
+[/OBJECT]
+[/VAR]
+
+%Steps%
+
+In this final section, we will update our Character Blueprint to change how controller rotation is handled as well as fix some minor issues before testing. 
+
+## Steps
+
+1.	In the **Content/AnimStarterPack** folder, open the **Ue4ASP_Character** Blueprint. 
+
+1.	Cick on **Ue4ASP_Character** in the **Components** window, then in the **Details** uncheck **Use Controller Rotation Yaw**. 
+
+	![](step4_04.png)
+
+	This will prevent the character from automatically turning to the location of the Controller's yaw. 
+
+1.	Cick on **CapsuleComponent** in the **Components** window, then in the **Details** check the **Hidden in Game** option. 
+
+	![](HideCapsule.png)
+
+	This will hide the debug collision display during gameplay. 
+
+1.	Click the **EventGraph** in **MyBlueprint**, then find the **Crouching** section and replace the **InputAction Crouch** node with a **C** Key Event. 
+
+	![](step4_05.png)
+
+	This will remove the yellow warning sign on the compile button in the upper left corner of the window since the project does not have an Action Mapping for Crouch by default, we are going to use the **C** button to crouch (you can use any Key Event you wish for this). 
+
+1.	**Compile** and **Save** then close the Blueprint 
+
+1.	Remove the **ThirdPersonCharacter** from the level.
+
+	![](DeleteCharacter.png)
+
+	This will make sure that we use the character assigned to Game Mode and not the one placed in the level. 
+
+1.	Click the **Play** button to play in the editor. 
+
+
+## End Result
+
+[OBJECT:EmbeddedVideo]
+[PARAMLITERAL:width]
+640
+[/PARAMLITERAL]
+[PARAMLITERAL:height]
+360
+[/PARAMLITERAL]
+[PARAMLITERAL:videoid]
+CEIeNE9lKxI
+[/PARAMLITERAL]
+[/OBJECT]
+
+When you play in the editor, the character will now aim while stationary and will react to mouse movement, pointing in the direction of your mouse. There are some restrictions in place that prevent the character from turning and facing backward while the mouse is behind them as well and while the character is jogging, moving the mouse also turns the character's torso in the direction you are pointing. 
+
+You can take this a few steps further by allowing the character to aim in a direction and play a shooting animation in that direction and (or) allow the character to play a shooting animation while moving in a direction or crouching by [layering animations](Engine\Animation\AnimHowTo\AdditiveAnimations) over the existing animations. Or, you can take a look at [Skeletal Mesh Sockets](Engine\Content\Types\SkeletalMeshes\Sockets) for examples on how to attach a weapon to the character's hands now that they can aim. 
+
+%Steps%
+
+
+
+
+
+
+
+
+
+
+
+
+
+
+
+
+
+
+
+
+
+
+
+
+
+	
+
+
+
+
+	
+
+
+
+
+
+
+
+
+
+
+
+
+
+
+
+
+
+
+
