Copying //UE4/Dev-Documentation to Samples-Main (//UE4/Samples-Main)

diff --git a/Engine/Documentation/Source/Engine/Animation/AnimHowTo/AimOffset/AimOffset.INT.udn b/Engine/Documentation/Source/Engine/Animation/AnimHowTo/AimOffset/AimOffset.INT.udn
--- a/Engine/Documentation/Source/Engine/Animation/AnimHowTo/AimOffset/AimOffset.INT.udn
+++ b/Engine/Documentation/Source/Engine/Animation/AnimHowTo/AimOffset/AimOffset.INT.udn
@@ -2,14 +2,15 @@ Availability: Public
 Title:Creating an Aim Offset
 Crumbs: %ROOT%, Engine, Engine\Animation, Engine\Animation\AnimHowTo
 Description:Using an Aim Offset, a character aims a weapon in the direction of your mouse or controller. 
-Related: Engine\Animation\AimOffset
-Related: Engine/Animation/AnimBlueprints
-Related: Engine\Content\Types\SkeletalMeshes\Sockets
-Related: Engine\Animation\AnimHowTo\AdditiveAnimations
-Version: 4.9
-SkillLevel: Advanced
+version: 4.14
+skilllevel: Advanced
+type:howto
+tags:Animation
+tags:How-To
+tags:Aim Offset
 
-In this guide, we will create an [**Aim Offset**](Engine/Animation/AimOffset), which is an asset that stores a blendable series of poses to help a character aim a weapon. We will take an animation sequence and chop it up into usable poses for an Aim Offset, get the player's pitch/yaw location of the mouse and use that to determine which blended pose to use so that the character moves and aims in the location of the mouse (with some restrictions) as seen in the example below. 
+In this guide, we will create an [Aim Offset](Engine/Animation/AimOffset), which is an asset that stores a blendable series of poses to help a character aim a weapon. 
+We will take an animation sequence and chop it up into usable poses for an Aim Offset, get the player's pitch/yaw location of the mouse and use that to determine which blended pose to use so that the character moves and aims in the location of the mouse (with some restrictions) as seen in the example below. 
 
 [OBJECT:EmbeddedVideo]
 [PARAMLITERAL:width]
@@ -19,439 +20,27 @@ In this guide, we will create an [**Aim Offset**](Engine/Animation/AimOffset), w
 360
 [/PARAMLITERAL]
 [PARAMLITERAL:videoid]
-uyeMVzNn3_s
+CEIeNE9lKxI
 [/PARAMLITERAL]
 [/OBJECT]
 
-[TOC(start:2 end:2)]
-
-[INCLUDE:Engine\Animation\AnimHowTo\AdditiveAnimations#intro]
-
-## Creating Aim Poses
-
-In this section, we will create the aiming poses needed for the Aim Offset by chopping up an animation sequence. 
-
-1.  With your project open, inside the **Content Browser** under the **Game/Blueprints** folder, open the **MyGame** Blueprint. 
-
-1.	Under **Default Pawn Class**, click the dropdown menu and select the **ASP_Character**. 
-
-	![](Engine\Animation\AnimHowTo\AdditiveAnimations\Animations1.png)
-
-1.	**Compile** and **Save** then close the Blueprint. 
-
-1.	In the **Content Browser** open the **Game/AnimStarterPack** folder. 
-
-1.	Click the **Create** button and create a folder called **AimPoses**. 
-
-	![](AimOffset1.png)
-
-1.	Drag the **Aim_Space_Ironsights** asset onto the **AimPoses** folder and select **Copy**. 
-
-	![](AimOffset2.png)
-
-	This is the animation sequence that contains a range of motions for aiming a weapon that we will chop up into poses. 
-
-1.	Inside the **AimPoses** folder, with **Aim_Space_Ironsights** selected, press **Ctrl+W** to duplicate it and name it **Aim_Center**. 
-
-	![](AimOffset3.png)
-
-1.	Open the **Aim_Center** asset, click the **Pause** button from the playback controls, then click the **ToFront** button. 
-
-	![](AimOffset4.png)
-
-	This will make sure that the animation is currently sitting at frame 0 of 102. 
-
-1.	**Right-click** on the scrub bar in the timeline, then choose the **Remove from frame 1 to frame 102** option. 
-
-	![](AimOffset5.png)
-
-1.	**Right-click** on the scrub bar in the timeline again, then choose the **Remove from frame 1 to frame 2** option. 
-
-	![](AimOffset6.png)
-
-	You should have a timeline that looks like below. 
-
-	![](AimOffset7.png)
-
-	With this, we now have a single frame containing a pose we can use in our Aim Offset. 
-
-1.	Click **Save** then close the **Aim_Center** animation. 
-
-1.	In the **AimPoses**, with **Aim_Space_Ironsights** selected, press **Ctrl+W** to duplicate it and name it **Aim_CenterDown**.  
-
-	![](AimOffset8.png)
-
-1.	Open **Aim_CenterDown**, make sure it is on frame **0**, then click the **ToNext** button and go to frame **20**. 
-
-	![](AimOffset9.png)
-
-	The character will now be aiming downward in the viewport, which is the pose we want to use for aiming down. 
-
-1.	**Right-click** on the scrub bar in the timeline, then choose the **Remove frame 0 to frame 19** option. 
-
-	![](AimOffset10.png)
-
-	Depending on where you **right-click** on the scrub bar, the frame range may be different, just make sure the character is aiming down. 
-	
-1.	**Right-click** on the scrub bar in the timeline again, then choose the **Remove from frame 1 to frame 83** option. 
-
-	![](AimOffset11.png)
-
-	Your **Aim_CenterDown** animation should now look like below in the timeline with the character aiming down in the viewport. 
-
-	![](AimOffset12.png)
-
-1.	In the **AimPoses**, with **Aim_Space_Ironsights** selected, press **Ctrl+W** to duplicate it and name it **Aim_CenterUp**.  
-
-	![](AimOffset13.png)
-
-1.	Open **Aim_CenterUp**, make sure it is on frame **0**, then click the **ToNext** button and go to frame **10**. 
-
-1.	**Right-click** on the scrub bar in the timeline, then choose the **Remove frame 0 to frame 10** option. 
-
-	![](AimOffset14.png)
-
-	Again, the frame range may be slightly different depending upon where you **right-click**, just make sure the character is aiming upwards.
-
-1.	**Right-click** on the scrub bar in the timeline again, then choose the **Remove from frame 1 to frame 92** option. 
-
-	![](AimOffset15.png)
-
-	Your **Aim_CenterUp** animation should now look like below in the timeline with the character aiming up in the viewport. 
-
-	![](AimOffset16.png)
-
-	Now that you have created an Aim Center, Aim Down and Aim Up pose; repeat the process cited above to create an Aim Left (Center, Up, and Down) and Aim Right (Center, Up, and Down). You can use the table below to identify suggested keyframes to start from and the keyframes to remove to generate your poses. 
-
-1.	**Duplicate** the **Aim_Space_Ironsights** asset each time and create each of the remaining poses based off the table below.
-	
-	| Anim Name | Start At Keyframe | Remove Frames 1 | Remove Frames 2 |
-	| ------- | ------ | ------ | ------ |
-	| **Aim_LeftCenter** | 35 | 0 - 34 | 1 - 68 |
-	| **Aim_LeftUp** | 40 | 0 - 39 | 1 - 63 |
-	| **Aim_LeftDown** | 50 | 0 - 49 | 1 - 53 |
-	| **Aim_RightCenter** | 65 | 0 - 64 | 1 - 38 |
-	| **Aim_RightUp** | 71 | 0 - 70 | 1 - 32 |
-	| **Aim_RightDown** | 81 | 0 - 80 | 1 - 32 |
-
-	For each animation, be sure to start at the suggested **Start At Keyframe** frame, then **Right-click** on the scrub bar and **Remove Frames 1**, **Right-click** on the scrub bar again and **Remove Frames 2**. Each of your animations should be a single frame with the character aiming in the direction of its corresponding name. 
-
-	When you have created each, your **AimPoses** folder should have several poses in it for each aiming direction. 
-
-	![](AimOffset17.png)
-
-1.	Open the **Aim_Center** animation from the **AimPoses** folder. 
-
-1.	In the **Anim Asset Details** window under **Additive Settings**, click the **No additive** option and change it to **Mesh Space**. 
-
-	![](AimOffset18.png)
-
-	[REGION:note]
-	For an animation to be compatible with an Aim Offset, it must be set to use an Additive Anim Type of **Mesh Space**. 
-	[/REGION:note]
-
-1.	Click the **Reference Pose** option and change it to **Selected animation frame**. 
-
-	![](AimOffset19.png)
-
-1.	Click the **None** dropdown option, then search for and select the **Idle_Rifle_Ironsights** animation to use as the Base Pose. 
-
-	![](AimOffset20.png)
-
-1.	In the **Asset Browser**, one-by-one, open each aim pose and repeat steps 22 - 24. 
-
-	![](AimOffset21.png)
-
-	Make sure that each is set to use an **Additive Anim Type** of **Mesh Space** and the **Idle_Rifle_Ironsights** Base Pose is assigned. 
-
-1.	Once you have done so, return to the **Content Browser** and click the **Save All** button. 
-
-	![](AimOffset22.png)
-
-## Creating the Aim Offset
-
-With the poses created, next we will create the Aim Offset asset and set up the poses for blending. 
-
-1.	In the **Content Browser**, click the **Create** button, then under **Animation** select the **Aim Offset**. 
-
-	![](AimOffset23.png)
-
-1.	In the **Pick Skeleton** window, make sure to select the **HeroTPP_Skeleton** that is in the **Game/AnimStarterPack/Character** folder. 
-
-	![](AimOffset24.png)
-
-	Name the Aim Offset, **Hero_AimOffset** then open it up. 
-
-1.	Inside the **Hero_AimOffset**, in the **Anim Asset Details** panel, click the **Preview Base Pose** under **Additive Settings**. 
-
-1.	In the pop-up window, search for and add the **Idle_Rifle_Ironsights** animation sequence. 
-
-	![](AimOffset25.png)
-
-	This will be used as the Preview Base Pose to blend from once we start adding our aim poses to the Aim Offset. 
-
-1.	In the **Hero_AimOffset** window (center window), under **Parameters**, enter the parameters shown below. 
-
-	![](AimOffset26.png)
-
-	Set the **X Axis Label** to **Yaw** and **Y Axis Label** to **Pitch**, then set both Axis Ranges to **-90 to 90** and click **Apply Parameter Changes**. 
-
-	[REGION:note]
-	For information on how **Blend Spaces** work (which are what Aim Offsets are based upon), see [**Blend Spaces**](Engine/Animation/Blendspaces). 
-	[/REGION]
-
-1.	In the **Asset Browser**, search for **Aim**, then drag the **Aim_Center** animation onto the graph in the center position as shown below. 
-	
-	![](AimOffset27.png)
-
-1.	Drag the **Aim_CenterUp** to position **1** and **Aim_CenterDown** to position **2** cited below. 
-	
-	![](AimOffset28.png)
-
-1.	Drag the **Aim_LeftCenter** to position **1** and **Aim_RightCenter** to position **2** cited below. 
-
-	![](AimOffset29.png)
-
-	If you move your mouse within the grid now, the character will start to aim and blend between the poses. 
-
-1.	Add the **Aim_LeftUp** (1), **Aim_RightUp** (2), **Aim_LeftDown** (3) and **Aim_RightDown** (4) poses to finish the Aim Offset. 
-
-	![](AimOffset30.png)
-
-	The Aim Offset is now set up, now we need to hook it up for use within our **Animation Blueprint**. 
-
-
-## Implementing the Aim Offset
-
-In this section, we will implement the Aim Offset inside the character's **Animation Blueprint**, so it is a part of their animation set. 
-
-1.	Inside the **Content Browser** under **Game/AnimStarterPack/Character**, open the **ASP_HeroTPP_AnimBlueprint**. 
-
-	![](AimOffset31.png)
-
-	This is the Animation Blueprint that is assigned to the playable character and governs the logic for the character's animations. 
-
-	[REGION:note]
-	For information on how **Animation Blueprints** work, see the [**Animation Blueprints**](Engine/Animation/AnimBlueprints) documentation. 
-	[/REGION]
-
-1.	Inside the Animation Blueprint, in the lower center window, click the **MyBlueprint** tab, then **Double-click** the **AnimGraph**. 
-
-	![](AimOffset32.png)
-
-	This will open the [**Anim Graph**](Engine/Animation/AnimBlueprints/AnimGraph) tab which is where we will implement our Aim Offset. 
-
-1.	In the **MyBlueprint** window, click the **Add Variable** button to create a new variable. 
-
-	![](AimOffset33.png)
-
-1.	In **Details** panel to the right, change the **Variable Name** to **PlayerAimYaw** and the **Variable Type** to a **Float**. 
-
-	![](AimOffset34.png)
-
-1.	In the **MyBlueprint** window, click the **Add Variable** button again and name the new variable **PlayerAimPitch** (also a **Float**). 
-
-1.	Click the **Asset Browser** tab, then search for **Hero** and drag the **Hero_AimOffset** into the graph. 
-
-	![](AimOffset35.png)
-
-1.	Click the **My Blueprint** tab, then drag the **PlayerAimYaw** variable onto the **Yaw** of the **Hero_AimOffset** node. 
-
-	![](AimOffset36.png)
-
-	Also drag the **PlayerAimPitch** variable onto the **Pitch** of the **Hero_AimOffset** node.
-	 
-1.	Connect the **Locomotion** pose to the **Base Pose** of the **Hero_AimOffset**, then the out pose to the **Result**.
-
-	![](AimOffset37.png)
-
-	This will take whatever the pose is from the **Locomotion** State Machine, then apply the **Aim Offset** on top of it as the final pose. 
-
-1.	Click the **MyBlueprint** tab in the lower center window, then **Double-click** the **EventGraph** to open it. 
-
-	![](AimOffset38.png)
-
-	This will open the [**Event Graph**](Engine/Animation/AnimBlueprints/EventGraph) tab which is used to keep the character animations updated with what is occurring in-game.
-
-1.	In the **EventGraph**, locate the **Sequence** node in the graph. 
-
-	![](AimOffset39.png)
-
-	Click the **Add pin** button on the **Sequence** node. 
-
-	![](AimOffset40.png)
-
-1.	From the **MyBlueprint** window, hold **Alt** and drag the **PlayerAimYaw** variable into the graph. 
-
-	![](AimOffset41.png)
-
-	This will create a **Set** node which we will use to get the player's current Yaw and determine the blend points on the Aim Offset. 
-
-1.	Hold **Alt** and drag in the **PlayerAimPitch**, then connect both to the **Then 2** pin of the **Sequence** node. 
-
-	![](AimOffset42.png)
-
-1.	Drag off the **As ASP Character** pin of the **Cast to ASP_Character** node and search for and add the **Get Control Rotation** node. 
-
-	![](AimOffset43.png)
-
-1.	Repeat the previous step, but search for and add the **Get Actor Rotation** node. 
-
-	![](AimOffset44.png)
-
-1.	**Right-click** in the graph and search for and add the **Delta (Rotator)** node. 
-
-	![](AimOffset45.png)
-
-1.	Connect the **Get Control Rotation** to the **A** and the **Get Actor Rotation** to the **B** of the **Delta(Rotator)** node. 
-
-1.	**Right-click** in the graph, search for and add the **RInterp To** node. 
-
-	![](AimOffset46.png)
-
-1.	**Right-click** in the graph, search for and add the **MakeRot** (Make Rotation) node. 
-
-	![](AimOffset47.png)
-
-1.	**Drag-and-drop** the **PlayerAimPitch** and **PlayerAimYaw** variables onto the **Pitch/Yaw** of the **MakeRot** node. 
-
-	![](AimOffset48.png)
-
-	Connect the nodes as shown above. 
-	
-	This will get the current value of **PlayerAimPitch** and **PlayerAimYaw** and smoothly transition to the player's actual Pitch/Yaw. 
-
-1.	At the beginning of the **Movement** section, drag off **Delta Time X** of **Event Blueprint Update Animation** and **Promote to variable**. 
-
-	![](AimOffset49.png)
-
-	We will use this to store a variable to Delta Time for use with our **RInterp** node. 
-
-1.	Name the variable **Time** and connect as shown below. 
-
-	![](AimOffset50.png)
-
-1.	Drag-and-drop the **Time** variable onto the **DeltaTime** pin of the **RInterp To** node and set **Interp Speed** to **15**. 
-
-	![](AimOffset51.png)
-
-	We have assigned a reference to time and set the speed at which to move from the player's old pitch/yaw location to their new location.
-
-1.	Drag off the **Return Value** of the **RInterp To** node and add the **Break Rot** (Break Rotation) node. 
-
-	![](AimOffset52.png)
-
-1.	Off the **Pitch**, add the **Clamp Angle** node. 
-
-	![](AimOffset53.png)
-
-	Set the **Min Angle** to **-90** and the **Max Angle** to **90**. 
-
-1.	**Copy/Paste** the **Clamp Angle** node and connect it to the **Yaw**, then connect both to the **PlayerAimYaw** and **PlayerAimPitch** nodes. 
-
-	![](AimOffset54.png)
-
-	The logic is setup for getting/using the player's pitch/yaw with our Aim Offset, just a few more things to finish up. 
-
-## Finishing Up
-
-In this final section, we will change the default idle pose and update Character Blueprint in order to use the Aim Offset correctly.  
-
-1.	Continuing from the previous section, click the **MyBlueprint** tab then **Double-click Locomotion** to open the graph. 
-
-	![](AimOffset55.png)
-
-1.	**Double-click** the **Idle** state inside the **Locomotion** graph. 
-
-	![](AimOffset56.png)
-
-1.	From the **Asset Browser**, drag in the **Idle_Rifle_Ironsights** and connect it to the **Result** pin. 
-
-	![](AimOffset57.png)
-
-	You can leave the **Play Idle_Rifle_Hip** node if you want to return to this later. 
-
-1.	**Compile** and **Save** then close the Animation Blueprint. 
-
-1.	Inside the **Content Browser** under **Game/AnimStarterPack/Character**, open the **ASP_Character** Blueprint. 
-
-	![](AimOffset58.png)
-
-1.	On the **Defaults** tab, search for **Rotation**, then uncheck the **Use Controller Rotation Yaw** option. 
-
-	![](AimOffset59.png)
-
-	This will prevent the character from automatically turning to the location of the Controller's yaw. 
-
-1.	Click the **Graph** tab, then find the **Crouching** section and replace the **InputAction Crouch** node with a **C** Key Event. 
-
-	![](AimOffset60.png)
-
-	This will remove the yellow warning sign on the compile button in the upper left corner of the window since the project does not have an Action Mapping for Crouch by default, we are going to use the **C** button to crouch (you can use any Key Event you wish for this). 
-
-1.	**Compile** and **Save** then close the Blueprint 
-
-	[REGION:note]
-	In Engine Version 4.8 and higher, you will need to remove the **ThirdPersonCharacter** from your level. This can be done by locating the **ThirdPersonCharacter** within the **Scene Outliner** of the **Editor Viewport**, clicking on the **ThirdPersonCharacter**, and pressing delete. 
-
-	This needs to be done due to the **ThirdPersonCharacter** being automatically possesed within the level. 
-
-	Be sure to **Save** the project once removing the **ThirdPersonCharacter** from the level.
-	[/REGION:note]
-
-1.	Click **Play** from the Main Editor Toolbar to play the game. 
-
-When you play in the editor, the character will now aim while stationary and will react to mouse movement, pointing in the direction of your mouse. There are some restrictions in place that prevent the character from turning and facing backwards while the mouse is behind them as well and while the character is jogging, moving the mouse also turns the character's torso in the direction you are pointing. 
-
-You can take this a few steps further by allowing the character to aim in a direction and play a shooting animation in that direction and (or) allow the character to play a shoot animation while moving in a direction or crouching by taking a look at the [**Using Additive Animations**](Engine\Animation\AnimHowTo\AdditiveAnimations) How To page to learn about blending multiple animations together. Or, you can take a look at [**Skeletal Mesh Sockets**](Engine\Content\Types\SkeletalMeshes\Sockets) for examples on how to attach an weapon to the character's hands now that they can aim. 
-
-
-
-
-
-
-
-
-
-
-
-
-
-
-
-
-
-
-
-
-
-
-
-
-
-
-
-	
-
-
-
-
-	
-
-
-
-
-
-
-
-
-
-
-
-
-
-
-
-
+Completing the following steps will result in a similar character as depicted above: 
+
+[OBJECT:TopicButtonList]
+	[PARAM:icon]
+		![](%ROOT%/build_icon.png)(convert:false)
+	[/PARAM]
+	[PARAM:title]
+		Steps
+	[/PARAM]
+	[PARAM:description]
+	[/PARAM]
+	[PARAM:links]
+	* [](Engine\Animation\AnimHowTo\AimOffset\01_CreatingAimPoses "%Engine\Animation\AnimHowTo\AimOffset\01_CreatingAimPoses:description%")
+	* [](Engine\Animation\AnimHowTo\AimOffset\02_CreatingAimOffset "%Engine\Animation\AnimHowTo\AimOffset\02_CreatingAimOffset:description%")
+	* [](Engine\Animation\AnimHowTo\AimOffset\03_ImplementingAimOffset "%Engine\Animation\AnimHowTo\AimOffset\03_ImplementingAimOffset:description%")
+	* [](Engine\Animation\AnimHowTo\AimOffset\04_FinishingUp "%Engine\Animation\AnimHowTo\AimOffset\04_FinishingUp:description%")
+	[/PARAM]
+[/OBJECT]
 
 
diff --git a/Engine/Documentation/Source/Engine/Animation/AnimHowTo/AnimBPOverride/AnimBPOverrideHowTo.INT.udn b/Engine/Documentation/Source/Engine/Animation/AnimHowTo/AnimBPOverride/AnimBPOverrideHowTo.INT.udn
--- a/Engine/Documentation/Source/Engine/Animation/AnimHowTo/AnimBPOverride/AnimBPOverrideHowTo.INT.udn
+++ b/Engine/Documentation/Source/Engine/Animation/AnimHowTo/AnimBPOverride/AnimBPOverrideHowTo.INT.udn
@@ -6,13 +6,15 @@ Related: Engine\Animation\Sequences
 Related: Engine\Animation\AnimBlueprints
 Related: Engine\Animation\Skeleton
 Related: Engine\Animation\Overview
-version: 4.9
+version: 4.14
 skilllevel: Advanced
+type:howto
+topic-image:animbpoverride_topic.png
+tags:Animation
+tags:Animation Blueprint
 
 When setting up characters and animation for those characters, there may be instances where you want to have one character perform an animation for an action and have another character perform a different animation for the same action. By using the **Asset Override Editor** inside of a child [Animation Blueprint](Engine/Animation/AnimBlueprints), you can override any previously set animation asset used in a [Sequence Player Node](Engine/Animation/Sequences/UserGuide) allowing you to quickly create variations of Animation Blueprints that each of your characters can use. 
 
-
-
 [REGION:note]
 For this guide, we are using the **Blueprint Third Person** template and included the **Infinity Blade: Warriors** and **Animation Starter Pack** assets which are available for free via the Marketplace. We also performed some [Animation Retargeting](Engine/Animation/AnimHowTo/Retargeting) and retargeted the Animation Blueprint and Animations that came with the Animation Starter Pack for use with the Infinity Blade characters. 
 [/REGION]
@@ -57,7 +59,7 @@ Below is an example where a character's default movement (assigned as Jog) in th
 360
 [/PARAMLITERAL]
 [PARAMLITERAL:videoid]
-MuQ5VQIGfY4
+2BNKUuTLGMM
 [/PARAMLITERAL]
 [/OBJECT]
 
diff --git a/Engine/Documentation/Source/Engine/Animation/AnimHowTo/AnimHowToOverview.INT.udn b/Engine/Documentation/Source/Engine/Animation/AnimHowTo/AnimHowToOverview.INT.udn
--- a/Engine/Documentation/Source/Engine/Animation/AnimHowTo/AnimHowToOverview.INT.udn
+++ b/Engine/Documentation/Source/Engine/Animation/AnimHowTo/AnimHowToOverview.INT.udn
@@ -2,10 +2,9 @@ Availability: Public
 Title:Animation How To's
 Crumbs: %ROOT%, Engine, Engine\Animation, Engine\Animation\AnimHowTo
 Description:Examples of different aspects of the Animation system in Unreal Engine 4 are demonstrated through several step-by-step How To guides.
-version: 4.12
 
-The Animation How To's provide several step-by-step examples of working with and using different aspects of [Persona](Engine/Animation), which is the animation tool inside Unreal Engine 4. 
-Whether you are new to animation in Unreal Engine 4 or want to learn how to work with certain aspects of Persona, the How To pages below cover a broad range of topics and aim to give you a place to start or expand upon your current knowledge. 
+The Animation How To's provide several step-by-step examples of working with and using the different [](Engine\Animation\Persona\Modes) within Unreal Engine 4. 
+Whether you are new to animation in Unreal Engine 4 or want to learn how to work with the different tools, the How To pages below cover a broad range of topics and aim to give you a place to start or expand upon your current knowledge. 
 
 Refer to the tables below for more information and click each link to view its contents. 
 
@@ -22,6 +21,7 @@ Refer to the tables below for more information and click each link to view its c
 | [](Engine\Animation\AnimHowTo\AimOffset) | %Engine\Animation\AnimHowTo\AimOffset:description% |
 | [](Engine\Animation\AnimHowTo\Retargeting) | %Engine\Animation\AnimHowTo\Retargeting:description% |
 | [](Engine\Animation\AnimHowTo\CreatePoseAsset) | %Engine\Animation\AnimHowTo\CreatePoseAsset:description% |
+| [](Engine\Animation\AnimHowTo\SubAnimInstance) | %Engine\Animation\AnimHowTo\SubAnimInstance:description% |
 
 
 ## Advanced
diff --git a/Engine/Documentation/Source/Engine/Animation/AnimHowTo/BlendSpace/BlendspaceHowTo.INT.udn b/Engine/Documentation/Source/Engine/Animation/AnimHowTo/BlendSpace/BlendspaceHowTo.INT.udn
--- a/Engine/Documentation/Source/Engine/Animation/AnimHowTo/BlendSpace/BlendspaceHowTo.INT.udn
+++ b/Engine/Documentation/Source/Engine/Animation/AnimHowTo/BlendSpace/BlendspaceHowTo.INT.udn
@@ -1,162 +1,170 @@
 Availability:Public
 Title:Creating a Blend Space
 Crumbs: %ROOT%, Engine, Engine\Animation, Engine\Animation\AnimHowTo
-Description:Using a Blend Space, a character blends between walking/jogging when a button is pressed. 
+Description:Example of a Blend Space used to toggle a character from walking/jogging animations smoothly. 
+version: 4.14
+skilllevel: Intermediate
+type:howto
+tags:Animation
+tags:How-To
+tags:Blend Space
 Related: Engine/Animation/Blendspaces
-Related: Engine/Animation/NodeReference/Blend
-Related: Engine/Animation/CharacterSetupOverview
-Related: Engine/Animation/AnimBlueprints
+Related: Engine/Animation/AimOffset
 Related: Engine/Animation/StateMachines
-version: 4.9
-skilllevel: Intermediate
+Related: Engine/Animation/Skeleton
 
+[Blend Spaces](Engine/Animation/Blendspaces) are special assets which can be sampled in **AnimGraphs** that allow for blending of animations based on the values of two inputs. 
+Simple blending between two animations based on a single input can be performed using one of the standard [Blend Nodes](Engine/Animation/NodeReference/Blend) available in Animation Blueprints. 
+Blend Spaces provide a means of doing more complex blending between multiple animations based on multiple values (currently limited to two).
 
-[Blend Spaces](Engine/Animation/Blendspaces) are special assets which can be sampled in **AnimGraphs** that allow for blending of animations based on the values of two inputs. Simple blending between two animations based on a single input can be performed using one of the standard [blend nodes](Engine/Animation/NodeReference/Blend) available in Animation Blueprints. Blend Spaces provide a means of doing more complex blending between multiple animations based on multiple values (currently limited to two).
+In this How-to, we use a Blend Space to blend between walking and jogging animations based on the character's movement speed and movement direction. 
 
-This page will show you how to take an existing character and apply a Blend Space which will blend between walking and jogging animations based on the characters movement speed. 
 
+[INCLUDE:Engine\Animation\AnimHowTo\AdditiveAnimations\01_BlendSetup#intro]
 
-[EXCERPT:Intro]
-For this guide, we have created a new project using the **Blueprint Third Person** template with **Starter Content** enabled. 
+## Steps
 
-We are also using the **Animation Starter Pack** which is available for download for free via the **Marketplace**. 
+1.  In your project in the **Content/ThirdPersonBP/Blueprints** folder, open the **ThirdPersonGameMode** and set the **Default Pawn Class** to **Ue4ASP_Character**. 
 
-![](Engine\Animation\AnimHowTo\AdditiveAnimations\AnimationAssetPack.png)
+	![](BlendSpace1.png)
 
-Once downloaded, you can add the Animation Starter Pack to your project via the **Add to project** button from the Launcher.
-[/EXCERPT:Intro]
+	This will allow us to use the playable character that was created along with the Animation Starter Pack project. 
 
-1.  With your project open, inside the **Content Browser** under the **Game/ThirdPersonBP/Blueprints** folder, open the **ThirdPersonGameMode** Blueprint. 
+1.	In the Level, select the **ThirdPersonCharacter** and delete it from the level. 
 
-1.	Under **Default Pawn Class**, click the dropdown menu and select the **Ue4ASP_Character**. 
+	![](BlendSpace2.png)
 
-	![](BS_DefaultPawnClass.jpg)
+	We need to delete this character from the level in order to spawn in as the character from the previous step. 
 
-1.	**Compile** and **Save** then close the Blueprint. 
+1.	In the **Content/AnimStarterPack/UE4_Mannequin/Mesh** folder, **Right-click** on the **UE4_Mannequin** and under **Create**, select **Blend Space**. 
 
-1.	In the **Content Browser**, click the **Add New** Button then under **Animation**, select **Blend Space**. 
+	![](BlendSpace3.png)
 
-	![](BS_BlendSpace.jpg)
+	Blend Space assets need a Skeleton asset to reference, here we are using a shortcut to create a Blend Space from this Skeleton asset. 
 
-1.	In the **Pick Skeleton** window, select the **UE4_Mannequin_Skeleton** that is located in the **Game/AnimStaterPack/UE4_Mannequin/Mesh** folder. 
+1.	In the **Parameters** section, apply the values below then press the **Apply Parameter Changes** button. 
 
-	![](BlendSpace2.png)(w:400)
+	![](BlendSpace4.png)
 
-	Give the Blend Space a name such as **Walk_Jog_BS**. 
+	* Set the **X Axis Label** to **Direction**
+	* Set the **X Axis Range** to **-180 - 180**
+	* Set the **Y Axis Label** to **Speed**
+	* Set the **Y Axis Range** to **150 - 260**
 
-	![](BS_Walk_Jog_BS.jpg)
+	This Blend Space using the parameters **Direction** and **Speed** to determine which pose we should be using or blending to or from. 
 
-1.	Open the **Walk_Jog_BS**, then in the center window under the **Parameters** section, enter the values shown. 
+1.	In the **Asset Browser**, drag the **Walk_Fwd_Rifle_Ironsights** animation to the center-bottom position on the grid. 
 
-	![](BlendSpace4.png)
+	![](BlendSpace5.png)
 
-	Set the **X Axis Label** to **Direction** and its range to **-180 to 180**. Set the **Y Axis Label** to **Speed** and its range to **150 to 270**. With the values set, click the **Apply Parameter Changes** button which will make the grid space to the left update with your values. Here we are defining the parameters that will determine the motion the Blend Space will produce based on the variables **Direction** and **Speed** (clamped by the range provided) which we will get from the player character. 
+	Placing this animation here indicates that when our Speed is 150 and are direction is not left or right, we should be walking forward with our rifle up at ironsights level. 
 
-1.	From the **Asset Browser** window, drag the **Walk_Fwd_Rifle_Ironsights** animation into the graph and the bottom center location. 
+1.	Drag the **Walk_Lt_Rifle_Ironsights** animation to location **1** and **Walk_Rt_Rifle_Ironsights** animation to location **2** below. 
 
-	![](BlendSpace5.png)
+	![](BlendSpace6.png)
 
-1.	Drag the **Walk_Lt_Rifle_Ironsights** animation to position 1 and the **Walk_Rt_Rifle_Ironsights** animation to position 2 below. 
+	Here we place the animations to use when moving left or right and our speed is 150. 
 
-	![](BlendSpace6.png)(w:640)
+1.	Drag the **Walk_Bwd_Rifle_Ironsights** animation to locations **1** and **2** below.
 
-1.	Drag the **Walk_Bwd_Rifle_Ironsights** animation to position 1 and 2 below. 
+	![](BlendSpace7.png)
 
-	![](BlendSpace7.png)(w:640)
+	Here we place the animation to use when moving backward and our speed it 150. 
 
-1.	Drag the **Jog_Fwd_Rifle** animation into the grid in the top center location. 
+1.	From the **Asset Browser**, drag the **Jog_Fwd_Rifle** animation to the top-center position on the grid. 
 
 	![](BlendSpace8.png)
 
-1.	Drag the **Jog_Lt_Rifle** animation to position 1 and the **Jog_Rt_Rifle** animation to position 2 below. 
+	This allows us to use a different animation when our speed is at max 270 and we are not moving left or right, we are jogging forward. 
 
-	![](BlendSpace9.png)(w:640)
+1.	Drag the **Jog_Lt_Rifle** animation to location **1** and **Jog_Rt_Rifle** animation to location **2** below. 
 
-1.	Drag the **Jog_Bwd_Rifle** animation to position 1 and 2 below.
+	![](BlendSpace9.png)
 
-	![](BlendSpace10.png)(w:640)
+1.	Drag the **Jog_Bwd_Rifle** animation to locations **1** and **2** below.
 
-	If you move your mouse around inside the grid, you will see the character start to blend between the different walk/jog animations. 
+	![](BlendSpace10.png)
 
-1.	Inside the **Anim Asset Details** window, under Sample Interpolation set **Target Weight Interpolation Speed Per Sec** to **4.0**. 
+1.	In the **Content/AnimStarterPack** folder, open the **UE4ASP_HeroTPP_AnimBlueprint** then in the **My Blueprint** panel, **Double-click** the **Locomotion** graph. 
 
 	![](BlendSpace11.png)
 
-	This will smooth out the transition between the animations in the Blend Space to prevent popping from pose-to-pose. 
+1.	**Double-click** the **Jog** State to open its graph. 
 
-1.	Click **Save** then close the Blend Space, then in the **Content Browser**, open the **Ue4ASP_HeroTPP_AnimBlueprint**. 
+	![](BlendSpace12.png)
 
-1.	Inside the Animation Blueprint, click the **MyBlueprint** tab, then expand **AnimGraph** and **Double-click Locomotion**. 
+	This [State Machine](Engine/Animation/StateMachines) already has a Jog State and in it, a Blend Space for different Jogging motions. 
 
-	![](BlendSpace12.png)(w:640)
+1.	From the **Asset Browser**, drag the **Walk_Jog_BS** into the graph and replace the **BS_Jog** Blend Space as shown below. 
 
-	This will open the [State Machine](Engine/Animation/StateMachines) used to drive the characters animation states. 
+	![](BlendSpace13.png)
 
-1.	**Double-click** on the **Jog** state to open its graph. 
+	The default Blend Space does not include Walking states which our Blend Space does, so let's replace it with ours! 
 
-	![](BlendSpace13.png)(w:640)
+1.	In the **Content/AnimStarterPack** folder, open the **Ue4ASP_Character**. 
 
-1.	Click the **Asset Browser** tab, then drag the **Walk_Jog_BS** into the graph. 
+1.	Locate the section marked **Crouch** in the graph and replace the **InputAction Crouch** node with a **C** Key Event. 
 
-	![](BlendSpace14.png)(w:640)
+	![](BlendSpace14.png)
 
-1.	Replace **R_Jog** with the **Walk_Jog_BS** Blend Space and connect **Direction** and **Speed** and feed it into the **Result**. 
+	By default, the Animation Starter Pack has bindings created for crouch which we will replace with a normal **C** press (or whatever key you want to crouch with). 
 
-	![](BlendSpace15.png)(w:640)
+1.	**Right-click** in the graph and add a **Left Shift** Key Event, then drag the **CharacterMovement** into the graph from the **Components** tab. 
 
-1.	Click **Compile** and **Save** then close the Animation Blueprint. 
+	![](BlendSpace15.png)
 
-1.	In the **Content Browser**, open the **Ue4ASP_Character** Blueprint. 
+	When we press and hold the Left Shift button, that will toggle us from walking into a jog. 
 
-1.	With the Blueprint open, find the **Crouching** section and replace the **InputAction Crouch** node with a **C** Key Event. 
+1.	Hold the **B** key and **Left-click** in the graph and create two **Branch** nodes, then drag in the **Crouch Button Down** variable and connect as shown below. 
 
-	![](BlendSpace17.png)(w:640)
+	![](BlendSpace16.png)
 
-	This AnimBlueprint was built with the expectation that there is an **Action Mapping** called **Crouch**, since it does not exist in a fresh project and to get rid of the yellow warning sign, we will replace it with the **C** Key Event so that the player crouches when C is pressed. 
+	We need to check and make sure that we are not crouching while holding shift (as for this example, sprinting can only occur when not crouched). 
 
-1.	Drag off the **Pressed** pin and add a **FlipFlop** node, then connect the **A** and **B** pins as shown below. 
+1.	Drag off the **CharacterMovement** node and **Set Max Walk Speed** for both **Branch** nodes with the values **260** and **150** respectively. 
 
-	![](BlendSpace18.png)(w:640)
+	![](BlendSpace17.png)
 
-	The **FlipFlop** node acts as a toggle causing the character to crouch the first time **C** is pressed, then stand the second time it is pressed. 
+	We are changing our movement speed based on if Left Shift is held and will alter our character's movement speed in game. 
 
-1.	Inside **MyBlueprint**, click the **Show inherited variables** box, then hold **Ctrl** and drag the **CharacterMovement** into the graph. 
-
-	![](BlendSpace19.png)
-
-1.	Drag off the **CharacterMovement** node in the graph, search for and add the **Set Max Walk Speed** node and set it to **125**. 
-
-1.	Duplicate the **Set Max Walk Speed** node and set the new node to **150**, then connect the nodes as shown below. 
+1.	**Right-click** in the graph and add an **Event Begin Play** node and connect a **CharacterMovement** component to **Set Max Walk Speed** to **150**. 
 
 	![](BlendSpace20.png)
 
-1.	Below the Crouch section, **Right-click** in the graph and add the **Left Shift** Key Event. 
-
-1.	Drag off the **Pressed** pin and add a **Branch** node, then also add a **Branch** of the **Released** pin. 
-
-1.	For the **Condition** for both, hold **Ctrl** and drag in the **Crouch Button Down** variable and connect it to both **Condition** pins. 
+	To start the game, we need to update our character's default movement speed so they start out walking. 
 
-	![](BlendSpace21.png)
+1.	Click the **CapsuleComponent** in the **Components** tab and in the **Details** panel, check the **Hidden in Game** option. 
 
-1.	Duplicate the **CharacterMovement** and **Set Max Walk Speed** nodes from the Crouch section. 
+	![](BlendSpace18.png)
 
-1.	Connect off the **False** pins of each **Branch** as shown below and set the **Max Walk Speed** values to **270** and **150**. 
+	This will turn off the collision capsule debug display in game. 
 
-	![](BlendSpace22.png)
+1.	Click the **CharacterMovement** component, then in the **Details** panel set the **Max Walk Speed** to **260**. 
 
-1.	In the upper-right of the Blueprint, click the **Defaults** tab. 
-
-1.	Enter **Max Walk Speed** in the search field, then set **Max Walk Speed** to **150** and **Max Walk Speed Crouched** to **125**. 
+	![](BlendSpace19.png)
 
-	![](BlendSpace23.png)
+	Here we are capping the maximum movement speed to reflect what is in our Blend Space. 
 
-1.	Click **Compile** and **Save** then close the Blueprint, then delete the existing ThirdPersonCharacter from the scene.
+1.	Click **Compile** then click the **Play** in editor button from the Toolbar. 
 
-1. Click **Play** to play in the editor. 
 
-When playing in the Editor, using **WASD** to move around, the character will now walk by default and blend between walking animations when changing directions. If you press **C** to crouch, the character will now stay in a crouched position while walking at a slightly reduced speed than found when standing/walking. While standing, pressing and holding **Left Shift** while using WASD will cause the character to blend from the walk/aim animation to a lowered/aim jog animation, essentially creating a "sprint" for the character. 
+## End Result
 
+[OBJECT:EmbeddedVideo]
+[PARAMLITERAL:width]
+640
+[/PARAMLITERAL]
+[PARAMLITERAL:height]
+360
+[/PARAMLITERAL]
+[PARAMLITERAL:videoid]
+vn0_YKcu1RU
+[/PARAMLITERAL]
+[/OBJECT]
 
+When playing in the Editor, using **WASD** to move around, the character will now walk by default and blend between walking animations when changing directions. 
+While standing, pressing and holding **Left Shift** while using WASD will cause the character to blend from the walk/aim animation to a lowered/aim jog animation, essentially creating a sprint for the character. 
+You can also press and hold **C**, to crouch walk using the Blend Space that was created as part of the Animation Starter Pack project, blending between crouch walking animations. 
 
 
 	
@@ -174,3 +182,4 @@ When playing in the Editor, using **WASD** to move around, the character will no
 
 
 
+
diff --git a/Engine/Documentation/Source/Engine/Animation/AnimHowTo/CreatePoseAsset/HowToCreatePoseAsset.INT.udn b/Engine/Documentation/Source/Engine/Animation/AnimHowTo/CreatePoseAsset/HowToCreatePoseAsset.INT.udn
--- a/Engine/Documentation/Source/Engine/Animation/AnimHowTo/CreatePoseAsset/HowToCreatePoseAsset.INT.udn
+++ b/Engine/Documentation/Source/Engine/Animation/AnimHowTo/CreatePoseAsset/HowToCreatePoseAsset.INT.udn
@@ -1,26 +1,25 @@
-Availability:Docs
+Availability:Public
 Title:Creating a Pose Asset
 Crumbs: %ROOT%, Engine, Engine\Animation, Engine\Animation\AnimHowTo
-Description:How to use Persona to generate a single frame Pose Asset.
-version: 4.13
+Description:How to use the Animation Editor to generate a single frame Pose Asset.
+version: 4.14
 skilllevel: Intermediate
 
-While you've been able to extract single frames of animation from an Animation Sequence for some time, you've been left to setup blends on your own. In 4.13, this changes with **Pose Assets**. These were implemented to support facial animation where either FACS (Facial Action Coding System) or viseme curves can drive poses. However, you could use this system to drive to create new animation by blending multiple poses. This supports bone transform as well as blendshapes. 
+While you've been able to extract single frames of animation from an Animation Sequence for some time, you've been left to set up blends on your own. 
+This changes with the use of **Pose Assets** which were implemented to support facial animation where either FACS (Facial Action Coding System) or viseme curves can drive poses. 
+However, you could use this system to drive and create new animation by blending multiple poses. Pose Assets also support bone transform as well as blendshapes which makes them an extremely flexible asset.
 
 This How-To will walk you through the creation of a **Pose Asset**.
 
 ## Steps
 
-1. From the Content Browser, you can right-click an **Animation Sequence** and create an **Pose Asset** from it:
+1. To create a series of Pose Assets from a single animation, from the Content Browser, you can right-click an **Animation Sequence** and create an **Pose Asset** from it:
 
 	[REGION:raw]
 	![](PoseAsset1.png)
 	[/REGION]
-	
+		
 
-1. Or, in Persona, you can use the **Create Asset** -> **Create PoseAsset** -> **From Current Pose** or **Animation** functions:
-
-	![](PoseAsset2.png)
 
 1. Once you create pose asset, the names will be auto generated by default. You can rename each pose or you can paste from clipboard. 
 
@@ -32,10 +31,18 @@ This How-To will walk you through the creation of a **Pose Asset**.
 
 ## Result
 
-You now have a new **Pose Asset** to use in your animations. To preview it make sure **Preview Post Asset** is set to your **Pose Asset**, and then add a **Variable Curve** to your animation that needs the **Pose Asset** data.
+You now have a new **Pose Asset** to use in your animations. 
+
+To preview it make sure **Preview Post Asset** is set to your **Pose Asset**, and then add a **Variable Curve** to your animation that needs the **Pose Asset** data.
 
 ![](PoseAsset5.png)
 
-For these to work at run time, you'll also need **PoseAsset** node in your Anim Blueprint.
+You can then add keyframes to the variable curve to drive the value of the pose weights, affecting the animation.
+
+For these to work at run time, you'll also need to include a **PoseAsset** node in your [Anim Blueprint](Engine\Animation\AnimBlueprints).
 
 ![](PoseAsset6.png)
+
+It's also worth noting you can create a **Pose Asset** in the **Animation Editor**. You can use the **Create Asset** -> **Create PoseAsset** -> **From Current Pose** or **Animation** options:
+
+![](PoseAsset2.png)
diff --git a/Engine/Documentation/Source/Engine/Animation/AnimHowTo/CurveDrivenAnimation/CurveDrivenAnimation.INT.udn b/Engine/Documentation/Source/Engine/Animation/AnimHowTo/CurveDrivenAnimation/CurveDrivenAnimation.INT.udn
new file mode 100644
--- /dev/null
+++ b/Engine/Documentation/Source/Engine/Animation/AnimHowTo/CurveDrivenAnimation/CurveDrivenAnimation.INT.udn
@@ -0,0 +1,111 @@
+Availability:Public
+Title: Curve Driven Animation
+Crumbs: %ROOT%, Engine, Engine/Animation
+Description:Illustrates how you can use a Pose Asset and Curve Data to create a new Animation Sequence driven by animation curve data.
+Related: Engine\Animation\AnimHowTo\CreatePoseAsset
+Related: Engine\Animation\Persona\AnimCurves
+Related: Engine\Animation\AnimPose
+version: 4.14
+skilllevel: Advanced
+
+
+In this How-to we create a new walking animation that relies purely on curve data to drive the animation. 
+By using a **Pose Asset** based on an existing **Animation Sequence**, we take two of the poses from the animation and blend between them to create our new walking motion. 
+While this is a **Full Body** example used to create locomotion, you can take these same concepts and apply them to a facial animation to create **Additive** facial animation that is driven by curve data. 
+
+[REGION:note]
+For this How-to, we are using the **Blueprint Thirdperson Template** project. 
+[/REGION]
+
+1. Inside the **Content Browser** under the **Content/Mannequin/Animations** folder, **Right-click** on the **ThirdPersonWalk** and select **Create PoseAsset**. 
+
+	![](CreatePoseAsset_01.png)
+
+1.	Use the default name for the new **Pose Asset**, then **Double-click** to open it.
+
+    [REGION:lightbox]
+	[![](CreatePoseAsset_2.png)(w:640)](CreatePoseAsset_2.png)
+	[/REGION]
+
+	_Click for full view._
+
+	When creating a **Pose Asset**, the **Pose Names** will be automatically generated by default (as long as the curves exist in the animation). 
+	A pose will be created for each frame of the animation in which the **Pose Asset** was created from. 
+	You can adjust the curve weights for each pose to see the resulting pose when weight influence is applied. 
+	For example, above we adjusted the weight of **Pose_6** and our character moves their right foot forward. 
+
+	[REGION:note]
+	By default, all **Pose Assets** are set to **Full Body** and are not **Additive** which means the weight values are either 0 (no influence) or 1 (full influence). 
+	Setting a weight value in this case to 0.2 for example is the same as setting it to 1. 
+	For more granularity, you can set the **Pose Asset** to **Additive** however doing so to an animation that is applied to the full body of a character (such as a walk or run) may cause your character to become compacted or scaled undesirably.
+	[/REGION]
+	
+1. **Right-click** on **Pose_6** and **Rename** it to **r_foot_fwd**. 
+
+	![](CreatePoseAsset_3.png)
+
+1. **Right-click** on **Pose_24** and **Rename** it to **l_foot_fwd**. 
+
+	![](CreatePoseAsset_Add01.png)
+
+1. From the **Main Toolbar**, click **Create Asset** and select **Create Animation** / **From Reference Pose**. 
+
+	![](CreatePoseAsset_4.png)
+
+1. Choose a save locaiton and save name (we called ours **CurveWalk**). 
+
+1. Inside the new animation, **Right-click** in the timeline and select **Append at the End** and add **30** frames. 
+
+	![](CreatePoseAsset_5.png)
+
+1. Under **Curves**, click the **Add** button, then under **Add Variable Curve..** select **l_foot_fwd** to add it then add **r_foot_fwd**. 
+
+	![](AddLeftCurve.png)
+
+	We now have the two poses that we want to drive with curve data. 
+
+1.	Click the **Expand Window** checkboxes for both the **l_foot_fwd** and **r_foot_fwd** curves. 
+
+	![](Checkboxes.png)
+
+1. Move the timeline to **0**, then in the **l_foot_fwd** curve, hold **Shift + Left Click** to create a key and set **Time** and **Value** to **0**. 
+
+	![](1stKey.png)
+
+1. **Shift + Left Click** and add keys for **Time 0.5** and **Value 1.0**, then **Time 1.0** and **Value 0.0** resepectively. 
+
+	![](2ndKey.png)
+
+	This will pose our character's left foot foward midway through the animation. 
+
+1.	In the **r_foot_fwd** curve, **Shift + Left Click** and add keys for: (**Time 0, Value 1**), (**Time 0.5, Value 0**), and (**Time 1, Value 1**). 
+
+	![](3rdKey.png)
+
+	At the start of our animation our character's right foot will be forward and midway through we switch to the left foot forward. 
+	As the animation nears the end, we use the curve data to drive the right foot forward pose again and when looped we generate our walk cycle. 
+
+1. To preview the animation with the Pose Asset, in the **Asset Details** panel, set the current **Preview Pose Asset** to use the **ThirdPersonWalk_PoseAsset**. 
+
+	![](CreatePoseAsset_7.png)
+
+
+## End Result
+
+We can now play the animation to see the two poses being blended together with our curve data. 
+
+[OBJECT:EmbeddedVideo]
+[PARAMLITERAL:width]
+640
+[/PARAMLITERAL]
+[PARAMLITERAL:height]
+360
+[/PARAMLITERAL]
+[PARAMLITERAL:videoid]
+7PRGvePzB2g
+[/PARAMLITERAL]
+[/OBJECT]
+
+To play this animation at runtime, you will need to use one of the [Pose Nodes](Engine\Animation\AnimPose\AnimBPPoseNodes\) within an Animation Blueprint. 
+
+
diff --git a/Engine/Documentation/Source/Engine/Animation/AnimHowTo/LayerEditing/LayerEditingHowTo.INT.udn b/Engine/Documentation/Source/Engine/Animation/AnimHowTo/LayerEditing/LayerEditingHowTo.INT.udn
--- a/Engine/Documentation/Source/Engine/Animation/AnimHowTo/LayerEditing/LayerEditingHowTo.INT.udn
+++ b/Engine/Documentation/Source/Engine/Animation/AnimHowTo/LayerEditing/LayerEditingHowTo.INT.udn
@@ -3,171 +3,103 @@ Title:Editing Animation Layers
 Crumbs: %ROOT%, Engine, Engine\Animation, Engine\Animation\AnimHowTo
 Description:An idle animation is edited to create a new reload animation through Animation Layer Editing.
 Related: Engine\Animation\Sequences
-Related: Engine\Animation\AnimMontage
 Related: Engine\Animation\AnimBlueprints
-Related: Engine\Animation\AnimHowTo\AdditiveAnimations
+Related: Engine/Animation/AnimMontage
 Related: Engine\Animation\Skeleton
 Related: Engine\Animation\Overview
-version: 4.9
+version: 4.14
 skilllevel: Advanced
+type:howto
+topic-image:LayeringTopic.png
+tags:Animation
+tags:Animation Blueprint
+tags:Animation Montage
 
-Animation Layer Editing allows you to modify any existing [](Engine\Animation\Sequences) inside Unreal Engine 4 as part of Additive Layer Tracks. With Animation Layer Editing you can take an existing animation and adjust it to produce a different take on that animation, or you can create a separate, new animation with its own motion you provide based off the original. 
+Animation Layer Editing allows you to modify any existing [](Engine\Animation\Sequences) inside Unreal Engine 4 as part of **Additive Layer Tracks**. 
+With Animation Layer Editing you can take an existing animation and adjust it to produce a different take on that animation, or you can create a separate, new animation with its own motion you provide based off the original. 
 
-This page will show you how to use the Animation Layer Editing tools and take an Idle animation and repurpose it into a Reload animation. 
-
-[TOC(start:2 end:2)]
+This page will show you how to use the Animation Layer Editing tools to take an Idle animation and repurpose it into a Reload animation. 
 
 [REGION:note]
 For this guide, we created a New Project using the **Blueprint First Person** template with **Starter Content** enabled. 
 [/REGION]
 
-## Creating the New Animation
-
-In this section, we will create a copy of the Idle animation that we can then modify, leaving the original intact. 
-
-1.  Inside the **Content Browser**, under the **Game/FirstPerson/Animations** folder, open the **FirstPerson_Idle** animation. 
-
-	![](AnimEdit1.png)
-
-1.	With the animation open, use the **Playback Controls** to pause the animation and set it to frame 0 so it is at the beginning. 
-
-	![](AnimEdit2.png)(w:640)
-
-1.	Click the **Create Asset** button from the Toolbar Menu, then under **Create Animation**, select **From Current Animation**. 
-
-	![](LE_FromCurrentAnimation.jpg)
-
-	Here we are creating a new animation based off the idle animation. There are also options for creating a new animation based off a specific pose within an animation or from the base reference pose. For our purposes, we want the entire idle animation from which to work with when modifying it to perform a reload motion, therefore **From Current Animation** is the best choice. 
-
-1.	In the **Create a New Animation Asset** window, select the **Animations** folder, name it **FirstPerson_Reload** then click **OK**. 
-
-	![](AnimEdit4.png)
-
-	The animation will be created and automatically open, allowing us to begin editing it. 
-
-## Editing the New Animation
-
-With the **FirstPerson_Reload** animation created, we can begin editing it through Animation Layer Editing. 
-
-1.  With the **FirstPerson_Reload** animation open, click the **To Next** button in the **Playback Controls** and advance to frame 15. 
-
-	![](AnimEdit5.png)
-
-1.	In the **Skeleton Tree**, click the **upperarm_l** Bone so that it is selected. 
-
-	![](LE_Leftarm.jpg)
-
-	In order to start editing, you must first select the Bone that will be modified in the **Additive Layer Tracks** which we have done above. 
+## Steps
 
-	[REGION:tip]
-	You might want to also fly the camera around to get a better view of the left arm which will make editing easier. Inside the viewport, hold the **Right Mouse Button** and use the **WASD** keys to move the camera around to your desired view. 
-	[/REGION]
+1.	Inside the **Content/FirstPerson/Animation** folder, duplicate the **FirstPerson_Idle** animation and name it **FirstPerson_Reload**. 
 
-1.	From the Toolbar Menu, click the **Add Key** button. 
+	![](LayerAnimation_01.png)
 
-	![](AnimEdit7.png)
+1.	Inside **FirstPerson_Reload**, pause the animation at frame 0, click the **upperarm_l** bone in the **Skeleton Tree**, then click the **Key** button. 
 
-	This will add keys to the Additive Layer Tracks as shown below for the Bone's current location. 
+	![](LayerAnimation_02.png)
 
-	![](AnimEdit8.png)(w:520)
+	This will add **Transform**, **Rotation** and **Scale** keys for the current position of the upperarm_l to the **Additive Layer Tracks** section. 
 
-	Above, three separate tracks were added for the **upperarm_l** Bone: a **Translation** track (1), a **Rotation** track (2) and a **Scale** track (3).
+	![](LayerAnimation_03.png)
 
-1.	Click the **To Next** button to advance to frame 30. 
+1.	Move the timeline ahead to frame **15**. 
 
-	![](AnimEdit9.png)
+	![](LayerAnimation_04.png)
 
-1.	Inside the viewport, **Left-click** on the green arch of the **Rotation** widget. 
+1.	Rotate the **upperarm_l** bone **-70** degrees, then press the **Key** button. 
 
-	![](LE_LeftArmRot.jpg)
+	![](LayerAnimation_05.png)
 
-	It will change to an Axis Rotation widget, move the mouse to the left to the **-60.00** mark. 
+	This will add a key for the bone's rotation when the animation is at frame 15. 
 
-	![](LE_LeftArmMinus60.jpg)
+1.	Move the timeline ahead to frame **20**, rotate the **upperarm_l** bone **70** degrees to its default position and press the **Key** button. 
 
-	From a first person perspective, this will pull the hand down off screen. 
+	![](LayerAnimation_06.png)
 
-1.	From the Toolbar Menu, click the **Add Key** button (or press the **S** key). 
+	Over the course of 20 frames, the arm will move down then back up to its default position. 
 
-	![](AnimEdit12.png)(w:640)
+1.	**Right-click** on the timeline marker and **Remove from frame 20 to frame 92**. 
 
-	The new keys will be added to the Additive Layer Tracks and a curve from the first key to the second key will be created. 
+	![](LayerAnimation_07.png)
 
-1.	Click the **To Next** button on the **Playback Controls** to advance to frame 45. 
+	This will remove the rest of the frames from the animation, keeping only our reloading motion. 
 
-1.	In the viewport, **Left-click** on and drag the green arch of the **Rotation** widget, then rotate to the right to the **60.00** mark. 
+1.	In the **Content Browser**, **Right-click** on the **FirstPerson_Reload** animation and select **Create AnimMontage**. 
 
-	![](LE_LeftArmUp60.jpg)
+	![](LayerAnimation_08.png)
 
-1.	From the Toolbar Menu, click the **Add Key** button (or press the **S** key). 
+	We are going to use this to play our new reload animation when the player presses the designated key during gameplay. 
 
-	![](AnimEdit15.png)(w:640)
+1.	Open the new AnimMontage, then change the slot to use the **DefaultGroup.Arms** slot. 
 
-	The new keys will be added and should be the same values as the first set of keys. 
+	![](LayerAnimation_09.png)
 
-1.	Click the **To Next** button on the **Playback Controls** to advance to frame 50, then press **S** to add a set of keys. 
+	Below, the Animation Blueprint assigned to the template character is set up to use a [Slot Node](https://docs.unrealengine.com/latest/INT/Engine/Animation/NodeReference/Blend/#slot) to play any Anim Montages assigned to the "Arms" slot. 
 
-	![](AnimEdit17.png)(w:640)
+	![](LayerAnimation_13.png)
 
-1.	Click the **To Next** button on the **Playback Controls** to advance to frame 60. 
+	This is currently being used to handle the playing of the firing weapon animation which we will also use for our reload motion. 
 
-1.	In the viewport, press **W** to switch to **Translation** mode, then **Left-click** and drag the red arrow to the right slightly. 
+1.	Inside the **Content/FirstPersonBP/Blueprints** folder, open the **FirstPersonCharacter** Blueprint. 
 
-	![](LE_LeftArmPump.jpg)
+	![](LayerAnimation_10.png)
 
-	This will shift the arm backwards as if "pumping" the shotgun. You can use your discretion to determine how much to move it.  
+	This is the Character Blueprint assigned to the template and is our playable character. 
 
-1.	From the Toolbar Menu, click the **Add Key** button (or press the **S** key). 
+1.	Locate the **Spawn Projectile** script and **Ctrl** select the **Mesh2P**, **Get Anim Instance** and **Montage Play** nodes shown below. 
 
-	![](AnimEdit20.png)(w:640)
+	![](LayerAnimation_11.png)
 
-	The keys for the pose will be added to the Additive Layer Tracks along with a curve from the previous keys. 
+	These nodes get the animation instance assigned to the Skeletal Mesh and signals the playing of an Animation Montage, in this case, the fire weapon montage. 
+	We are going to use this same script but instead call our reload Animation Montage when the key we specify is pressed. 
 
-1.	Click the **To Next** button on the **Playback Controls** to advance to frame 65, then press **S** to add a set of keys. 
+1.	**Right-click** and copy the selected nodes, then paste them in an empty space in the graph. 
 
-	![](AnimEdit22.png)(w:640)
+1.	**Right-click** and add a **R** Key Event and connect them as shown to the pasted nodes, setting **Montage to Play** to your reload montage. 
 
-	The values for the new keys will be zeroed out, which you can verify by clicking on the key in the track. 
+	![](LayerAnimation_12.png)
 
-	![](AnimEdit23.png)(w:640)
+1.	**Compile** and **Play in the Editor**. 
 
-	You can also manually enter the **Time** or **Value** for a key by clicking on it in the track to get exact values if you wish. 
+## End Result
 
-1.	From the Toolbar Menu, click the green **Apply** button to apply the changes to the animation. 
-
-	![](AnimEdit24.png)
-
-	When keys and curves are added to an animation, the result is an additive motion applied to the base animation. By clicking the **Apply** button, the key and curve data will be baked into the final animation. 
-
-1.	Click the **Save** button as your animation is complete.
- 
-1. Inside of the **Game/FirstPerson/Animations** folder, **Double-click** the **FirstPerson_AnimBP** to open it.
-
-	![](AnimEdit28.png)
-
-1. Inside of the **Animation Blueprint Editor**, open the **Anim Graph** by **Double-clicking** it in the **My Blueprint** panel
-
-	![](AnimEdit29.png)
-
-1.  **Double-click** the **New State Machine** node in the **Anim Graph**.
-
-	![](AnimEdit30.png)
-
-1. Find the node titled **FPP_Idle** and **Double-click** to open it.
-
-	![](AnimEdit31.png)
-
-1. **Right-click** in the graph and create a **Play FirstPerson_Reload** node.
-
-	![](AnimEdit32.png)
-
-1. Replace the existing **Play FirstPerson_Idle** node with your new **Play FirstPerson_Reload** node by holding **Ctrl** and dragging while holding **Left-click** from the output of the **Play FirstPerson_Idle** node to the **Play FirstPerson_Reload** node.
-
-	![](AnimEdit33.png)
-
-1. **Compile** and **Save** the Animation Blueprint before closing it and then click **Play** to test out your new animation.
-
-Below is an example of the animation playing on a character in-game. 
+When playing in the editor, pressing the **R** key will drop the character's hand off screen before bringing it back up to the gun, simulating a reload motion. 
 
 [OBJECT:EmbeddedVideo]
 [PARAMLITERAL:width]
@@ -177,12 +109,10 @@ Below is an example of the animation playing on a character in-game.
 360
 [/PARAMLITERAL]
 [PARAMLITERAL:videoid]
-Soior02AUtE
+8nH8cw2bSB8
 [/PARAMLITERAL]
 [/OBJECT]
 
-## Conclusion & FAQ
-
 You can continue to modify this animation by moving to a frame on the timeline and selecting the Bone you want to modify and keying it, then moving the timeline forward some and moving the Bone to a new location and keying it again. The Bone will move from your first key to your second key along the curve created as the timeline progresses. You can modify as many Bones as you wish for more precision than our example, but the process of keying is the same. 
 
 ### FAQ
@@ -197,7 +127,7 @@ You can continue to modify this animation by moving to a frame on the timeline a
 
 **Can I temporarily disable or remove tracks?**
 
-*	Yes, you can do both from the **Track Options** menu located in the upper right-corner of the track window. 
+*	Yes, you can do both from the **Track Options** menu located in the upper right corner of the track window. 
 
 	![](AnimEdit25.png)
 
@@ -205,13 +135,15 @@ You can continue to modify this animation by moving to a frame on the timeline a
 
 *	Inside the **Show** menu options, you can enable the **Source Animation** and (or) **Baked Animation** for display in the viewport. The Source Animation can be seen as a white skeleton, while the Baked Animation can be seen as a blue skeleton as shown in the image below. 
 
-	![](LE_SourceBakedAnim.jpg)
+	![](LE_SourceBakedAnim.png)
+
+	Just make sure that you have enabled the display of **Bones** from the **Show -> Bone Hiearchy** menu. 
 
 **Can Animation Layer Editing be used to create root motion?**
 
 *	Yes. Select your Root Bone and key it, then scrub the timeline forward and move the Root Bone and re-key it. 
 
-	![](LE_ThirdPersonAnim.jpg)
+	![](LE_ThirdPersonAnim.png)
 
 	Above we have enabled the Source and Baked Animation preview where the white skeleton is the source animation. 
 
@@ -239,5 +171,6 @@ You can continue to modify this animation by moving to a frame on the timeline a
 
 
 
+
 
 
diff --git a/Engine/Documentation/Source/Engine/Animation/AnimHowTo/PreviewingAnimations/PreviewingHowTo.INT.udn b/Engine/Documentation/Source/Engine/Animation/AnimHowTo/PreviewingAnimations/PreviewingHowTo.INT.udn
--- a/Engine/Documentation/Source/Engine/Animation/AnimHowTo/PreviewingAnimations/PreviewingHowTo.INT.udn
+++ b/Engine/Documentation/Source/Engine/Animation/AnimHowTo/PreviewingAnimations/PreviewingHowTo.INT.udn
@@ -1,18 +1,19 @@
 Availability: Public
 Title:Previewing Animations
 Crumbs: %ROOT%, Engine, Engine\Animation, Engine\Animation\AnimHowTo
-Description:Shows how to preview animations and control animation previews and toggle animation information. 
+Description:Shows how to preview animations, control animation previews and toggle animation information. 
 Related: Engine/Animation/Overview
 Related: Engine/Animation/Persona/Viewport
 Related: Engine/Animation/Sequences
 Related: Engine/Animation/AnimMontage
 Related: Engine/Animation/Blendspaces
 Related: Engine/Animation/AimOffset
-version: 4.9
+version: 4.14
 skilllevel: Beginner
 
 
-This page covers the different methods in which you can load and preview animation assets inside the [Persona Animation Editor](Engine/Animation/Persona), which is Unreal Engine 4's Animation Editing Tool. The controls used to manipulate the playback or information displayed about the animation currently being previewed is also covered and additional links to more information are provided. 
+This page covers the different methods in which you can load and preview animation assets inside the **Animation Editor** within Unreal Engine 4. 
+The controls used to manipulate the playback or information displayed about the animation currently being previewed is also covered and additional links to more information are provided. 
 
 
 [TOC(start:2 end:2)]
@@ -21,133 +22,102 @@ This page covers the different methods in which you can load and preview animati
 
 The most common method of opening and previewing an animation asset is by **Double-clicking** on it inside the **Content Browser**. 
 
-![](PA_Intro.jpg)
+![](PA_Intro.png)
 
-This will open the asset selected inside the **Persona Animation Editor**. 
+This will open the asset selected inside the **Animation Editor**. 
 
-![](pa_DoubleClickOpen.jpg)
+![](pa_DoubleClickOpen.png)
 
-The selected asset will automatically start playing inside the [](Engine/Animation/Persona/Viewport) and auto-loop playback of the motion. 
+The selected asset will automatically start playing inside the [Animation Editor's Viewport](Engine/Animation/Persona/Viewport) and auto-loop playback of the motion. 
 
-## Preview from Persona
+## Preview from Animation Editor
 
-If **Persona** is already open, on the **Animation** tab, you can **Double-click** an asset in the **Asset Browser** to preview it in the Viewport. 
+If the **Animation Editor** is already open, you can **Double-click** an asset in the **Asset Browser** to preview it in the Viewport. 
 
-[REGION:lightbox]
-[![](PA_SwitchingAnim.jpg)](PA_SwitchingAnim.jpg)
-[/REGION:lightbox]
+![](PA_SwitchingAnim.png)
 
-_Click for full image._
+Also while in the **Animation Editor**, you can set the preview animation from the asset pop-up menu. 
 
-Also while on the **Animation** tab of **Persona**, you can set the preview animation from the asset pop-up menu. 
+![](PA_AssetToPreview.png)
 
-![](Preview6.png)
 
-After clicking the pick asset button, a pop-up menu will appear where you can select the animation asset to preview. 
+## Viewport Options
 
-![](PA_AssetToPreview.jpg)
+In addition to the standard [Editor Viewport](Engine/UI/LevelEditor/Viewports) controls and options, there are a few settings in particular that you may want to pay attention to when previewing animations. 
 
-If you are on a different tab of **Persona**, you can enable the **Asset Browser** window from the **Menu Toolbar** / **Window** option. 
+![](PreviewModeOptions.png)
 
-[REGION:lightbox]
-[![](PA_Skeleton.jpg)](PA_Skeleton.jpg)
-[/REGION:lightbox]
+When clicking the ![](DownArrow.png) in the viewport, you will have the option of defining the **Preview Scene Settings**, **Scene Setup** and **Turn Table** options. 
 
-_Click for full image._
+### Preview Scene Settings
 
-The **Asset Browser** window will then appear, which you can dock to any other window. 
+![](PreviewSceneSettings.png)
 
-[REGION:lightbox]
-[![](PA_SkelAssetBrowser.jpg)](PA_SkelAssetBrowser.jpg)
-[/REGION:lightbox]
+| Option | Description |
+|-------|------|
+| **Profile Name** | Name to identify the profile. |
+| **Directional Light Intensity** | Manually set the directional light intensity (0.0 - 20.0). |
+| **Direction Light Color** | Manually set the directional light color. |
+| **Sky Light Intensity** | Manually set the sky light intensity (0.0 - 20.0). |
+| **Rotate Sky and Directional Lighting** | Toggle rotating of the sky and directional lighting. Press **K** and drag for manual rotating of Sky and **L** for Directional Lighting. |
+| **Show Environment** | Toggle visibility of the environment sphere. |
+| **Show Floor** | Toggle visibility of the floor mesh. |
+| **Environment Cube Map** | Sets the environment cube map used for sky lighting and reflections. |
+| **Pose Process Settings** | Manually set [Post Process Settings](Engine/Rendering/PostProcessEffects). |
+| **Post Processing Enabled** | Whether or not the Post Processing should influence the scene. |
+| **Lighting Rig Rotation** | Current rotation value of the sky in degrees (0 - 360). |
+| **Rotation Speed** | Speed at which the sky rotates when rotating is toggled. |
 
-_Click for full image._
+Your settings can be saved as a **Profile** and you can add and switch between Profiles from the Profiles menu at the top of the window. 
 
-You can also get a smaller, looping preview directly inside the **Asset Browser** when highlighting an animation asset. 
+![](Profiles.png)
 
-![](PA_AssetPreview.jpg)
+### Scene Setup
 
+![](SceneSetup.png)
 
-## Persona Viewport Controls
+This window is where you can define the animation mode used (toggle between Reference Pose or Animation Asset). There is also an option for assigning the current animation to preview (which is the same as **Double-clicking** an animation asset from the **Asset Browser**. 
 
-Below are the controls for navigating the **Persona Viewport**, allowing you to move the camera around. 
+Another important feature of the Scene Setup window is the ability to change the **Preview Mesh** for your animation. If you have multiple Skeletal Meshes that share animations with one another and want to preview the animation using a different mesh, you can do so from the **Preview Mesh (Animation)** drop-down menu. 
 
-### Standard
+You can also setup and define **Preview Mesh Collections** to use as **Additional Meshes**. These are used to build collections of related Skeletal Meshes that are animated together (such as components of a character). You can create a new collection from this window or assign a previously created collection. 
 
-These controls represent the default behavior when clicking and dragging in the viewports with no other keys or buttons pressed. These are also the only controls that can be used to navigate the orthographic viewports.
 
-| Control | Action |
-| ------- | ------- |
-| **LMB + Drag** | Moves the camera forward and backward and rotates left and right. |
-| **RMB + Drag** | Rotates the viewport camera. |
-| **LMB + RMB + Drag** | Moves up and down. |
-| **F** |Focuses the camera on the Mesh in the viewport. |
+### Turn Table
 
-### Game-style
+This option allows you to set up auto-rotation of the preview mesh. 
 
-The **WASD** keys will feel natural to those who are used to playing shooter games on the PC. They are enabled by default and can be used whenever you are holding the **Right Mouse Button** (**RMB**). This means you will still use the mouse to turn the camera while navigating in this way.
+![](TurnTable.png)
 
-These controls are mirrored on the arrow keys and num pad to provide alternate access to them.
-
-| Control | Action |
-| ------- | ------ |
-| **W** &#124; **Numpad8** &#124; **Up** | Moves the camera forward. |
-| **S** &#124; **Numpad2** &#124; **Down** | Moves the camera backward. |
-| **A** &#124; **Numpad4** &#124; **Left** | Moves the camera left. |
-| **D** &#124; **Numpad6** &#124; **Right** | Moves the camera right. |
-| **E** &#124; **Numpad9** &#124; **Page Up** | Moves the camera up. |
-| **Q** &#124; **Numpad7** &#124; **Page Dn** | Moves the camera down. |
-| **Z** &#124; **Numpad1** | Zooms the camera out (raises FOV). |
-| **C** &#124; **Numpad3** | Zooms the camera in (lowers FOV). |
-
-### Pan, Orbit, and Zoom
-
-[INCLUDE:Engine\UI\LevelEditor\Viewports\ViewportControls#mayalikecontrols]
-
-
-## Persona Viewport Options
-
-This section illustrates the options contained within the **Persona Viewport**. 
-
-![](PA_PersonaViewportOptions.jpg)
-
-1.	View Options Toolbar
-2.	Movement Toolbar
-3.	Animation Display Information
-
-### View Options Toolbar
-
-[INCLUDE:Engine/Animation/Persona/Viewport#viewoptions]
-
-### Movement Toolbar
-
-[INCLUDE:Engine/Animation/Persona/Viewport#movement]
+You can **Play**, **Pause** or **Stop** the **Turn Table Mode** as well as define the **Turn Table Speed**. 
 
 ### Animation Display Information
 
 You can toggle the information displayed in the viewport from the **Show** option under **Display info**. 
 
-![](PA_AnimationDisplay.jpg)
+![](PA_AnimationDisplay.png)
 
 **Basic** will show a condensed version of the display info, while **Detailed** will show more information and **None** will hide the display info. 
 
-## Persona Playback Controls
+## Playback Controls
 
-When you open an **AnimSequence** or **AnimMontage**, at the bottom of the center window you will see the **Timeline** and **Control Bar** used to controlling the playback of the animation asset. You can play, pause, skip, scrub and loop the animation from the Timeline and Control Bar options highlighted in yellow below. 
+When you open an **AnimSequence** or **AnimMontage**, at the bottom of the center window you will see the **Timeline** and **Control Bar** used to control the playback of the animation asset. 
+You can play, pause, skip, scrub and loop the animation from the Timeline and Control Bar options highlighted in yellow below. You can also record and create a new animation based on your currently selected animation. 
 
 ![](Preview9.png)
 
-[INCLUDE:Engine\Animation\Persona\Viewport#previewcontrols]
+[INCLUDE:Engine\Animation\Sequences\Editor#timeline]
 
-## Blend Space & Aim Offsets
 
-You can preview the blending of animations when applied to a [Blend Space](Engine/Animation/Blendspaces) or [](Engine/Animation/AimOffset) inside **Persona** as long as the **Enable Preview BlendSpace** option is checked. 
+## Preview Blend Space & Aim Offsets
 
-![](Preview10_1.jpg)
+You can preview the blending of animations when applied to a [Blend Space](Engine/Animation/Blendspaces) or [](Engine/Animation/AimOffset) in the **Animation Editor** as long as the **Enable Preview BlendSpace** option is enabled.
+
+![](Preview10_1.png)
 
 To see the blending between animations, move the **Mouse Cursor** inside the **Parameters** grid. 
 
-![](Preview11_1.jpg)
+![](Preview11_1.png)
 
 Remove the Tooltip Display by unchecking the **Enable Tooltip Display** option. 
 
@@ -165,3 +135,5 @@ Both **Blend Spaces** and **Aim Offsets** can be previewed in the same fashion a
 
 
 
+
+
diff --git a/Engine/Documentation/Source/Engine/Animation/AnimHowTo/PropertyBasedBlending/AnimBP/AnimBPSetup.INT.udn b/Engine/Documentation/Source/Engine/Animation/AnimHowTo/PropertyBasedBlending/AnimBP/AnimBPSetup.INT.udn
--- a/Engine/Documentation/Source/Engine/Animation/AnimHowTo/PropertyBasedBlending/AnimBP/AnimBPSetup.INT.udn
+++ b/Engine/Documentation/Source/Engine/Animation/AnimHowTo/PropertyBasedBlending/AnimBP/AnimBPSetup.INT.udn
@@ -10,7 +10,7 @@ Related: Engine/Animation/StateMachines
 version: 4.10
 skilllevel: Intermediate
 checkpoint: animpropertyblend
-type:sequential
+type:multi-step
 order:2
 parent:Engine/Animation/AnimHowTo/PropertyBasedBlending
 
diff --git a/Engine/Documentation/Source/Engine/Animation/AnimHowTo/PropertyBasedBlending/Character/CharacterSetup.INT.udn b/Engine/Documentation/Source/Engine/Animation/AnimHowTo/PropertyBasedBlending/Character/CharacterSetup.INT.udn
--- a/Engine/Documentation/Source/Engine/Animation/AnimHowTo/PropertyBasedBlending/Character/CharacterSetup.INT.udn
+++ b/Engine/Documentation/Source/Engine/Animation/AnimHowTo/PropertyBasedBlending/Character/CharacterSetup.INT.udn
@@ -10,7 +10,7 @@ Related: Engine/Animation/StateMachines
 version: 4.10
 skilllevel: Intermediate
 checkpoint: animpropertyblend
-type:sequential
+type:multi-step
 order:3
 parent:Engine/Animation/AnimHowTo/PropertyBasedBlending
 
diff --git a/Engine/Documentation/Source/Engine/Animation/AnimHowTo/PropertyBasedBlending/PropertyBasedBlending.INT.udn b/Engine/Documentation/Source/Engine/Animation/AnimHowTo/PropertyBasedBlending/PropertyBasedBlending.INT.udn
--- a/Engine/Documentation/Source/Engine/Animation/AnimHowTo/PropertyBasedBlending/PropertyBasedBlending.INT.udn
+++ b/Engine/Documentation/Source/Engine/Animation/AnimHowTo/PropertyBasedBlending/PropertyBasedBlending.INT.udn
@@ -6,7 +6,7 @@ version: 4.10
 skilllevel: Intermediate
 order:1
 checkpoint: animpropertyblend
-type:howto
+type:how-to
 
 
 [Blend Spaces](Engine/Animation/Blendspaces) are special assets which can be sampled in **AnimGraphs** that allow for blending of animations based on the values of two inputs. Simple blending between two animations based on a single input can be performed using one of the standard [blend nodes](Engine/Animation/NodeReference/Blend) available in Animation Blueprints. Blend Spaces provide a means of doing more complex blending between multiple animations based on multiple values (currently limited to two).
diff --git a/Engine/Documentation/Source/Engine/Animation/AnimHowTo/Retargeting/RetargetingHowTo.INT.udn b/Engine/Documentation/Source/Engine/Animation/AnimHowTo/Retargeting/RetargetingHowTo.INT.udn
--- a/Engine/Documentation/Source/Engine/Animation/AnimHowTo/Retargeting/RetargetingHowTo.INT.udn
+++ b/Engine/Documentation/Source/Engine/Animation/AnimHowTo/Retargeting/RetargetingHowTo.INT.udn
@@ -1,29 +1,31 @@
 Availability: Public
 Title:Using Retargeted Animations
 Crumbs: %ROOT%, Engine, Engine\Animation, Engine\Animation\AnimHowTo
-Description:Examples of how to setup and use Retargeted Animations for multiple characters. 
+Description:Examples of how to set up and use Retargeted Animations for multiple characters. 
 Related: Engine\Animation\AnimationRetargeting
-Related: Engine\Animation\RetargetingDifferentSkeletons
 Related: Engine\Animation\Skeleton
 Related: Engine\Animation\Overview
-version: 4.9
+version: 4.14
 skilllevel: Intermediate
 
 **Animation Retargeting** is the process of repurposing existing animations for use between multiple characters, eliminating the need to create entirely new animations as you can instead share your animation assets between characters. 
 
-There are two forms of Animation Retargeting, one in which the Skeleton of the character you wish to share animations with is using the **Same Skeleton** as the character that the animations were originally created for. The other involves an intermediary object called a **Rig** which allows you to retarget animations from one character's Skeleton and pass that skeleton's Bone information to a **Different Skeleton** using the Rig which both Skeleton's share. 
+There are two forms of Animation Retargeting, one in which the Skeleton of the character you wish to share animations with is using the **Same Skeleton** as the character that the animations were originally created for. 
+The other involves an intermediary object called a **Rig** which allows you to retarget animations from one character's Skeleton and pass that Skeleton's Bone information to a **Different Skeleton** using the Rig which both Skeleton's share. 
 
 Refer to the two links below for the type of Animation Retargeting you are looking to achieve: 
 
-[TOC(start:2)]
+[TOC(start:2 end:2)]
 
 ## Retargeting Using the Same Skeleton
 
 This section will show you how to set up a **Skeleton** to use retargeted animations for all biped characters that use the Skeleton.
 
+### Steps
+
 1.  Inside the **Content Browser**, open the **Skeleton** asset that is being used for each of your characters. 
 
-1.	In the **Skeleton Tree** of **Persona**, check the **Show Retargeting Options** checkbox. 
+1.	In the **Skeleton Tree** window of the **Skeleton Editor**, check the **Show Retargeting Options** checkbox. 
 
 	![](HowTo1.png)
 
@@ -37,7 +39,7 @@ This section will show you how to set up a **Skeleton** to use retargeted animat
 
 	![](HowTo3.png)
 
-	This will make sure that the Pelvis bone sits at the right height and is still animated. Any other bone you want translation to be animated and retargeted should also use this setting. 
+	This will make sure that the bone sits at the right height and is still animated. Any other bone you want translation to be animated and retargeted should also use this setting. 
 
 1.	Find the **Root** bone, any **IK** bones, any **Weapon** bones you may be using or other marker-style bones and set them to **Animation**. 
 
@@ -45,39 +47,134 @@ This section will show you how to set up a **Skeleton** to use retargeted animat
 
 	By using **Animation** as the Bone Translation Retargeting Mode, the bone's translation comes from the animation data itself and is unchanged. 
 
-	Above is an example portion of the **Skeleton Tree** set to use the different **Translation Retargeting** methods. The methods used may vary and you may have to experiment with the different methods for different bones to produce your desired results.	
+	[REGION:note]
+	You may have to experiment with the different methods for different bones to produce your desired results.	
+	[/REGION]
 
-1.	Click the **Grid** icon next to the current Mesh to bring up the Mesh selection box and select a different Mesh (a different character). 
+1.	Click the **Scene Setup** button in the Viewport and select a different **Preview Mesh**. 
 
 	![](HowTo5.png)
 
 	Try out different animations on different **Skeletal Meshes** to ensure that your Skeleton is set up properly for each of your animations. 
 
-1.	Inside the viewport, click **Show** and check **NonRetarget Animation** and **Bones** to show the normal and retargeted Skeletons. 
+1.	Inside the Viewport, click **Show** and check **NonRetarget Animation** and under **Bone**, choose **All Hierarchy** to display all Bones. 
 
 	![](HowTo6.png)
+
+### End Result
 	
-	You should now see something similar to below on your Skeleton. 
+You should now see something similar to below on your Skeleton. 
 
 [INCLUDE:Engine\Animation\AnimationRetargeting#aftertable]
 
 Shown in beige above is the original Skeleton along with the current, retargeted Skeleton in white. When retargeting is applied to the characters, the differences in their proportions are taken out of the equation and animations applied play properly on each of the different characters and the original Skeleton (beige Skeleton) is depicted in the same position in both the Short and Tall Characters. 
 
-[REGION:note]
-The steps above are intended to quickly get you up and running on how to retarget animations for characters that share the same Skeleton asset. Please see the [](Engine\Animation\AnimationRetargeting) documentation for additional, more detailed information. 
-[/REGION]
-
 
 ## Retargeting Using Different Skeletons
 
-[INCLUDE:Engine\Animation\RetargetingDifferentSkeletons#guide]
+In this section, we will go through the process of retargeting animations from one character to another who both use different Skeleton assets. 
+We are also going to use two asset packs that are available for free via the **Marketplace** that have additional skeletons and animations in it that we can use to retarget. If you have two Skeleton assets ready to use, you can skip the note below and use your assets. If you do not, you can follow along with the steps by adding the two packs outlined in the note below. 
 
 [REGION:note]
-The steps above showcase how to retarget an animation between two different Skeletons however there are additional measures you can take to improve the quality of the retargeted animation. Please see the [](Engine\Animation\RetargetingDifferentSkeletons) documentation for more information. 
+In this example, we are using the following assets which are available for download for free from the **Marketplace**.
+
+* The Source Skeleton we are using is the **UE4_Mannequin_Skeleton** that is included in the **Animation Starter Pack**. 
+* The Target Skeleton we are using is the **SK_Mannequin_Skeleton** that is included in the **Infinity Blade: Warriors** pack.
+
+Once you have installed the packs above, you can add them to a project via the **Add to project** button in the Unreal Launcher. 
+
+![](AnimStarterPack.png)(w:300) ![](InfinityBladeWarriors.png)(w:300)
 [/REGION]
 
+### Steps
+
+1.	Inside the **Content Browser** under the **Content/AnimStarterPack/UE4_Mannequin/Mesh** folder, open the **UE4_Mannequin_Skeleton** asset. 
+
+	![](Retarget1.png)
+
+1.	Click the **Retarget Manager** button, then under **Set up Rig**, click the **Select Rig** dropdown and select **Select Humanoid Rig**. 
+
+	![](Retarget2.png)
 
+1.	In the **Content Browser** under the **Content/InfinityBladeWarriors/Character/CompleteCharacters** folder, open the **SK_Mannequin_Skeleton** asset.
+
+	![](Retarget3.png)
+
+1.	Click the **Retarget Manager** button, then under **Set up Rig**, click the **Select Rig** dropdown and select **Select Humanoid Rig**.
+
+	![](Retarget4.png)
+
+1.	With the Rig assigned, next you will need to match up the Bones from the Target to the Nodes on the Rig. 
+
+	![](Retarget5.png)
+
+	Since the Skeleton in the Infinity Blade: Warriors pack uses the same naming convention and Bone hierarchy, they are automatically detected and applied to the Rig. 
+
+	If your Bone naming convention differs from the Epic Skeleton, you may see "none" applied to each of the Bones in the Rig: 
+
+	![](Retarget6.png)
+
+	In this case, you will need to use the **Auto Mapping** feature or manually assign each Bone from the Skeleton that best matches the Node on the Rig. 
+
+	[REGION:tip]
+	When using the **Auto Mapping** feature, it is a good idea to double check which Bones were assigned to which Nodes on the Rig. 
+	Auto Mapping will try to get the best match based on Bone position and (or) Bone Naming Conventions and assign those to ones on the Rig, but they may not the ideal Bone to use in each case. 
+	[/REGION]
 	
+1.	Click the **Show Advanced** button to toggle between Advanced/Base configuration.
+
+	![](AdvancedConfig.png)
+
+	In this section, you can assign Bones from the Skeleton that corresponds to Nodes on the Rig for digits of fingers, IK Bones or Twist/Roll Bones. 
+
+	[REGION:note]
+	Similar to the Base configuration, the Bones within Advanced may be set to "none" depending on your Skeleton which you will have to set. 
+	[/REGION]
+
+1.	In the **Content Browser**, under **Content/AnimStarterPack**, select an animation that you want to retarget to the new Skeleton.
+
+1.	**Right-click** on the animation and choose **Retarget Anim Assets -> Duplicate Anim Assets and Retarget**. 
+
+	![](Retarget7.png)
+
+1.	In the **Select Skeleton** window, you should see your Skeleton in the upper portion of the window, select it and click **Retarget**. 
+
+	![](Retarget8.png)(w:840)
+
+	This will copy the animation to the specified folder location. You can assign a New Asset Name in this window as well. 
+
+	[REGION:note]
+	The **Target** Skeleton should have a similar Base pose as the Source to generate the best possible retargeted animation.
+	Above, both Skeletons are in an A-Pose, however, you may have a character that is an A-Pose and one that is in a T-Pose. 
+	You can use a [Retargted Base Pose](Engine/Animation/Persona/BasePoseManager) to fix this prior to retargeting animations. 
+	[/REGION]
+
+### End Result
+
+The new retargeted animation will be available to any Skeletal Meshes that use the Skeleton asset the animation was retargeted for. 
+
+[OBJECT:EmbeddedVideo]
+[PARAMLITERAL:width]
+640
+[/PARAMLITERAL]
+[PARAMLITERAL:height]
+360
+[/PARAMLITERAL]
+[PARAMLITERAL:videoid]
+FKlniDMGxL4
+[/PARAMLITERAL]
+[/OBJECT]
+
+
+While this example illustrates how to retarget an animation, the same **Right-click** method can be used on any animation asset and retargeting can be applied. 
+
+
+
+
+
+
+
+
 
 
 
diff --git a/Engine/Documentation/Source/Engine/Animation/AnimHowTo/SubAnimInstance/SubAnimInstanceHowTo.INT.udn b/Engine/Documentation/Source/Engine/Animation/AnimHowTo/SubAnimInstance/SubAnimInstanceHowTo.INT.udn
new file mode 100644
--- /dev/null
+++ b/Engine/Documentation/Source/Engine/Animation/AnimHowTo/SubAnimInstance/SubAnimInstanceHowTo.INT.udn
@@ -0,0 +1,124 @@
+Availability: Public
+Title:Using Sub Anim Instances
+Crumbs: %ROOT%, Engine, Engine\Animation, Engine\Animation\AnimHowTo
+Description:Showcases how you can use Sub Animation Blueprint Instances within other Animation Blueprints.
+Related: Engine\Animation\AnimBlueprints
+Related: Engine\Animation\Skeleton
+Related: Engine\Animation\Overview
+version: 4.14
+skilllevel: Intermediate
+type:howto
+tags:Animation
+tags:Animation Blueprint
+tags:Skeletal Controls
+
+As you start to create more and more complex Animation Blueprints for your characters, there may be times when you want to re-use portions of an Animation Blueprint within another Animation Blueprint. 
+Rather than re-creating the node network, you can split that network off into its own Animation Blueprint and use a **Sub Anim Instance** to access the Animation Blueprint when it is needed. 
+
+In this how-to we create an Animation Blueprint with logic to affect our playable character, then use a Sub Anim Instance inside the template's Animation Blueprint to call that logic. 
+
+[REGION:note]
+For this guide, we are using the **Blueprint Third Person** template with **Starter Content** enabled.
+[/REGION]
+
+## Steps
+
+1. In the **Content/Mannequin/Animations** folder, **Right-click** on and **Duplicate** the **ThirdPerson_AnimBP** and call it **AltAnimBP**. 
+
+	![](SubAnimBP01.png)
+
+	This AltAnimBP will be used as our Sub Anim Instance that will be called from the ThirdPerson_AnimBP. 
+
+	[REGION:warning]
+	When creating an Animation Blueprint to use as a Sub Anim Instance, you must use the same Skeleton Asset as the one in the Animation Blueprint you plan to incorporate it in. 
+	In our example, we duplicated the template's Animation Blueprint, however, when creating Animation Blueprints from scratch you will need to select a Skeleton. 
+	[/REGION]
+
+1. In the **AnimGraph** of **AltAnimBP**, **Right-click** and add a **Sub-Graph Input** node. 
+
+	![](SubAnimBP02.png)
+
+	This will create an **In Pose** on the Sub Anim Instance, allowing us to pass through external pose data to this Animation Blueprint.
+
+1. **Right-click** again and add a **FABRIK** skeletal control node. 
+
+	![](SubAnimBP03.png)
+
+	For this example, we will use IK to push the character's arms up into the air as they run around during gameplay. 
+
+1. Assign the following settings to the **FABRIK** node from the **Details** panel. 
+	
+	* Uncheck the **(As Pin) Effector Transform** 
+	* Set **Location** to **100, 0, 200**. 
+	* Set **Tip Bone** to **index_01_l**. 
+	* Set **Root Bone** to **clavicale_l**.
+
+	![](SubAnimBP04.png) 
+
+1. Copy the **FABRIK** node in the graph and make the following changes to it from the **Details** panel. 
+
+	* Set **Location** to **-100, 0, 200**. 
+	* Set **Tip Bone** to **index_01_r**. 
+	* Set **Root Bone** to **clavicale_r**.
+
+	![](SubAnimBP05.png)
+
+	This will ensure that both the Left and Right arms are affected. 
+
+1. **Right-click** on the **Alpha** pin of a **FABRIK** node and **Promote to Variable** called **Effect Alpha**. 
+
+	![](SubAnimBP05b.png)
+
+1. Recreate the node network shown below (the **Local to Component** and **Component to Local** nodes will be automatically created). 
+
+	![](SubAnimBP06.png)
+
+	If you **Compile** the Animation Blueprint, the character in the preview window will put their arms up. 
+
+	![](CompliedBlueprint.png)
+
+1. Inside the **Content Browser** under **Content/Mannequin/Animations**, open the **ThirdPerson_AnimBP**. 
+
+1. Inside the **AnimGraph**, **Right-click** and add the **Sub Anim Instance** node. 
+
+	![](SubAnimBP07.png)
+	
+	This is the node that we use to gain access to our Animation Blueprint we created.
+
+1. In the **Details** panel for the **Sub Anim Instance** node, set the **Instance Class** to **AltAnimBP**. 
+	
+	![](SubAnimBP08.png)
+	
+1. Click the **Expose** checkbox next to **EffectAlpha** to expose the property. 
+
+	![](SubAnimBP09.png)
+
+	By exposing properties, you can update those properties with data from the parent Animation Blueprint. 
+
+1. From the **MyBlueprint** panel, drag in the **Speed** variable and divide it by **600** then plug that into the **Effect Alpha**. 
+
+	![](SubAnimBP10.png)
+
+	Here we divide the character's movement speed by a value to produce a blend between the amount of Skeletal Control applied.
+
+1. **Compile** then **Play in the Editor**. 
+
+
+## End Result
+
+Below, as our character's speed increases it is divided by the specified amount before being used to drive the amount of Skeletal Control applied inside our Sub Anim Instance network.
+In this example we used the **Sub-Graph** input node to determine the resulting pose data, however, you can create an Animation Blueprint with its own State Machine and logic to generate the resulting pose data as well.
+
+[OBJECT:EmbeddedVideo]
+[PARAMLITERAL:width]
+640
+[/PARAMLITERAL]
+[PARAMLITERAL:height]
+360
+[/PARAMLITERAL]
+[PARAMLITERAL:videoid]
+FqVcYpYvzqA
+[/PARAMLITERAL]
+[/OBJECT]
+
+
