Copying //UE4/Dev-Documentation to Samples-Main (//UE4/Samples-Main)

diff --git a/Engine/Documentation/Source/Engine/Animation/AnimMontage/AnimMontage.INT.udn b/Engine/Documentation/Source/Engine/Animation/AnimMontage/AnimMontage.INT.udn
--- a/Engine/Documentation/Source/Engine/Animation/AnimMontage/AnimMontage.INT.udn
+++ b/Engine/Documentation/Source/Engine/Animation/AnimMontage/AnimMontage.INT.udn
@@ -9,7 +9,7 @@ Related: Engine/Animation/StateMachines
 Related: Gameplay/AI
 Related: Engine/Content/Tools/MayaRiggingTool
 Related: Gameplay/Input
-version: 4.9
+version: 4.14
 
 [TOC(start:2 end:2)]
 
@@ -21,12 +21,11 @@ Montages are animation assets that can be created and exist within the **Content
 
 Montages also provide a way to replicate [](Engine\Animation\RootMotion\) across networked games if your animation has Root Motion enabled. 
 
-
 ## Montage Example
 
 ![](Montage_Screen2.png)
 
-The image above is a melee attack with 3 sections [start, loop, and end]. When the player **Left-clicks**, it triggers the _Start_ section by default when you ask to play the montage. The middle section called _Loop_ (which is set to loop), immediately follows the _Start_ section and will loop there forever. If the player lets the mouse button go, it will stop, but you do not want to stop right away since the animation will pop in the middle of the loop. You would like to relink the _Loop_ to the _End_ section, where it will then transition out to the end and finish the animation. 
+The image above is a melee attack with 3 sections, labeled _Start_, _Loop_, and _End_; the divisions between the sections are marked in yellow. When the player presses the attack button, the animation is played by the game, and it begins at the _Start_ section. The looping middle section, called _Loop_, immediately follows the _Start_ section and will be programmed to repeat indefinitely. If the player releases the attack button, the animation will stop looping and eventually reach the _End_ section and transition out.
 
 Some additional uses for Montages include:
 * The ability to play an animation from within an AnimBlueprint's **EventGraph**.
@@ -40,7 +39,6 @@ Some additional uses for Montages include:
 A Montage applied to a character can also be seen on the [Animation Content Examples](Resources\ContentExamples\Animation) page under section 1.5.
 [/REGION]
 
-
 ## Montage Properties
 
 The following is a breakdown of the Montage asset properties. These are available in **Persona** when looking at a Montage, and can also be accessed by **Right-clicking** on a Montage within the **Content Browser** and choosing Properties from the context menu.
@@ -81,7 +79,7 @@ When looking at a Montage in Persona, it is useful to know what each area is and
 
 ## Montage Area
 
-The Montage Area breaks down like so:
+The Montage Area breaks down as follows:
 
 [REGION:raw]
 ![](MontageArea.png)
@@ -94,7 +92,6 @@ The Montage Area breaks down like so:
 You may have as many Slot Tracks as you like within a single Montage, each with their own name and containing their own unique animations. However, you may only have one Section Track for each Montage.
 [/REGION]
 
-
 ### Sections
 
 Montage Sections provide a way to break a Slot up into multiple portions of animation. Each Section has a name and location in the Slot's timeline. Using the name, you can either jump directly to a particular Section or queue it to play next, when the current segment is complete. In Blueprint, you can query the current Section, jump to a Section, or set the next Section that will play.
@@ -107,7 +104,7 @@ Sections are created by **Right-clicking** on the Section track and choosing **N
 ### Slots
 
 [EXCERPT:Slots]
-Within a Montage, a **Slot** is just a single track that can hold any number of animations. You can name this Slot, and then blend to those specific animations by calling the Slot name. A great example is having a character with a weapon reload animation. You may have different versions of the reload for when the player is standing, crouching, and lying prone. So long as all 3 of the animations used the same timing, you could place each one within a separate Slot in your Montage; the Slots could be named Standing, Crouching, and Prone. In your Animation Blueprint's **AnimGraph**, you can use the Slot node to determine which one you want to play based on your character's current state. When they are standing, you can use the result of the animation in the Standing Slot. When they are prone, you can see the result of the Prone Slot.
+Within a Montage, a **Slot** is just a single track that can hold any number of animations. you can pick any name for a Slot and then use that name to blend to those specific animations. A great example is having a character with a weapon reload animation. You may have different versions of the reload for when the player is standing, crouching, and lying prone. So long as all 3 of the animations used the same timing, you could place each one within a separate Slot in your Montage; the Slots could be named Standing, Crouching, and Prone. In your Animation Blueprint's **AnimGraph**, you can use the Slot node to determine which one you want to play based on your character's current state. When they are standing, you can use the result of the animation in the Standing Slot. When they are prone, you can see the result of the Prone Slot.
 
 It is an important point to keep in mind that although much of your Montage control will take place in the Animation Blueprint's **Event Graph**, Slots are actually handled within the **Anim Graph**. This is done by way of the Slot node, which simply takes in the name for a Slot. By positioning this node at a strategic point along your AnimGraph's execution, you can have multiple Montages that utilize the same Slot name.
 [/EXCERPT:Slots]
@@ -188,13 +185,36 @@ Please see [](Engine/Animation/Sequences\Curves) for more information.
 [/REGION]
 
 
+## Creating Variant Montages
+**Child Anim Montages** provide a convenient way to create animation variants without affecting gameplay, and without needing to recreate and maintain notifications or other details across multiple assets. An example use case for this might be varying the auto-attack animations on units in an RTS game, or changing a character's idle animation in an RPG or action game based on what kind of weapons or armor they are holding or wearing. A Child Animation Montage can be created from an existing Animation Montage via the "Create Child Montage" context-menu option.
+
+![](CreateChildMontage.png)
+
+The new montage will have the same name as the original, but with "_Child" added to the end.
+
+![](ChildAnimMontageCB.png)
+
+[REGION:note]
+Although Child Anim Montages are not a separate class from standard Anim Montages (both are `UAnimMontage`), Child Anim Montages cannot be used as parents for other Child Anim Montages.
+[/REGION]
+
+When editing a Child Anim Montage, all editable parts of the interface will be dimmed and deactivated except for the Slot Track. The Slot Track is in a special mode which enables animations to be overridden. New animations cannot be added, and the details of an animation (such as play rate) will not be editable. This helps to create varied animations with minimal effort and without affecting gameplay, since you can randomly pick between the original montage and its children freely, knowing that the timing of the animations themselves, and all notification events, will be identical. When editing a Child Anim Montage, animations that have not been changed from the parent will be shown in green, and animations that have been overridden will be shown in yellow.
+
+![](AnimMontageSlot.png)
+
+![](ChildAnimMontageSlot.png)
+
+The animation substitution in the picture above is indicated by the yellow box in the Child Anim Montage's editor window. 
+
+It is possible to substitute any animation for any other, even if they are of different lengths. The engine will automatically adjust the play rate, start time, and end time of the replacement animation so that it uses the equivalent portion and runs in the same time as the original would have. For example, if a one-second animation with a start time of 0.5 and an end time of 1.0 is replaced with a two-second animation, the two-second animation will play with a start time of 1.0, an end time of 2.0, and will play at double the speed of the one-second animation. The replacement animation will always fit perfectly into the original slot.
+
 ## Playing Montages
 
 Montages can be played in a number of different ways: 
 
-* You can assign them as the **Anim to Play** for a **Skeletal Mesh**.  
-* You can call them directly through **Blueprint** script or C++. 
-* You can incorporate them into the **Animation Graph** as part of your **Animation Blueprint**. 
+* You can assign them as the **Anim to Play** for a **Skeletal Mesh**.
+* You can call them directly through **Blueprint** script or C++.
+* You can incorporate them into the **Animation Graph** as part of your **Animation Blueprint**.
 
 [REGION:note]
 Please see [](Engine/Animation/AnimHowTo/AdditiveAnimations) for an example of how a Montage can be used inside your Animation Blueprints. This example illustrates a character that can perform two animations at the same time independently of one another (in this case the character can move and shoot at the same time).
@@ -202,44 +222,44 @@ Please see [](Engine/Animation/AnimHowTo/AdditiveAnimations) for an example of h
 
 ### Playing Multiple Montages
 
-Prior to Unreal Engine version 4.6, you could only play a single Montage at a time. However, as of 4.6 or later it is now possible to play as many Montages at you would like through the use of **Anim Slots** and **Anim Groups**. 
+Prior to Unreal Engine version 4.6, you could only play a single Montage at a time. However, as of 4.6 or later it is now possible to play as many Montages at you would like through the use of **Anim Slots** and **Anim Groups**.
 
 Here's a quick rundown of how it works: 
 
-* Animations are organized into Anim Groups. You can play a single Montage per group. If you want to play another Montage at the same time, you can create a new Anim Group for it. 
-* Anim Groups and Anim Slots are centralized into the **Anim Slot Manager** tab inside **Persona** (see [Anim Slot Manager](#animslotmanager) below). 
-* You can use the drop-down UI to select between valid slots. 
-* When creating a new Montage, Slot nodes start out with a single **Default Slot** so they are immediately functional. 
+1. Animations are organized into Anim Groups. You can play a single Montage per group. If you want to play another Montage at the same time, you can create a new Anim Group for it.
+1. Anim Groups and Anim Slots are centralized into the **Anim Slot Manager** tab inside **Persona** (see [Anim Slot Manager](#animslotmanager) below).
+1. You can use the drop-down UI to select between valid slots.
+1. When creating a new Montage, Slot nodes start out with a single **Default Slot** so they are immediately functional.
 
 ### Anim Slot Manager
 
-The **Anim Slot Manager** is where you will manage your Anim Slot and Anim Group names. It displays a hierarchical view and lets you make changes. You can open the Anim Slot Manager inside a Montage in **Persona**, by clicking the magnifying glass icon (pictured below). 
+The **Anim Slot Manager** is where you will manage your Anim Slot and Anim Group names. It displays a hierarchical view and lets you make changes. You can open the Anim Slot Manager inside a Montage in **Persona**, by clicking the magnifying glass icon (pictured below).
 
 ![](OpenAnimSlotManger.png)
 
-This will open the Anim Slot Manager. 
+This will open the Anim Slot Manager.
 
 ![](AnimSlotManager.png)
 
-Groups and Slots are saved in the **Skeleton** asset so all Animation Blueprints and Montages from the same family can use the same settings. The **Save** button inside the Anim Slot Manager is there so that you can easily save the changes to the Skeleton asset whenever you make any changes. 
+Groups and Slots are saved in the **Skeleton** asset so all Animation Blueprints and Montages from the same family can use the same settings. The **Save** button inside the Anim Slot Manager is there so that you can easily save the changes to the Skeleton asset whenever you make any changes.
 
-When **Right-clicking** on a Slot name, a context menu is presented which will also allow you to move it to a different group if you desire. This makes it easier to manage the Slots you have created and place and move them into different groups and arrange them as you see fit. 
+When **Right-clicking** on a Slot name, a context menu is presented which will also allow you to move it to a different group if you desire. This makes it easier to manage the Slots you have created and place and move them into different groups and arrange them as you see fit.
 
 ![](SlotClicking.png)
 
 ### Editing Slots on Montage Tracks
 
-Below is an example of the **Montage Editor** inside **Persona**, the drop-down box (highlighted) lists the available Slots for that track. 
+Below is an example of the **Montage Editor** inside **Persona**, the drop-down box (highlighted) lists the available Slots for that track.
 
 ![](SlotExample1.png)
 
-If you need to create a new Slot or Group, clicking the magnifying glass icon will take you to the **Anim Slot Manager**. 
+If you need to create a new Slot or Group, clicking the magnifying glass icon will take you to the **Anim Slot Manager**.
 
-The status bar across the top will tell you which Group the Montage belongs to. This is determined by the first Track's Slot. 
+The status bar across the top will tell you which Group the Montage belongs to. This is determined by the first Track's Slot.
 
 ![](SlotExample2.png)
 
-If you have multiple Slots, the status bar will also tell you if there are any issues with the setup. For example, if you mix and match Slots that belong to different groups, a warning sign will appear to which you can hover over and see more information about the problem. 
+If you have multiple Slots, the status bar will also tell you if there are any issues with the setup. For example, if you mix and match Slots that belong to different groups, a warning sign will appear to which you can hover over and see more information about the problem.
 
 ![](SlotExample3.png)
 
@@ -379,14 +399,3 @@ And that is it! If we compile, we will now see that the character continues to s
 [/OBJECT]
 
 -------------------->
-
-
-
-
-
-
-
-
-
-
-
diff --git a/Engine/Documentation/Source/Engine/Animation/AnimPose/AdditivePA/AdditivePA.INT.udn b/Engine/Documentation/Source/Engine/Animation/AnimPose/AdditivePA/AdditivePA.INT.udn
new file mode 100644
--- /dev/null
+++ b/Engine/Documentation/Source/Engine/Animation/AnimPose/AdditivePA/AdditivePA.INT.udn
@@ -0,0 +1,34 @@
+Availability:Public
+Title: Additive vs. Full Body
+Crumbs: %ROOT%, Engine, Engine/Animation
+Description:Describes how the Additive and Full Body settings affect a Pose Asset. 
+Related: Engine\Animation\AnimHowTo\CreatePoseAsset
+Related: Engine\Animation\AnimHowTo\CurveDrivenAnimation
+Related: Engine\Animation\AnimPose\AnimBPPoseNodes
+Related: Engine\Animation\Persona\AnimCurves
+Related: Engine\Animation\AnimBlueprints
+version: 4.14
+
+Inside a [Pose Asset](Engine\Animation\AnimPose) you can specify if you want the animation to be **Additive** or non-additive in the form of **Full Body**.
+
+![](AdditiveImage.png)
+
+Usually, for full body animation, you will want to leave this as non-additive. This creates 0 or 1 weight values (0.2 is the same as 1 for example). 
+Setting to **Additive** and providing a value between 0 and 1 might cause your character to compact, expand or scale undesirably. 
+However, for facial animation, enabling **Additive** will give you more precision between min/max influence of a given curve. 
+
+[REGION:imagetable]
+| [REGION:lightbox][![](0_Weight.png)(w:210)](0_Weight.png)[/REGION] | [REGION:lightbox][![](025_Weight.png)(w:210)](025_Weight.png)[/REGION]|[REGION:lightbox][![](050_Weight.png)(w:210)](050_Weight.png)[/REGION]| [REGION:lightbox][![](075_Weight.png)(w:210)](075_Weight.png)[/REGION] |[REGION:lightbox][![](1_Weight.png)(w:210)](1_Weight.png)[/REGION]|
+| ----------- |------| ----------- |------|------|
+| **0.00** | **0.25** | **0.50** | **0.75** | **1.00** | 
+[/REGION]
+
+[REGION:caption]
+Above we set various weight values to control the amount of smile on the right side of the character's face. 
+[/REGION]
+
+You can also specify the Base Pose to use when converting to **Additive**. 
+
+![](AdditiveBasePose.png)
+
+
diff --git a/Engine/Documentation/Source/Engine/Animation/AnimPose/AnimBPPoseNodes/AnimBPPoseNodes.INT.udn b/Engine/Documentation/Source/Engine/Animation/AnimPose/AnimBPPoseNodes/AnimBPPoseNodes.INT.udn
new file mode 100644
--- /dev/null
+++ b/Engine/Documentation/Source/Engine/Animation/AnimPose/AnimBPPoseNodes/AnimBPPoseNodes.INT.udn
@@ -0,0 +1,122 @@
+Availability:Public
+Title: AnimGraph Pose Nodes
+Crumbs: %ROOT%, Engine, Engine/Animation
+Description:Describes how Pose Assets, Animation Blueprints, and Pose Nodes can be used to play curve driven animation at runtime.
+Related: Engine\Animation\AnimHowTo\CreatePoseAsset
+Related: Engine\Animation\Persona\AnimCurves
+Related: Engine\Animation\Sequences\Curves
+Related: Engine\Animation\AnimBlueprints
+Related: Engine\Animation\NodeReference\SkeletalControls
+version: 4.14
+
+[Pose Assets](Engine\Animation\AnimPose) can be driven at runtime with **Animation Blueprints** as the **AnimGraph** supports a few different Pose Nodes. 
+
+[TOC (start:2 end:2)]
+
+## Pose by Name / Pose Blender Nodes
+
+One way of driving **Pose Assets** is through the **Pose by Name** and **Pose Blender** nodes: 
+
+* **Pose by Name** - where you can get the transform of a pose by its name. 
+* **Pose Blender** - which outputs the correct pose driven by curves from a Source Pose. 
+
+In the example below, we have a face that uses a **Pose Asset** to smile when the player is close and frown when moving away. 
+
+[OBJECT:EmbeddedVideo]
+[PARAMLITERAL:width]
+640
+[/PARAMLITERAL]
+[PARAMLITERAL:height]
+360
+[/PARAMLITERAL]
+[PARAMLITERAL:videoid]
+d3oGXxq2ekw
+[/PARAMLITERAL]
+[/OBJECT]
+
+Our State Machine inside our Animation Blueprint is set up with three states: 
+
+![](StateMachine.png)
+
+Our default **State** which is our neutral expression, our **Smile** state which is entered when the player enters a trigger volume and our **Frown** state when the player leaves the trigger volume. 
+In this example, we only move between smiling and frowning states after triggering the first smile. Inside each state we use the **Pose Blender** and our animation sequence to apply the curve driven animation. 
+
+![](PoseBlender.png)
+
+Pose Blender nodes are created automatically when dragging a **Pose Asset** into the graph. 
+
+If you plug an AnimSequence directly into the Final Animation Pose without using a Pose Blender (below), you will not see the curve driven animation applied. 
+
+![](NoPoseBlender.png)
+
+Additionally, using just a Pose Blender without supplying any incoming curve data will not output your desired pose. 
+
+![](NoCurveData.png)
+
+You can also convert a Pose Blender node to a **Pose by Name** (which can also be converted to a Pose Blender node) through the **Right-click** context menu. 
+
+![](PoseByName.png)
+
+This will allow you to manually specify which pose to use inside the **Details** panel from the **Pose Asset**. 
+
+![](AdditivePoseByName.png)
+
+For our example we are controlling two aspects of our character's face, the right side of a smile and the left side of a smile so a **Pose Blender** make more sense to use.
+If we wanted to control only one pose within our **Pose Asset**, the **Pose by Name** grants us the ability to control a single pose by its name (which is usually good for full body context). 
+
+[REGION:note]
+If you are using a **Pose Asset** with **Additive** enabled, you will also need to use the **Apply Additive** node to correctly display your desired pose. For **Pose Assets** that are not **Additive**, you can directly link a **Pose by Name** node to the **Final Animation Pose** output. 
+[/REGION]
+
+
+## Pose Driver Node
+
+Available within the AnimGraph of an Animation Blueprint is the **Pose Driver** node which lets you drive a curve value based on a bone's movement. 
+One use for this node would be to create a pose space deformer to drive corrective blend shapes based on a bone's orientation. 
+For example, in the clip below, clavicle bone poses are being used to move the shoulder relative to the arm bone's movement. 
+
+[OBJECT:EmbeddedVideo]
+[PARAMLITERAL:width]
+640
+[/PARAMLITERAL]
+[PARAMLITERAL:height]
+360
+[/PARAMLITERAL]
+[PARAMLITERAL:videoid]
+tJSYgPmzoKc
+[/PARAMLITERAL]
+[/OBJECT]
+
+To achieve this, an RBF (Radial Basis Function) interpolates driven values based on the orientation of a target bone. 
+You will need to create a **Pose Asset** and define the target poses for the bone and the desired curve values at each pose. 
+
+On the **Pose Driver** node itself, you have the following options: 
+
+![](PoseDriverNode.png)
+
+| Option | Description |
+| ----------- |------|
+| **Source Bone** | Bone to use for driving parameters based on its orientation. | 
+| **Radial Scaling** | Scaling of radial basis, applied to max distance between poses.| 
+| **Pose Source** | The **Pose Asset** used as the source of pose data. | 
+| **Include Ref Pose as Neutral Pose** | Should the mesh ref pose of Source Bone be considered as a "neutral pose" (zero curves). | 
+| **Type** | **Swing and Twist** (Consider full rotation for interpolation) or **Swing Only** (Consider only swing for interpolation). | 
+| **Twist Axis** | The axis to use when **Type** is set to **SwingOnly** (X, Y or Z). | 
+
+For additional information see the Unreal Engine Live Stream on Animation Topics which include a **Pose Driver** example:
+ 
+[OBJECT:EmbeddedVideo]
+[PARAMLITERAL:width]
+640
+[/PARAMLITERAL]
+[PARAMLITERAL:height]
+360
+[/PARAMLITERAL]
+[PARAMLITERAL:videoid]
+h2egrj1pXzw
+[/PARAMLITERAL]
+[/OBJECT]
+
+
+
+
diff --git a/Engine/Documentation/Source/Engine/Animation/AnimPose/AnimPose.INT.udn b/Engine/Documentation/Source/Engine/Animation/AnimPose/AnimPose.INT.udn
new file mode 100644
--- /dev/null
+++ b/Engine/Documentation/Source/Engine/Animation/AnimPose/AnimPose.INT.udn
@@ -0,0 +1,154 @@
+Availability:Public
+Title: Animation Pose Assets
+Crumbs: %ROOT%, Engine, Engine/Animation
+Description:Describes the Animation Pose Asset which can be used to drive animation through weighted curve data. 
+Related: Engine\Animation\Persona\AnimCurves
+Related: Engine\Animation\AnimBlueprints
+version: 4.14
+
+Traditional animation consists of a series of keyframed poses that are blended over a timeline to achieve the desired motion. 
+In facial animation, curve driven animation is very popular as various weighted values can drive a particular expression. 
+The animation system in Unreal Engine 4 (UE4) also supports curve driven animation with a special asset called the **Animation Pose Asset** (or **Pose Asset** for short). 
+
+The **Pose Asset** can additively blend between multiple poses using normal animation curves (as long as they exist within the animation). 
+This allows you to create new full body animations driven purely by animation curve data or it can be driven by FACS (Facial Action Coding System) or Viseme to create facial animation. 
+
+To illustrate how blending works, consider the following two key poses: 
+
+[REGION:imagetable]
+| ![](Pose_B.png)(w:306) |![](Pose_A.png)(w:302)|
+| ----------- |------|
+| **Pose A** | **Pose B** | 
+[/REGION]
+
+We could weight the value of each pose to produce the resulting output pose (shown in the images across the top): 
+
+|  |![](Pose_A.png)(w:200) |![](Pose_B.png)(w:204) | ![](Pose_C.png)(w:154)|
+| ----------- |------|------|------|
+| **Pose A Weight** | 0 | 1 | 0.5 |
+| **Pose B Weight** | 1 | 0 | 0.5 |
+
+Above, if we weighted Pose A to 0 and Pose B to 1 (full influence), the resulting pose would be Pose B (right foot forward). 
+If we weighted Pose A to 1 and Pose B to 0, the resulting pose would be left foot forward.
+However, if we weighted both to 0.5, we would get this new pose somewhere in between where the character appears to be standing still.
+Using these weighted values, we could create a new motion that blends between Pose A and Pose B over a timeline in an **Animation Sequence**. 
+
+For example, below we have created a motion where the character appears to be walking driven through curve data. 
+
+[OBJECT:EmbeddedVideo]
+[PARAMLITERAL:width]
+640
+[/PARAMLITERAL]
+[PARAMLITERAL:height]
+360
+[/PARAMLITERAL]
+[PARAMLITERAL:videoid]
+7PRGvePzB2g
+[/PARAMLITERAL]
+[/OBJECT]
+
+We can also use this process for creating facial animation as seen in the video below. 
+
+[OBJECT:EmbeddedVideo]
+[PARAMLITERAL:width]
+640
+[/PARAMLITERAL]
+[PARAMLITERAL:height]
+360
+[/PARAMLITERAL]
+[PARAMLITERAL:videoid]
+Z_EO1wfHfMs
+[/PARAMLITERAL]
+[/OBJECT]
+
+Above, we have an Animation Sequence that contains curve data for various poses. 
+We then created a **Pose Asset** from the Animation Sequence to gain access to each of the curves/poses (we isolated **R_smile_pose** and **L_smile_pose** to show how we can make the face smile). 
+Next, we created a new Animation Sequence from the default base pose and added two curves to that sequence (the two smile poses). 
+Finally, we added new weight values midway through the animation that allowed us to create the expression of the face smiling. 
+
+[OBJECT:TopicButtonList]
+	[PARAM:icon]
+		![](%ROOT%/start_icon.png)(convert:false)
+	[/PARAM]
+	[PARAM:title]
+		Getting Started
+	[/PARAM]
+	[PARAM:links]
+		[REGION:topics fourth]
+		[INCLUDE:#topics]
+		[/REGION]
+	[/PARAM]
+[/OBJECT]
+
+<!--
+[EXCERPT:Topics]
+[OBJECT:Topic]
+	[PARAM:image]
+		![](topic_create.png)(w:291 h:120)
+	[/PARAM]
+	[PARAM:icon]
+		![](%ROOT%/blueprint_icon.png)(convert:false)
+	[/PARAM]
+	[PARAM:title]
+		%Engine\Animation\AnimHowTo\CreatePoseAsset:title%
+	[/PARAM]
+	[PARAM:description]
+		%Engine\Animation\AnimHowTo\CreatePoseAsset:description%
+	[/PARAM]
+	[PARAM:path]
+		[RELATIVE:Engine\Animation\AnimHowTo\CreatePoseAsset]
+	[/PARAM]
+[/OBJECT]
+[OBJECT:Topic]
+	[PARAM:image]
+		![](topic_howto.png)(w:291 h:120)
+	[/PARAM]
+	[PARAM:icon]
+		![](%ROOT%/blueprint_icon.png)(convert:false)
+	[/PARAM]
+	[PARAM:title]
+		%Engine\Animation\AnimHowTo\CurveDrivenAnimation:title%
+	[/PARAM]
+	[PARAM:description]
+		%Engine\Animation\AnimHowTo\CurveDrivenAnimation:description%
+	[/PARAM]
+	[PARAM:path]
+		[RELATIVE:Engine\Animation\AnimHowTo\CurveDrivenAnimation]
+	[/PARAM]
+[/OBJECT]
+[OBJECT:Topic]
+	[PARAM:image]
+		![](topic_animbp.png)(w:291 h:120)
+	[/PARAM]
+	[PARAM:icon]
+		![](%ROOT%/blueprint_icon.png)(convert:false)
+	[/PARAM]
+	[PARAM:title]
+		%Engine\Animation\AnimPose\AnimBPPoseNodes:title%
+	[/PARAM]
+	[PARAM:description]
+		%Engine\Animation\AnimPose\AnimBPPoseNodes:description%
+	[/PARAM]
+	[PARAM:path]
+		[RELATIVE:Engine\Animation\AnimPose\AnimBPPoseNodes]
+	[/PARAM]
+[/OBJECT]
+[OBJECT:Topic]
+	[PARAM:image]
+		![](topic_additive.png)(w:291 h:120)
+	[/PARAM]
+	[PARAM:icon]
+		![](%ROOT%/blueprint_icon.png)(convert:false)
+	[/PARAM]
+	[PARAM:title]
+		%Engine\Animation\AnimPose\AdditivePA:title%
+	[/PARAM]
+	[PARAM:description]
+		%Engine\Animation\AnimPose\AdditivePA:description%
+	[/PARAM]
+	[PARAM:path]
+		[RELATIVE:Engine\Animation\AnimPose\AdditivePA]
+	[/PARAM]
+[/OBJECT]
+[/EXCERPT:Topics]
+-->
diff --git a/Engine/Documentation/Source/Engine/Animation/Animation.INT.udn b/Engine/Documentation/Source/Engine/Animation/Animation.INT.udn
--- a/Engine/Documentation/Source/Engine/Animation/Animation.INT.udn
+++ b/Engine/Documentation/Source/Engine/Animation/Animation.INT.udn
@@ -1,5 +1,5 @@
 Availability:Public
-Title: Persona: Skeletal Mesh Animation
+Title: Skeletal Mesh Animation System
 Crumbs: %ROOT%, Engine
 Description:System for deforming skeletal meshes based on keyframed animation data and morph targets.
 Navigation:topic
@@ -60,7 +60,7 @@ order:3
 	[/PARAM]
 	[PARAM:links]
 		* [](Engine/Animation/Overview "%Engine/Animation/Overview:description%")
-		* [](Engine/Animation/Persona "%Engine/Animation/Persona:description%")
+		* [](Engine\Animation\Persona\Modes "%Engine\Animation\Persona\Modes:description%")
 		* [](Engine\Animation\AnimHowTo "%Engine\Animation\AnimHowTo:description%")
 	[/PARAM]
 [/OBJECT]
@@ -99,10 +99,10 @@ order:3
 		Features
 	[/PARAM]
 	[PARAM:links]
+		* [](Engine\Animation\AnimPose "%Engine\Animation\AnimPose:description%")
 		* [](Engine/Animation/Persona/BoneManipulation "%Engine/Animation/Persona/BoneManipulation:description%")
 		* [](Engine/Animation/AnimationBlending "%Engine/Animation/AnimationBlending:description%")
 		* [](Engine/Animation/AnimationRetargeting "%Engine/Animation/AnimationRetargeting:description%")
-		* [](Engine/Animation/RetargetingDifferentSkeletons "%Engine/Animation/RetargetingDifferentSkeletons:description%")
 		* [](Engine/Animation/PhysicallyDrivenAnimation "%Engine/Animation/PhysicallyDrivenAnimation:description%")
 		* [](Engine/Animation/RootMotion "%Engine/Animation/RootMotion:description%")
 		* [](Engine/Animation/NonUniformScale "%Engine/Animation/NonUniformScale:description%")
diff --git a/Engine/Documentation/Source/Engine/Animation/AnimationRetargeting/AnimationRetargeting.INT.udn b/Engine/Documentation/Source/Engine/Animation/AnimationRetargeting/AnimationRetargeting.INT.udn
--- a/Engine/Documentation/Source/Engine/Animation/AnimationRetargeting/AnimationRetargeting.INT.udn
+++ b/Engine/Documentation/Source/Engine/Animation/AnimationRetargeting/AnimationRetargeting.INT.udn
@@ -1,33 +1,21 @@
 Availability:Public
-Title: Animation Retargeting (Same Skeleton)
+Title: Animation Retargeting
 Crumbs: %ROOT%, Engine, Engine/Animation
-Description:Retargeting allows you to reuse animations between characters with highly differed proportions.
-version:4.9
+Description:Describes how retargeted animations can be used with multiple Skeletal Meshes, allowing you to share animations. 
+version:4.14
 
 [REGION:fullwidth]
 ![](RetargetingHeader.png)
 [/REGION]
 
-[TOC(start:2 end:2)]
-
-
-
-**Animation Retargeting (Same Skeleton)** is a feature that allows animations to be reused between characters that share the same Skeleton asset, but have vastly different proportions. Through retargeting, you can prevent your animated skeletons from losing their proportions or becoming unnecessarily deformed when using animations from differently shaped characters.
-
-[REGION:warning]
-In order for different Skeletal Meshes to take advantage of retargeting, they must use the same Skeleton asset to ensure Bone Hierarchy, but also the same Joint Rotation in the Bind pose. Please see the [](Engine/Animation/Skeleton) documentation for details.
-[/REGION]
-
-### Animation Retargeting for Different Skeletons
-
-The content on this page covers the concepts behind **Animation Retargeting** and cites examples where it can be used for retargeting animations for characters that use the same Skeleton asset. There may be instances when you want to retarget animations for a character that uses a different Skeleton asset, at which point **Animation Retargeting (Different Skeletons)** will show you how you can accomplish this. 
-
-Please see [Animation Retargeting (Different Skeletons)](Engine\Animation\RetargetingDifferentSkeletons) for more information. 
+**Animation Retargeting** is a feature that allows animations to be reused between characters that use the same [Skeleton](Engine/Animation/Skeleton) asset but may have vastly different proportions. Through retargeting, you can prevent your animated skeletons from losing their proportions or becoming unnecessarily deformed when using animations from differently shaped characters.
+With Animation Retargeting it is also possible to share animations between characters that use **different Skeleton** assets as long as they share a similar Bone Hierarchy and use a shared asset called a **Rig** to pass animation data from one Skeleton to the other. 
 
+[TOC(start:2 end:2)]
 
 ## Why Use Retargeting?
 
-Consider three characters that you would like to share animation: A base character, a short stocky version, and a tall skinny version.
+Suppose you have multiple characters that you would like to share animations between A base character, a short stocky character, and a tall skinny character.
 
 [VAR:Chars]w:300[/VAR]
 
@@ -40,7 +28,7 @@ Consider three characters that you would like to share animation: A base charact
 
 ### Results Before Retargeting
 
-Before retargeting is applied, you can already use animations between any Skeletal Meshes that share the same Skeleton. However, when you have differing proportions as shown above, you can get some unsightly results. Notice how the short character gets unnecessarily stretched and the tall character gets compacted, both so that they can fit onto the Base Character's skeletal proportions.
+Before retargeting is applied, you can already use animations between any Skeletal Meshes that share the same Skeleton asset. However, when you have differing proportions, as shown above, you can get some unsightly results. Notice how the short character gets unnecessarily stretched and the tall character gets compacted, both so that they can fit onto the Base Character's skeletal proportions.
 
 
 [REGION:imagetable]
@@ -51,7 +39,9 @@ Before retargeting is applied, you can already use animations between any Skelet
 
 ### Results After Retargeting
 
-Once retargeting is applied to the characters, the differences in their proportions are taken out of the equation and the animation plays properly on the characters. You can also go to **Show > Non-Retargeted Animation** and see the difference in the original skeleton (shown in beige) and the current skeleton (white). Notice how on the Base Character, the beige non-retargeted bones perfectly match the skeleton.
+Once retargeting is applied to the characters, the differences in their proportions are taken out of the equation and the animation plays properly on each character. 
+You can also go to **Show > Non-Retargeted Animation** from the **Viewport** options and see the difference in the original skeleton (shown in beige) and the current skeleton (white). Notice how on the Base Character, the beige non-retargeted bones perfectly match the skeleton.
+
 [EXCERPT:AfterTable]
 [REGION:imagetable]
 |![](BaseCharacterRunningRetargetedBones.png)(%chars%)|![](ShortStockyRunningRetargetedBones.png)(%chars%)|![](TallSkinnyRunningRetargetedBones.png)(%chars%)|
@@ -60,18 +50,13 @@ Once retargeting is applied to the characters, the differences in their proporti
 [/REGION]
 [/EXCERPT:AfterTable]
 
-## Does Retargeting Affect Performance?
-
-There is no significant difference of performance between using retargeted and non-retargeted animations. The benefit of using animation retargeting is increasing the number of unique characters without having to create an entirely new set of matching animations which could seriously cut down on your animation memory budget.
-
-
 ## How Does Retargeting Work?
 
-Animations are bound to a [Skeleton asset](Engine/Animation/Skeleton). The Skeleton asset is really just a list of bone names and hierarchy data, but it also stores the initial proportions from the original Skeletal Mesh used to define the Skeleton asset. This data is stored as bone translation data. It is important to note that the retargeting system only retargets the bone's translation component. The bone's rotation always comes from the animation data.
+Animations are bound to a [Skeleton](Engine/Animation/Skeleton) asset. The Skeleton asset is really just a list of bone names and hierarchy data, but it also stores the initial proportions from the original Skeletal Mesh used to define the Skeleton asset. This data is stored as bone translation data. It is important to note that the retargeting system only retargets the bone's translation component. The bone's rotation always comes from the animation data.
 
 Because the original Skeletal Mesh is used to define the proportions of the Skeleton asset, any other Skeletal Meshes using that Skeleton asset that have different proportions (such as a mesh that is much shorter than the original one) will require retargeting to work properly. Without it, the Skeletal Meshes of varied proportions would try to use the translation data of the original mesh, leading to the types of errors seen at the beginning of this document.
 
-To solve this problem, the Skeleton List within Persona offers a few settings to change how translation retargeting between bones is handled. There are 3 different settings for bone translation retargeting:
+To solve this problem, the Skeleton Tree within the **Skeleton Editor** offers a few settings to change how translation retargeting between bones is handled. There are 3 different settings for Bone Translation Retargeting:
 
 ![](RetargetingSettings.png)
 
@@ -79,6 +64,20 @@ To solve this problem, the Skeleton List within Persona offers a few settings to
 -	**Skeleton** - Bone Translation comes from the Target Skeleton's bind pose.
 -	**AnimationScaled** - Bone translation comes from the animation data, but is scaled by the Skeleton's proportions. This is the ratio between the bone length of the Target Skeleton (the skeleton the animation is being played on), and the Source Skeleton (the skeleton the animation was authored for).
 
+Also with Animation Retargeting, there is no significant difference in performance between using retargeted and non-retargeted animations. The benefit of using animation retargeting is increasing the number of unique characters without having to create an entirely new set of matching animations which could seriously cut down on your animation memory budget.
+
+### Retargeting Animations for Different Skeletons
+
+The process in which Animation Retargeting is handled for characters that do not share the same Skeleton asset involves assigning a special asset called a **Rig** that processes animation data from one Skeleton to another. 
+The Skeleton assets associated with each character communicate with the shared **Rig** asset to properly pass transform data from a source to its intended target. 
+
+The **Rig** can be assigned inside the **Skeleton Editor** under the [Retarget Manager](#usingtheretargetmanager) and the same Rig will need to be assigned to both Skeleton assets. 
+
+[REGION:note]
+Please see [](Engine\Animation\AnimHowTo\Retargeting) for a step-by-step guide.
+[/REGION]
+
+
 
 ## How Are End Effectors Handled With Retargeting?
 
@@ -88,55 +87,50 @@ One method that can be taken regarding the holding of props is to create a separ
 
 Having a separate chain of bones allows you to switch smoothly between FK and IK easily when you want to (for example, you would want the Hand IK ON when reloading the weapon, and OFF when reaching for a magazine in a pocket). 
 
-The system is very flexible and can be customized to fit your needs. Perhaps you only want the left hand to be IK and the right hand to use its FK position, but IK rotation. Feet are sometimes handled this way, sometimes not. When stepping on a very precise props, you would want to turn on IK; when you are just running around, you would want to use the FK, so you do not end up with a bow legged character (or the opposite). 
+The system is very flexible and can be customized to fit your needs. Perhaps you only want the left hand to be IK and the right hand to use its FK position, but IK rotation. Feet are sometimes handled this way, sometimes not. When stepping on a very precise prop, you would want to turn on IK; when you are just running around, you would want to use the FK, so you do not end up with a bow-legged character (or the opposite). 
 
 
 ## Setting Up Retargeting
 
 <!---Establishing animation retargeting on your --->
 
+As mentioned above under _How Retargeting Works_, the first thing you will want to do is set up the **Bone Translation Retargeting** mode for the Bones within the Skeleton.
 
-
-Retargeting can be set up from the **Skeleton Tree** panel in **Persona**.
- 
 Typically, you will want to use these settings for bipeds:
 
 * Root Bone, IK Bones, weapon bones, and any kind of markers will use the Animation mode.
-* The pelvis will use AnimationScaled so it sits at the right height and is still animated. And any other bone you want translation to be animated and retargeted should also use AnimationScaled.
+* The pelvis will use AnimationScaled so it sits at the right height and is still animated. 
+	* Any other bone you want translation to be animated and retargeted should also use AnimationScaled.
 * All other bones will use Skeleton. They will use the static translation from the target skeleton.
 
-This means that if you want to use retargeting, a fast workflow would be:
+This means that a fast workflow would be:
 
 ![](RecursivelySet.png)
 
-1. **Right-click** on your root bone and recursively set all bones to **Skeleton**.
+1. **Right-click** on your root bone and **Recursively Set Translation Retargeting Skeleton** so all bones are set to Skeleton. 
 1. Find the Pelvis or equivalent bone and set that to **AnimationScaled**.
 1. Find the Root bone, any IK bones, any weapon bones, or other marker-style bones and set them to use **Animation**.
 
-[REGION:tip]
-In the **Show** menu from the viewport, you can check **NonRetargeted Animation** and **Bones** to compare the original with the retargeted animation.
+[REGION:note]
+If you are retargeting animations for characters that share the same Skeleton asset, those animations will now be affected by animation retargeting. 
+If you want to retarget animations for characters that do not share the same Skeleton asset, there a couple extra steps you will need to take to ensure the animations are properly retargeted.
+Please see [](Engine\Animation\AnimHowTo\Retargeting) for more information.
 [/REGION]
 
 
-## Using the Retarget Source Manager
+## Using the Retarget Manager
 
 ![](BasePoseManager.png)
 
-Since retargeting uses the Skeleton asset, and since the Skeleton has its original proportions defined by the Skeletal Mesh for which it was initially created, it follows that one-directional retargeting will work smoothly using the steps described above. However, there will often come times when you will have special animations that are built just for the variant versions. For example, consider a special animation that was created just for the tall version of the character.
+Another tool you can use as part of your Animation Retargeting efforts is the **Retarget Manager** which allows you to: 
 
-If you imported this new "tall-only" animation in, you would still need to use the same Skeleton asset as before, which was created from the Base version of the character. This causes the new animation's proportions to be distorted. The solution is to use the **Retarget Source Manager**.
+* **Manage Retarget Source** assets - useful for when you have different proportional meshes per skeleton, you can use this setting to indicate if a particular animation is from a different source. 
+* **Set up Rig** - which allows you to retarget animations to different skeletons that use the same rig. 
+* **Manage Retarget Base Pose** - which is used when retargeting assets to a different skeleton and allows you to alter the base pose of the target so it lines up with the source's base pose and provides a more accurate retargeted animation. 
 
-The **Retarget Source Manager** allows you to associate an Animation Sequence with the actual Skeletal Mesh for which it was designed. In this way, any problems of retargeting for special animations will be fixed.
-
-### Using the Retarget Source Manager
-
-Think of the **Retarget Source Manager** as a list of Skeletal Meshes that you can use to designate proportions. If you ever need a special animation for a retargeted Skeletal Mesh, you would need to make sure that the mesh was listed in the **Retarget Source Manager**.
-
-1. In Persona's menu bar, choose **Window > Retarget Source Manager**.
-1. **Right-click** in the list area (center) of the **Retarget Source Manager** panel and choose **Add...** from the context menu.
-1. Choose the Skeletal Mesh for which the special animation was created. You should now see that Skeletal Mesh listed in the **Retarget Source Manager**.
-1. In Persona's **Asset Browser**, **Double-click** the Animation Sequence asset for the special animation you created for the Skeletal Mesh.
-1. In the Anim Asset **Details** panel, locate the **Animation** category and locate the **Retarget Source** property. In the dropdown, you should now see the Skeletal Mesh you added to the **Retarget Source Manager**. By choosing this mesh, you are specifying that you want to use the proportions of that mesh to drive this animation.
+[REGION:note]
+Please see the [](Engine\Animation\Persona\BasePoseManager) page for more information.
+[/REGION]
 
 
 <!----
@@ -144,7 +138,7 @@ Think of the **Retarget Source Manager** as a list of Skeletal Meshes that you c
 If you are creating all of your animations for a single Skeletal Mesh and retargeting them to some other variations, then the information presented so far will get you a nice retargeted result. The reason it works is that using just one original Skeletal Mesh as the base is like a one-way connection, from the original to all the variations. But what happens if you need special animations just for the variants?
 
 Consider if we needed a special set of animations for our tall character, such as a different idle animation. An animator could take that Skeletal Mesh, create the necessary animation, then import it using the same Skeleton asset as before, but there is now a problem. Since the original proportions of the Skeleton asset were based off the Base Character's proportions, the animation will be distorted.
-j
-
 
 --->
+
+
diff --git a/Engine/Documentation/Source/Engine/Animation/NodeReference/AnimationNodes.INT.udn b/Engine/Documentation/Source/Engine/Animation/NodeReference/AnimationNodes.INT.udn
--- a/Engine/Documentation/Source/Engine/Animation/NodeReference/AnimationNodes.INT.udn
+++ b/Engine/Documentation/Source/Engine/Animation/NodeReference/AnimationNodes.INT.udn
@@ -42,12 +42,32 @@ space nodes, and conversion nodes.
 	[/PARAM]
 [/OBJECT]
 
+
+## Pins and Properties
+
+Animation nodes can have pins (both input and output) which values can be passed to in the **EventGraph** 
+or **AnimGraph** and properties which can be edited in the **Details** panel. Pins can be either data 
+pins identical to those found on nodes in regular Blueprints which take in data from variables or pure 
+functions, or they can be pose pins which are similar to exec pins in regular Blueprints in that they 
+determine the flow of execution.
+
+Data pins are actually properties like any other and are exposed in the **Details** panel. These can be toggled 
+between being used as a data pin on the node in the graph and being a simple property. When exposed as a data 
+pin, the property is only modifiable in the graph by linking a value to it. When not exposed as a data pin, 
+properties can only be modified through the **Details** panel.
+
+![](AsPin.png)
+
+Above a Blendspace Player node has properties in the **Details** panel that can be set as pins. 
+The X and Y coordinate pins which make up the **Speed** property of our Blendspace as well as Play Rate and Loop are exposed on the node since they are checked inside the **Details** panel. 
+
+
 ## Animation Node Pose Watching
 
 Animation Node **Pose Watching** allows you to see a representation of the pose being generated at any point in the **AnimGraph** dynamically. 
 You can also have multiple Pose Watches active at once allowing you to compare poses at different points to find the exact point in which any errors in your current pose are introduced. 
 
-To enable this, inside your **AnimGraph** on a node that contains a pose data, **Right-click** on the node and select **Toggle Pose Watch**. 
+To enable this, inside your **AnimGraph**, **Right-click** on on a node that contains pose data and select **Toggle Pose Watch**. 
 
 [OBJECT:EmbeddedVideo]
 [PARAMLITERAL:width]
@@ -67,22 +87,3 @@ You can also specify different colors for each of your watched poses by clicking
 
 You can **Remove Pose Watch** from the color picker window or by reselecting **Toggle Pose Watch** on the Animation Node. 
 
-## Pins and Properties
-
-Animation nodes can have pins (both input and output) which values can be passed to in the **EventGraph** 
-or **AnimGraph** and properties which can be edited in the **Details** panel. Pins can be either data 
-pins identical to those found on nodes in regular Blueprints which take in data from variables or pure 
-functions, or they can be pose pins which are similar to exec pins in regular Blueprints in that they 
-determine the flow of execution.
-
-Data pins are actually properties like any other and are exposed in the **Details** panel. These can be toggled 
-between being used as a data pin on the node in the graph and being a simple property. When exposed as a data 
-pin, the property is only modifiable in the graph by linking a value to it. When not exposed as a data pin, 
-properties can only be modified through the **Details** panel.
-
-![](AsPin.png)
-
-Above a Blendspace Player node has properties in the **Details** panel that can be set as pins. 
-The X and Y coordinate pins which make up the **Speed** property of our Blendspace as well as Play Rate and Loop are exposed on the node since they are checked inside the **Details** panel. 
-
-
diff --git a/Engine/Documentation/Source/Engine/Animation/Overview/AnimationOverview.INT.udn b/Engine/Documentation/Source/Engine/Animation/Overview/AnimationOverview.INT.udn
--- a/Engine/Documentation/Source/Engine/Animation/Overview/AnimationOverview.INT.udn
+++ b/Engine/Documentation/Source/Engine/Animation/Overview/AnimationOverview.INT.udn
@@ -1,119 +1,71 @@
 Availability:Public
 Title:Animation System Overview
 Crumbs: %ROOT%, Engine, Engine/Animation
-Description:Overview of the system for deforming SkeletalMeshes based on animation data using Animation Blueprints, Blendspaces, and AnimationSequences.
+Description:Overview of the system for deforming Skeletal Meshes based on animation data using Animation Blueprints, Blendspaces, and Animation Sequences.
+Related: Engine/Content/FBX/Animations
+Related: Engine/Animation/PhysicallyDrivenAnimation
+Related: Engine/Content/Tools/MayaRiggingTool
+version: 4.14
+type:overview
+topic-image:AnimationToolsTopic.png
+tags:Animation
 
-[TOC(start:2 end:2)]
 
+[TOC(start:2 end:2)]
 
 [EXCERPT:Intro]
-**Unreal Engine 4's** (**UE4**) animation system allows for a deep level of control for your characters and Skeletal Meshes. In this document, we will provide a high level overview  of how animation works in UE4, and the base concepts that drive it. <!---In the documents that follow, we'll cover the---> 
-
-
-
-The **Animation** system mixes skeletal-based deformation of meshes with morph-based vertex deformation to 
-allow for complex animation. This system can be used to make basic player movement appear much more realistic 
-by playing and blending between canned AnimationSequences, create customized special moves such as scaling 
-ledges and walls using AnimMontages, apply damage effects or facial expressions through MorphTargets, or 
-directly control the transformations of bones using SkeletalControls.
+The animation system in Unreal Engine 4 (UE4) is comprised of several Animation Tools and Editors which mixes skeletal-based deformation of meshes with morph-based vertex deformation to allow for complex animation. 
+This system can be used to make basic player movement appear much more realistic by playing and blending between canned **Animation Sequences**, create customized special moves such as scaling 
+ledges and walls using **Anim Montages**, apply damage effects or facial expressions through **Morph Targets**, directly control the transformations of bones using **Skeletal Controls**, or create logic based **State Machines** that determine which animation a character should use in a given situation. 
 [/EXCERPT:Intro]
 
 
 The purpose of this document is to offer a high-level overview of UE4's animation system, geared primarily toward users who are new to animating in UE4. However, this should not be considered the comprehensive manual for how to animate skeletal assets in the UE4. Rather, think of this as a primer to familiarize you with the various pieces of the animation system, providing an illustration of how they all fit together and showing how animation in UE4 is more powerful flexible than ever before. 
 
-We will start by identifying the primary terms and concepts of UE4's animation system:
+## Terms and Concepts
 
+We will start by identifying the primary terms and concepts of UE4's animation system (click each header to view more documentation):
 
 (#Skeleton)
 [OBJECT:IconSection]
 	[PARAM:icon]
-		[![](PersonaIcon.png)(w:192)](Engine/Animation/Persona)
+		![](AnimationToolsTopic.png)
 	[/PARAM]
 	[PARAM:title]
-		[Persona](Engine/Animation/Persona)
+		[](Engine/Animation/Persona/Modes)
 	[/PARAM]
 	[PARAM:description]
-	[INCLUDE:Engine/Animation/Persona#Intro]
-		
+	[INCLUDE:Engine/Animation/Persona/Modes#overview]
 	[/PARAM]
 [/OBJECT]
 
-
-
 [OBJECT:IconSection]
 	[PARAM:icon]
-		[![](SkeletonImage.png)(w:192)](Engine/Content\Types\SkeletalMeshes)
+		![](SkeletonTopic.png)
 	[/PARAM]
 	[PARAM:title]
-		[Skeleton](Engine/Content\Types\SkeletalMeshes)
+		[](Engine/Animation/Skeleton)
 	[/PARAM]
 	[PARAM:description]
 		A **Skeleton** is a hierarchy of bone locations and rotations used to deform a Skeletal Mesh. In UE4, Skeletons are abstracted from Skeletal Meshes in their own asset. This means that animations are applied to the Skeleton, rather than the Skeletal Mesh. By using the same Skeleton, multiple Skeletal Meshes can share animations.
-		
-		
 	[/PARAM]
 [/OBJECT]
 
 [OBJECT:IconSection]
 	[PARAM:icon]
-		![](AnimSequenceIcon.png)
+		![](AnimationSeqTopic.png)
 	[/PARAM]
 	[PARAM:title]
 		[](Engine/Animation/Sequences)
 	[/PARAM]
 	[PARAM:description]
 		[INCLUDE:Engine/Animation/Sequences#Intro]
-		
-
-		
-		**Additive Animation**
-		
-		Animations can be used as absolute or additive at any time. The system performs the appropriate calculations to 
-		convert the data to offsets when necessary. This makes the system much more flexible than that of previous versions 
-		of the engine since you do not have to specify ahead of time that an animation is to be considered additive.
-		
-		**Blending Animations**
-		
-		Animations can be blended together to create more complex final poses through the use of [](Engine/Animation/NodeReference/Blend) 
-		and [](Engine/Animation/Blendspaces). For instance, a running animation can be blended with aim poses or 
-		facial animation to create an output animation of the player running and aiming or talking at the same time while 
-		being able to keep the source animations separate. This allows for the reuse of a single source animation in many 
-		circumstances.
-
-		**Animation Poses**
-
-		An Animation Pose is essentially a snapshot of the Skeleton, including the position and rotation of all of its 
-		bones. Think of a pose as what you see if you pause the playback of an animation. For example, if you paused a run 
-		animation mid-stride, you might see something like this:
-
-		![](AnimationPose_Running.png)(w:600)
-
-		It is as if the mesh, via the underlying Skeleton, has been posed. The data that represents the transforms of 
-		all of the bones in the Skeleton is what makes up the actual Animation Pose.
-
-		Some AnimationSequences are designed purely to be used as poses. For instance, you might have a series of short 
-		one- or two-frame animations that simply define a pose. The pose below is from this type of Animation Sequence, 
-		which is designed to be used to cause the mesh to aim its weapon where the player is looking.
-
-		![](AnimationPose_Aiming.png)(w:600)
-
-		The Animation system uses this concept of Animation Poses within the AnimGraph. The data output from an AnimationSequence, 
-		Blendspace, or simple blend node is an Animation Pose, or snapshot of the animation up to that point, to apply 
-		to the SkeletalMesh. 
-
-		**Component vs Local Space**
-
-		[INCLUDE:Engine/Animation/NodeReference/SpaceConversion#intro]
-
-		**Morph Targets**
-
-		[INCLUDE:Engine/Content\FBX\MorphTargets#Intro]
 	[/PARAM]
 [/OBJECT]
 
 [OBJECT:IconSection]
 	[PARAM:icon]
-		![](ParticleEffectNotify.png)(w:192)
+		![](AnimNotifiesTopic.png)
 	[/PARAM]
 	[PARAM:title]
 		[](Engine/Animation/Sequences/Notifies)
@@ -123,24 +75,21 @@ We will start by identifying the primary terms and concepts of UE4's animation s
 	[/PARAM]
 [/OBJECT]
 
-
 [OBJECT:IconSection]
 	[PARAM:icon]
-		[![Vim Blueprint](Engine/Animation/AnimBlueprints/EventGraph/eventgraph_thumb.png)(w:192)](Engine/Animation/AnimBlueprints)
+		![](AnimationBlueprintTopic.png)
 	[/PARAM]
 	[PARAM:title]
 		[](Engine/Animation/AnimBlueprints)
 	[/PARAM]
 	[PARAM:description]
 		[INCLUDE:Engine/Animation/AnimBlueprints#Intro]
-
-		
 	[/PARAM]
 [/OBJECT]
 
 [OBJECT:IconSection]
 	[PARAM:icon]
-		[![Blendspace](Engine/Animation/Blendspaces/edit.png)(w:192)](Engine/Animation/Blendspaces)
+		![](BlendSpaceTopic.png)
 	[/PARAM]
 	[PARAM:title]
 		[](Engine/Animation/Blendspaces)
@@ -152,7 +101,7 @@ We will start by identifying the primary terms and concepts of UE4's animation s
 
 [OBJECT:IconSection]
 	[PARAM:icon]
-		[![](MontageIcon.png)(w:192)](Engine/Animation/AnimMontage)
+		![](MontageTopic.png)
 	[/PARAM]
 	[PARAM:title]
 		[](Engine/Animation/AnimMontage)
@@ -164,7 +113,7 @@ We will start by identifying the primary terms and concepts of UE4's animation s
 
 [OBJECT:IconSection]
 	[PARAM:icon]
-		[![SkelControl](Engine/Animation/NodeReference/SkeletalControls/skelcontrol.png)(w:192)](Engine/Animation/NodeReference/SkeletalControls)
+		![](SkeletonControlTopic.png)
 	[/PARAM]
 	[PARAM:title]
 		[](Engine/Animation/NodeReference/SkeletalControls)
@@ -174,25 +123,25 @@ We will start by identifying the primary terms and concepts of UE4's animation s
 	[/PARAM]
 [/OBJECT]
 
-
-<!--
 [OBJECT:IconSection]
 	[PARAM:icon]
-		[![](StateMachines.png)(w:192)](Engine/Animation/StateMachines)
+		![](StateMachineTopic.png)
 	[/PARAM]
 	[PARAM:title]
 		[](Engine/Animation/StateMachines)
 	[/PARAM]
 	[PARAM:description]
-		[INCLUDE:Engine/Animation/StateMachines#Intro]
+		[INCLUDE:Engine/Animation/StateMachines#overview]
 	[/PARAM]
 [/OBJECT]
--->
 
+[REGION:note]
+Examples of some of the animation features can also be found in the [Content Examples](https://docs.unrealengine.com/latest/INT/Resources/ContentExamples/) project available for free from the **Learn Tab** of the Epic Games Launcher. 
+[/REGION]
 
 ## System Breakdown
 
-The following is a simple breakdown of where certain parts of the animation system are created and edited.
+The following is a breakdown of where certain parts of the animation system are created and edited.
 
 [REGION:raw]
 ![](CreationDiagram.png)
@@ -204,13 +153,6 @@ This diagram illustrates the flow and ownership of data for a typical animated c
 ![](AnimationDiagram.png)
 [/REGION] 
 
-## Animation Content Examples
-
-A practical application of the animation system and its features are on display inside the [Animation Content Examples](Resources\ContentExamples\Animation) page. 
-
-Here you will be able to follow along by looking at bite-sized examples of each feature along with additional details. 
-
-For more information on Content Examples and how to use them, refer to the [Content Examples Overview](Resources\ContentExamples) page. 
 
 
 ## Animation Debug Features
@@ -229,9 +171,7 @@ The `ShowDebug Bones` command will display the information for the playable char
 At this time when using the `viewactor` command, the camera will be moved to and locked on the Actor you are currently viewing. If you need to manipulate the camera while debugging still, a work around for this is to create a [Play in Editor](Engine/UI/LevelEditor/InEditorTesting) session with more than one player from the dropdown menu using one window for debugging and the other for controlling a player character. 
 [/REGION]
 
-An alternative look, matching the bones displayed in Persona, can be enabled by entering the `ShowDebugToggleSubCategory 3DBones` command inside the console. 
-
-
+Instead of displaying white lines for bones, you can enable 3D bones with the `ShowDebugToggleSubCategory 3DBones` command inside the console. 
 
 ### Animation Debug Output
 
@@ -247,4 +187,50 @@ The information displayed is split up into 5 sections, each of which can be togg
 | **Montages** | Lists the montages currently being used by the character. The active montage is highlighted in green. |
 | **Curves** | Lists the curve values (in Name: Value pairs) that have been activated by the playing animation(s). |
 | **Notifies** | Display any notify states that are currently in effect. |
-| **Graph** | Displays the active pose graph. The display starts with the last node (the root node) which represents the final pose and goes on to list all the nodes that go into making that final pose. Nodes are represented in such a way as to keep their hierarchy, allowing the user to see which nodes are connected to what without having to refer to the original Blueprint asset. Active nodes are colored green and (if they have been toggled to display using `ShowDebugToggleSubCategory FullGraph`) inactive nodes are colored grey. |
+| **Graph** | Displays the active pose graph. The display starts with the last node (the root node) which represents the final pose and goes on to list all the nodes that go into making that final pose. Nodes are represented in such a way as to keep their hierarchy, allowing the user to see which nodes are connected to what without having to refer to the original Blueprint asset. Active nodes are colored green and (if they have been toggled to display using `ShowDebugToggleSubCategory FullGraph`) inactive nodes are colored gray. |
+
+<!-- Pulled from the intro section (overkill for overview page).
+		
+		**Additive Animation**
+		
+		Animations can be used as absolute or additive at any time. The system performs the appropriate calculations to 
+		convert the data to offsets when necessary. This makes the system much more flexible than that of previous versions 
+		of the engine since you do not have to specify ahead of time that an animation is to be considered additive.
+		
+		**Blending Animations**
+		
+		Animations can be blended together to create more complex final poses through the use of [](Engine/Animation/NodeReference/Blend) 
+		and [](Engine/Animation/Blendspaces). For instance, a running animation can be blended with aim poses or 
+		facial animation to create an output animation of the player running and aiming or talking at the same time while 
+		being able to keep the source animations separate. This allows for the reuse of a single source animation in many 
+		circumstances.
+
+		**Animation Poses**
+
+		An Animation Pose is essentially a snapshot of the Skeleton, including the position and rotation of all of its 
+		bones. Think of a pose as what you see if you pause the playback of an animation. For example, if you paused a run 
+		animation mid-stride, you might see something like this:
+
+		![](AnimationPose_Running.png)(w:600)
+
+		It is as if the mesh, via the underlying Skeleton, has been posed. The data that represents the transforms of 
+		all of the bones in the Skeleton is what makes up the actual Animation Pose.
+
+		Some AnimationSequences are designed purely to be used as poses. For instance, you might have a series of short 
+		one- or two-frame animations that simply define a pose. The pose below is from this type of Animation Sequence, 
+		which is designed to be used to cause the mesh to aim its weapon where the player is looking.
+
+		![](AnimationPose_Aiming.png)(w:600)
+
+		The Animation system uses this concept of Animation Poses within the AnimGraph. The data output from an AnimationSequence, 
+		Blendspace, or simple blend node is an Animation Pose, or snapshot of the animation up to that point, to apply 
+		to the SkeletalMesh. 
+
+		**Component vs Local Space**
+
+		[INCLUDE:Engine/Animation/NodeReference/SpaceConversion#intro]
+
+		**Morph Targets**
+
+		[INCLUDE:Engine/Content\FBX\MorphTargets#Intro]
+-->
diff --git a/Engine/Documentation/Source/Engine/Animation/Persona/AnimAssetDetails/AnimAssetDetails.INT.udn b/Engine/Documentation/Source/Engine/Animation/Persona/AnimAssetDetails/AnimAssetDetails.INT.udn
--- a/Engine/Documentation/Source/Engine/Animation/Persona/AnimAssetDetails/AnimAssetDetails.INT.udn
+++ b/Engine/Documentation/Source/Engine/Animation/Persona/AnimAssetDetails/AnimAssetDetails.INT.udn
@@ -1,8 +1,12 @@
 Availability:Public
-Title:Persona Anim Asset Details
+Title:Animation Asset Details
 Crumbs: %ROOT%, Engine, Engine/Animation, Engine/Animation/Persona
-Description:User guide for the Anim Asset Details in the Persona animation editor.
-version: 4.9
+Description:User guide for the Animation Asset Details in the Animation Editor.
+version: 4.14
+related: Engine/Animation/Sequences
+related: Engine/Animation/blendspaces
+related: Engine/Animation/animmontage
+related: Engine/Animation/Persona\AssetBrowser
 
 [VAR:Topic]
 [OBJECT:Topic]
@@ -24,27 +28,35 @@ version: 4.9
 [/OBJECT]
 [/VAR]
 
-[EXCERPT:Overview]
+
 ![](AnimAssetDetails.png)(w:520)
 
-The **Anim Asset Details** panel is a context-sensitive property editor that allows you to change settings on various animation assets, such as Anim Sequences, Blend Spaces, Anim Montages, and so on. By default, the **Anim Asset Details** panel is only visible when Persona is in **Animation** mode.
+[EXCERPT:Overview]
+The animation **Asset Details** panel is a context-sensitive property editor that allows you to change settings on animation assets, including but not limited to [Anim Sequences](Engine/Animation/Sequences), [Blend Spaces](Engine/Animation/Blendspaces) or [Anim Montages.](Engine/Animation/AnimMontage) 
+The **Asset Details** panel is visible while in the Animation Editor and consists of settings related to animation assets, however there is an **Asset Details** panel that is available when working inside the [Skeletal Mesh Editor](Engine/Animation/Persona/Modes/Mesh) that provides properties that can be changed to affect the related Skeletal Mesh.  
+[/EXCERPT:Overview]
 
-**To open an animation asset in the Anim Asset Details panel:**
-* **Double-click** it in the Persona **Asset Browser**.
+[REGION:note]
+Refer to each animation asset's documentation page for a breakdown of the options contained within the **Asset Details** panel. 
+[/REGION]
+## Viewing Animation Asset Details
 
-For more information about the Persona **Asset Browser**, see [](Engine/Animation/Persona\AssetBrowser).
+To view the animation **Asset Details** panel, open your animation asset from the **Content Browser**: 
 
-[region:note]
-Opening an animation asset in the **Anim Asset Details** panel also opens it in the Asset Editor. For more information about the Asset Editor, see [](Engine/Animation/Persona\AssetEditor).
-[/region]
+![](ContentBrowser1.png)
 
-[/EXCERPT:Overview]
+The **Asset Details** panel will be located in the left panel of the editor: 
 
-The **Anim Asset Details** panel shows a different set of properties based on the type of animation asset being edited. Outside the fact that this panel is reserved for specific animation assets, it is functionally identical to the standard **Details** panel of the main editor. 
+![](OpenAnimationAsset.png)
 
-For more information about the basics of using this panel, please see the [Editor Details Panel](Engine/UI/LevelEditor\Details) documentation.
+You can also **Double-click** an asset from the [Asset Browser](Engine/Animation/Persona\AssetBrowser) panel to open the **Asset Details** for the asset: 
 
+![](OpenFromAssetBrowser.png)
 
 [REGION:note]
-The **Anim Asset Details** panel only shows content with animation-based assets. Selecting things like Sockets will not show any properties in this panel. For information about such assets, refer to the Persona **Details** panel (in the upper-right corner of the Persona UI by default).
+Docked as a tab alongside the Asset Details tab is the [](Engine/Animation/Persona/SkeletonTree) panel, which shows the skeletal hierarchy of the current Skeleton asset. 
 [/REGION]
+
+
+
+
diff --git a/Engine/Documentation/Source/Engine/Animation/Persona/AnimCurves/AnimCurves.INT.udn b/Engine/Documentation/Source/Engine/Animation/Persona/AnimCurves/AnimCurves.INT.udn
new file mode 100644
--- /dev/null
+++ b/Engine/Documentation/Source/Engine/Animation/Persona/AnimCurves/AnimCurves.INT.udn
@@ -0,0 +1,90 @@
+Availability:Public
+Title:Anim Curves
+Crumbs: %ROOT%, Engine, Engine/Animation, Engine/Animation/Persona
+Description:Describes actions available within the Anim Curves window
+version: 4.14
+Related:Engine/Animation/Sequences/Curves/
+
+[VAR:Topic]
+[OBJECT:Topic]
+	[PARAM:image]
+		![%Engine/Animation/Persona/Toolbar:title%](Engine/Animation/animation_topic.png)
+	[/PARAM]
+	[PARAM:icon]
+		![](%ROOT%/content_icon.png)(convert:false)
+	[/PARAM]
+	[PARAM:title]
+		%Engine/Animation/Persona/Toolbar:title%
+	[/PARAM]
+	[PARAM:description]
+		%Engine/Animation/Persona/Toolbar:description%
+	[/PARAM]
+	[PARAM:path]
+		[RELATIVE:Engine/Animation/Persona/Toolbar]
+	[/PARAM]
+[/OBJECT]
+[/VAR]
+
+
+![](AnimCurvesAlone.png)
+
+The **Anim Curves** panel displays curve values for **Morph Target**, **Attribute**, and **Material** curves for your project. 
+You can delete and rename curves here, as well as preview curve data.
+You can filter the visible curves to show only active curves by clicking on **All Curves** to deselect it, which also enables you to filter by **Morph Target**, **Attribute**, and **Material** curves with checkboxes.
+In addition, you can define [Material Curves](Engine/Animation/Sequences/Curves) which allow you to drive Material Parameters or Morph Curves (if Morph Targets exist within the associated animation). 
+
+To open the **Anim Curves** panel in either the **Skeleton Editor** or **Animation Editor**: 
+
+* Click on the **Window** menu, then select **Anim Curves** from the dropdown menu.
+
+	![](OpenAnimCurve.png)
+
+
+By default, as you play an animation back, the curve values will update along with the playback. In the image below, curves for **L_smile_pose** and **R_smile_pose** are being used to create a smile on the face mesh. 
+The values of these curves are shown in the **Anim Curves** panel at the current point in the animation (where the full smile is achieved by being weighted at **1.0** when reaching the end of the animation). 
+
+![](AllAnimCurve.png)
+
+If you uncheck the **Auto** checkbox for a particular curve, it will no longer pull data from the curve. This allows you to selectively turn off the influence of a curve and preview the result. You can
+also enter a value manually or by clicking and dragging on the **Weight** field. In the below image, the **L_brow_up_pose** curve value has been manually set to 1.0. You can see the effect
+in the viewport, as the left eyebrow of the face is raised all the way.
+
+![](UncheckAuto.png)
+
+## Managing Curves
+
+**To rename a curve:**
+
+1. Right-click on the curve and select **Rename** from the dropdown menu.
+
+	![](RenameCurve.png)
+
+1. Enter your new name.
+
+	![](NewNameCurve.png)
+
+1. Press **Enter** to save your new name.
+
+**To add a curve**:
+
+1. Right-click on any curve and select **Add** from the dropdown menu.
+
+	![](AddCurve.png)
+
+1. Enter your new curve's name.
+
+	![](NewNameAddCurve.png)
+
+1. Press **Enter** to save your new name.
+
+This adds a curve to the current skeleton, but you will need to then add the curve and set its information in the **Animation Editor** to determine what the curve affects.
+
+**To delete a curve:**
+
+1. Right-click on the curve and select **Delete** from the dropdown menu.
+	
+	![](DeleteCurve.png)
+
+	If the curve is currently being used by any assets in your project, a confirmation window will appear warning you that the curves will be removed from those assets. 
+
+1. Click **Yes** to continue with the deletion, or **No** to cancel.
diff --git a/Engine/Documentation/Source/Engine/Animation/Persona/AnimationNotifies/AnimationNotifies.INT.udn b/Engine/Documentation/Source/Engine/Animation/Persona/AnimationNotifies/AnimationNotifies.INT.udn
--- a/Engine/Documentation/Source/Engine/Animation/Persona/AnimationNotifies/AnimationNotifies.INT.udn
+++ b/Engine/Documentation/Source/Engine/Animation/Persona/AnimationNotifies/AnimationNotifies.INT.udn
@@ -1,8 +1,12 @@
 Availability:Public
-Title:Persona Animation Notifies
+Title:Anim Notifies Window
 Crumbs: %ROOT%, Engine, Engine/Animation, Engine/Animation/Persona
-Description:User guide for the Animation Notifies window in the Persona animation editor.
-version: 4.9
+Description:User guide for the Animation Notifies window inside the Skeleton Editor.
+version: 4.14
+related: Engine/Animation/Sequences/Notifies
+related: Engine/Animation/Persona/Modes/Skeleton
+related: Engine/Animation/Skeleton
+related: Engine/Animation/Persona/AssetEditor
 
 [VAR:Topic]
 [OBJECT:Topic]
@@ -24,18 +28,43 @@ version: 4.9
 [/OBJECT]
 [/VAR]
 
-The Personal **Animation Notifies** window displays custom Animation Notifies created in the Asset Editor.
+![Animation Notifies Window](AnimationNotifies.png)
 
-The **Animation Notifies** window is not visible by default. To view it, on Persona's **Window** menu, click **Animation Notifies**.
+The **Anim Notifies** window available inside the **Skeleton Editor** displays all the custom [Animation Notifies](Engine/Animation/Sequences/Notifies) for the currently selected [Skeleton](Engine/Animation/Skeleton) asset. 
+The Anim Notifies window is best thought of as an Animation Notifies manager as you can see all the currently assigned custom notifies and which animation the notify is associated with, as well as delete or rename a notify, however you cannot create new notifies from this window as that is done inside the [](Engine/Animation/Persona/AssetEditor) of the animation asset itself. 
 
-![Animation Notifies Window](AnimationNotifies.png)
+## Viewing the Animation Notifies Window
+
+To view the **Animation Notifies** window, open a **Skeleton** asset from the **Content Browser**: 
+
+![](OpenSkeletonAsset.png)
 
-[REGION:note]
-Selecting a custom Animation Notify in the **Animation Notifies** window exposes its properties in Persona's **Details** panel.
+The **Animation Notifies** window will be located in the lower-right portion of the **Skeleton Editor**: 
+
+[REGION:lightbox]
+[![](AnimationNotifiesWindow.png)(w:800)](AnimationNotifiesWindow.png)
+[/REGION]
+
+[REGION:caption]
+Click image for full view. 
 [/REGION]
 
-To delete a custom Animation Notify, **Right-click** it, and then click **Delete**.
+If the **Animation Notifies** window is not present or it has been closed, you can re-open it from the **Main Toolbar** option: 
+
+![](MainToolbarOption.png)
+
+## Renaming/Deleting Notifies
+
+To **Rename** or **Delete** a notifiy, **Right-click** on the notify in the **Animation Notifies** window: 
+
+![](RightClickMenu.png)
+
+## Skeleton Notify Window
+
+When selecting a notify in the **Animation Notifies** window, the **Skeleton Notify** window will become populated: 
+
+![](SkeletonNotify.png)
+
+This window will display the name of the notify and the animation it currently resides in. 
 
-To rename a custom Animation Notify, **Right-click** it, and then click **Rename**.
 
-For more information about Persona's **Asset Editor**, see [](Engine/Animation/Persona\AssetEditor).
diff --git a/Engine/Documentation/Source/Engine/Animation/Persona/AssetBrowser/PersonaAssetBrowser.INT.udn b/Engine/Documentation/Source/Engine/Animation/Persona/AssetBrowser/PersonaAssetBrowser.INT.udn
--- a/Engine/Documentation/Source/Engine/Animation/Persona/AssetBrowser/PersonaAssetBrowser.INT.udn
+++ b/Engine/Documentation/Source/Engine/Animation/Persona/AssetBrowser/PersonaAssetBrowser.INT.udn
@@ -1,41 +1,52 @@
 Availability:Public
-Title:Persona Asset Browser
+Title:Asset Browser
 Crumbs: %ROOT%, Engine, Engine/Animation, Engine/Animation/Persona
-Description:User guide for the Asset Browser in the Persona animation editor.
-version: 4.9
-
-[VAR:Topic]
-[OBJECT:Topic]
-	[PARAM:image]
-		![%Engine/Animation/Persona/AssetBrowser:title%](Engine/Animation/animation_topic.png)
-	[/PARAM]
-	[PARAM:icon]
-		![](%ROOT%/content_icon.png)(convert:false)
-	[/PARAM]
-	[PARAM:title]
-		%Engine/Animation/Persona/AssetBrowser:title%
-	[/PARAM]
-	[PARAM:description]
-		%Engine/Animation/Persona/AssetBrowser:description%
-	[/PARAM]
-	[PARAM:path]
-		[RELATIVE:Engine/Animation/Persona/AssetBrowser]
-	[/PARAM]
-[/OBJECT]
-[/VAR]
+Description:User guide for the Asset Browser in the Animation and Animation Blueprint Editors
+version:4.14
+related:Engine/Animation/Persona\AnimAssetDetails
+related:Engine/Animation/Persona\AssetEditor
+related:Engine\Animation\AnimBlueprints
+tags:Animation
+type:reference
+topic-image:assetbrowsersmall.png
 
 [EXCERPT:Overview]
 ![](AssetBrowser.png)
 
-The **Asset Browser** panel displays the AnimSequences, BlendSpaces, AnimMontages, and AimOffsetBlendSpaces 
-that are valid for the target Skeleton in a mini **Content Browser**-style interface. 
+The **Asset Browser** window displays the [Animation Sequences](Engine\Animation\Sequences), [BlendSpaces](Engine\Animation\Blendspaces), [AnimMontages](Engine\Animation\AnimMontage) and other animation assets that are useable for the selected [Skeleton](Engine\Animation\Skeleton) asset. 
 [/EXCERPT:Overview]
 
-The **Asset Browser** panel is visible only when Persona is in **Animation** mode.
+**Double-clicking** on an asset in the **Asset Browser** will open the asset inside the **Animation Editor** so that you can preview the animation. 
+
+[REGION:lightbox]
+[![](PreviewingAsset.png)(w:800)](PreviewingAsset.png)
+[/REGION]
+
+[REGION:caption]
+Click image for full view. 
+[/REGION]
+
+This will also populate the [Asset Details](Engine/Animation/Persona\AnimAssetDetails) window with varying options based on the asset type clicked. 
+
+Each asset in the **Asset Browser** is color-coded and uses the same color-coding found during animation asset creation from the **Content Browser**. 
+
+![](AddNewButton.png)
+
+By default, the Advanced Details column view in the **Asset Browser** is hidden. You can un-hide the Advanced Details with the **Reset Columns** option from the **View Options**. 
+
+![](UnHideColumns.png)
+
+This will give you more information about each asset as well as the ability for advanced column sorting. 
+
+You can also highlight an asset in the **Asset Browser** to display a tooltip containing information about that asset as well as a preview. 
+
+![](HighlightToolTip.png)
+
+**Right-clicking** on an asset in the **Asset Browser** will give you a context menu with varying options based on the type of asset selected. 
+
+![](RightClickMenu.png)
+
+
+
 
-| Action | Description |
-| ------ | ----------- |
-| **Left-click Asset** | Selects the asset and enables it in the preview Viewport. |
-| **Double-click Asset** | Selects the asset and displays its properties in the **Anim Asset Details** panel and the **Asset Editor**. For more information, see [](Engine/Animation/Persona\AnimAssetDetails) and [](Engine/Animation/Persona\AssetEditor).|
-| **Drag Asset into Level Editor viewport** | Creates an instance of the asset as a Skeletal Mesh Actor in the Level Editor Viewport. |
 
diff --git a/Engine/Documentation/Source/Engine/Animation/Persona/AssetEditor/AssetEditor.INT.udn b/Engine/Documentation/Source/Engine/Animation/Persona/AssetEditor/AssetEditor.INT.udn
--- a/Engine/Documentation/Source/Engine/Animation/Persona/AssetEditor/AssetEditor.INT.udn
+++ b/Engine/Documentation/Source/Engine/Animation/Persona/AssetEditor/AssetEditor.INT.udn
@@ -1,41 +1,65 @@
 Availability:Public
-Title:Asset Editor
+Title:Animation Asset Editor
 Crumbs: %ROOT%, Engine, Engine/Animation, Engine/Animation/Persona
-Description:User guide for the Asset Editor in the Persona animation editor.
-version: 4.9
-
-[VAR:Topic]
-[OBJECT:Topic]
-	[PARAM:image]
-		![%Engine/Animation/Persona/AssetBrowser:title%](Engine/Animation/animation_topic.png)
-	[/PARAM]
-	[PARAM:icon]
-		![](%ROOT%/content_icon.png)(convert:false)
-	[/PARAM]
-	[PARAM:title]
-		%Engine/Animation/Persona/AssetBrowser:title%
-	[/PARAM]
-	[PARAM:description]
-		%Engine/Animation/Persona/AssetBrowser:description%
-	[/PARAM]
-	[PARAM:path]
-		[RELATIVE:Engine/Animation/Persona/AssetBrowser]
-	[/PARAM]
-[/OBJECT]
-[/VAR]
-
-Persona's **Asset Editor** opens when you **Double-click** an animation asset in the **Asset Browser**. For more information about the **Asset Browser**, see [](Engine/Animation/Persona\AssetBrowser).
-
-You can use the **Asset Editor** to add Notifies, including Particle Systems, Sound Cues, and Animation Notifies, to specific points in time on the tracks of animation assets.
+Description:User guide for the Asset Editor inside the Animation Editor.
+version: 4.14
+related:Engine/Animation/Persona/AssetBrowser
+related:Engine/Animation/Persona/AssetEditor
+related:Engine/Animation/Sequences/Notifies
+related:Engine/Animation/Sequences/Curves
+related:Engine/Animation/AnimHowTo/LayerEditing
+tags:Animation
+type:reference
+topic-image:asseteditorsmall.png
 
 ![](AssetEditor.png)
 
-To add a Notify to an animation track, **Right-click** the track, click **Add Notify**, and do one of the following:
-* To add a new Particle System, click **PlayParticleEffect**.
-* To add a new Sound Cue, click **PlaySound**.
-* To add a new custom Animation Notify, click **New Notify**, and give your new Animation Notify a name.
+The **Asset Editor** inside the Animation Editor provides you with the ability to perform animation focused editing based on your selected animation asset type. 
+By **Double-clicking** on an animation asset in the **Content Browser** (or inside the [](Engine/Animation/Persona/AssetBrowser)), the **Asset Editor** window will change based on the type of asset opened. 
+For example, when you open any [Animation Sequence](Engine/Animation/Sequences) the **Asset Editor** will appear as shown below. 
 
-![](AssetEditor_AddNotify.png)
+[REGION:lightbox]
+[![](AnimSequenceEditor.png)(w:800)](AnimSequenceEditor.png)
+[/REGION]
 
-For more information about Animation Notifies, see [](Engine/Animation/Sequences\Notifies).
+[REGION:caption]
+Click for full view. 
+[/REGION]
 
+Inside the **Asset Editor** for an Animation Sequence, you can: 
+
+* Create [](Engine/Animation/Sequences/Notifies) which are like events that are fired off at defined points within an animation. 
+* Create [](Engine/Animation/Sequences/Curves) which allow you to change the value of a Material Parameter or Morph Target while an animation is playing back. 
+* Create [Animation Layers](Engine/Animation/AnimHowTo/LayerEditing) which allow you to modify any Animation Sequence as part of Additive Layer Tracks. 
+
+When opening a [Blend Space](Engine/Animation/Blendspaces), the **Asset Editor** changes to allow you to modify how the Blend Space functions. 
+
+[REGION:lightbox]
+[![](BlendSpaceEditor.png)(w:800)](BlendSpaceEditor.png)
+[/REGION]
+
+[REGION:caption]
+Click for full view. 
+[/REGION]
+
+Opening an [](Engine/Animation/AnimMontage), the **Asset Editor** provide a way for you to construct how your montage functions. 
+
+[REGION:lightbox]
+[![](AnimMontageEditor.png)(w:800)](AnimMontageEditor.png)
+[/REGION]
+
+[REGION:caption]
+Click for full view. 
+[/REGION]
+
+[REGION:note]
+Refer to each animation asset's documentation pages for more information on working with that asset type. 
+[/REGION]
+
+Regardless of the animation asset type that is opened, across the bottom, you will see the **Playback Toolbar**. 
+
+![](AnimationPlayback.png)
+
+You can playback your animation in forward or reverse modes as well as step through the animation backward or forward frame-by-frame. 
+You can also use the **Record** option to create a new animation object based on your currently selected animation. This will prompt you to define a name and save location for the new asset and will begin recording from the current asset. 
+To stop recording, click the **Record** button again and the new animation asset will be created.
diff --git a/Engine/Documentation/Source/Engine/Animation/Persona/BasePoseManager/BasePoseManager.INT.udn b/Engine/Documentation/Source/Engine/Animation/Persona/BasePoseManager/BasePoseManager.INT.udn
--- a/Engine/Documentation/Source/Engine/Animation/Persona/BasePoseManager/BasePoseManager.INT.udn
+++ b/Engine/Documentation/Source/Engine/Animation/Persona/BasePoseManager/BasePoseManager.INT.udn
@@ -1,8 +1,10 @@
 Availability:Public
-Title:Persona Retarget Source Manager
+Title:Retarget Manager
 Crumbs: %ROOT%, Engine, Engine/Animation, Engine/Animation/Persona
-Description:User guide for the Retarget Source Manager in the Persona animation editor.
-version: 4.9
+Description:Breaks down the Retarget Manager options within the Skeleton Editor.-
+version: 4.14
+related: Engine\Animation\AnimHowTo\Retargeting
+related: Engine/Animation/Skeleton
 
 [VAR:Topic]
 [OBJECT:Topic]
@@ -24,15 +26,93 @@ version: 4.9
 [/OBJECT]
 [/VAR]
 
-Persona's **Retarget Source Manager** enables you to create a list of potential base poses for a specific skeleton, which aids in retargeting skeletons for use with multiple Skeletal Meshes.
+The **Retarget Manager** within the **Skeleton Editor** allows you to manage your retarget sources, set up rigs and define retarged base poses for use with [Animation Retargeting](Engine\Animation\AnimHowTo\Retargeting) for characters that do not share the same [Skeleton](Engine/Animation/Skeleton) asset. You can open the Retarget Manager from the Toolbar of the Skeleton Editor. 
 
-![](BasePoseManager.png)
+[TOC(start:2 end:2)]
 
-**To add a new retarget source for reference:**
 
-1. **Left-click** the ** Add New Retarget Source** button.
-1. Select the Skeletal Mesh whose base pose you want to use from the list.
+## Manage Retarget Source
 
-[region:note]
-After you have added a base pose in the **Retarget Source Manager**, you can choose it for a specific animation asset in the **Anim Asset Details** panel, under **Animation**, in the **Retarget Source** drop-down list.
-[/region]
+Since retargeting uses the Skeleton asset, and since the Skeleton has its original proportions defined by the Skeletal Mesh for which it was initially created, it follows that one-directional retargeting will work smoothly in most cases. However, there will often come times when you will have special animations that are built just for the variant versions. 
+For example, suppose you have multiple characters that share the same Skeleton asset (a base character, a short character and a tall character) and you have a special animation that was created just for the tall version of the character.
+
+If you imported this new "tall-only" animation in, you would still need to use the same Skeleton asset as before, which was created from the base version of the character. This will cause the new animation's proportions to be distorted. 
+The solution is to use the **Manage Retarget Source** option within the **Retarget Manager** which allows you to associate an Animation Sequence with the actual Skeletal Mesh for which it was designed. In this way, any problems of retargeting for special animations will be fixed.
+
+Think of the **Mange Retarget Source** as a list of Skeletal Meshes that you can use to designate proportions. If you ever need a special animation for a retargeted Skeletal Mesh, you would need to make sure that the mesh was listed as a retarget source, steps for which are outlined below. 
+
+1. Inside the **Skeleton Editor**, from the **Toolbar** click the **Retarget Source Manager** button. 
+
+	![](RetargetManager.png)
+
+1. Click the **Add New Retarget Source** button. 
+
+	![](AddSource.png)
+
+1. Choose the Skeletal Mesh for which the special animation was created.
+
+	![](SelectSkeletalMesh.png)
+
+	 You should now see that Skeletal Mesh listed in the **Retarget Manager**.
+
+	![](Created.png)
+
+1. Open the special case Animation Sequence intented for your specific Skeletal Mesh. 
+
+1. In the **Asset Details** panel, locate the **Animation** category and locate the **Retarget Source** property and choose your Skeletal Mesh from the dropdown menu.
+
+	![](DropDownSelection.png)
+
+	 By choosing this mesh, you are specifying that you want to use the proportions of that mesh to drive this animation.
+
+
+## Set up Rig
+
+![](SetUpRig.png)
+
+The middle section of the Retarget Manager allows you assign a **Rig** to the Skeleton which can be used to pass animation data to a different Skeleton that uses the same Rig. 
+This process is needed in order to perform any animation retargeting for characters that use different Skeleton assets. 
+You can select the Rig to use from the **Select Rig** drop-down option where a **Humanoid** option is available which you will want to select for most characters (future engine updates may include the ability to create Rigs). 
+
+![](HumanoidRig.png)
+
+Once you assign the Humanoid Rig, you will need to assign Bones from the Skeleton that correspond to the same (or similar) location as the Node on the Rig. 
+You can the drop-down menu for Node and manually assign the corresponding Bone from your Skeleton, or you can use the **Auto Mapping** feature located at the top of the menu. 
+This will look through your Skeleton and attempt to find the best matching Bone for each Node on the Rig. 
+
+The **Clear Mapping** button will wipe all the currently assigned Bones from their corresponding Node assignment while the **Show Advanced** button will allow you to assign additional Node/Bone correlations for things like finger digits, IK Bones or Twist Bones. 
+Once you set up the Rig for your source Skeleton (or the Skeleton asset that drives the animation that you want to retarget to another character), you will need to go into the Skeleton of the target Skeleton and assign the same Rig and define the Bones in the new Skeleton that best match up with the Nodes on the Rig. 
+
+[REGION:note]
+Please see [](Engine\Animation\AnimHowTo\Retargeting) for a step-by-step guide on setting up a Rig to retaget animations between characters that use different Skeletons.
+[/REGION]
+
+
+
+
+## Manage Retarget Base Pose
+
+There may be instances where you want to retarget animations to a Skeleton that is not quite using the same base pose as your Source Skeleton. 
+
+Take for example the Skeleton in the image below on the right when attempting to retarget an animation:  
+
+![](ManageBasePose1.png)(w:840)
+
+Our Source Skeleton (left) is in a A-Pose while our Target Skeleton is using an T-Pose. If we were to retarget animations as is now, we may encounter issues: 
+
+![](TargetAnimation.png)(w:540)![](BadBasePose.png)(w:550)
+
+Above we have our target animation on the left of a character holding a shotgun, however when we retarget the animation to our new character on the right (because they are using different base poses) the arm positioning is not correct. 
+We can fix this by Retargeting the Base Pose inside the **Retarget Manager** which allows us to define a retargeted Base Pose to use for animation retargeting purposes rather than using the characters normal Base pose. 
+
+We can select the Bones in our character and rotate them (in this case the Left and Right Shoulders) so that our character is in an A-Pose, then we click **Save Pose**: 
+
+![](RetargetedBasePose.png)
+
+Now when we go to retarget that animation, we will see the updated retargeted base pose: 
+
+![](NewBasePose.png)(w:840)
+
+And when we retarget animations for Skeletons that have more similar base poses, we get better looking results from the retargted animation: 
+
+![](HoldingShotGun.png)
\ No newline at end of file
diff --git a/Engine/Documentation/Source/Engine/Animation/Persona/BlueprintDefaults/BlueprintDefaults.INT.udn b/Engine/Documentation/Source/Engine/Animation/Persona/BlueprintDefaults/BlueprintDefaults.INT.udn
--- a/Engine/Documentation/Source/Engine/Animation/Persona/BlueprintDefaults/BlueprintDefaults.INT.udn
+++ b/Engine/Documentation/Source/Engine/Animation/Persona/BlueprintDefaults/BlueprintDefaults.INT.udn
@@ -1,61 +1,69 @@
 Availability:Public
-Title:Persona Anim Preview Editor
+Title:Anim Preview Editor
 Crumbs: %ROOT%, Engine, Engine/Animation, Engine/Animation/Persona
-Description:User guide for the Anim Preview Editor panel in the Persona animation editor.
-version: 4.9
+Description:User guide for the Animation Preview Editor panel inside the Animation Blueprint Editor. 
+version:4.14
+related:Engine\Animation\AnimBlueprints
+related:Engine/Blueprints
+related:Engine/Blueprints\Editor\Modes\Defaults
+tags:Animation
+tags:Blueprints
+tags:Basics
+type:reference
+topic-image:BlueprintDefaultSmall.png
 
-[VAR:Topic]
-[OBJECT:Topic]
-	[PARAM:image]
-		![%Engine/Animation/Persona/AssetBrowser:title%](Engine/Animation/animation_topic.png)
-	[/PARAM]
-	[PARAM:icon]
-		![](%ROOT%/content_icon.png)(convert:false)
-	[/PARAM]
-	[PARAM:title]
-		%Engine/Animation/Persona/BlueprintDefaults:title%
-	[/PARAM]
-	[PARAM:description]
-		%Engine/Animation/Persona/BlueprintDefaults:description%
-	[/PARAM]
-	[PARAM:path]
-		[RELATIVE:Engine/Animation/Persona/BlueprintDefaults]
-	[/PARAM]
-[/OBJECT]
-[/VAR]
+![](AnimPreviewEditorWindow.png)
 
-[TOC (start:2)]
+The **Anim Preview Editor** found inside of [](Engine\Animation\AnimBlueprints) displays all the [Variables](Engine/Blueprints/UserGuide/Variables) and Variable categories associated with an Animation Blueprint. 
+When you open any Animation Blueprint, you can find the Anim Preview Editor in the lower-right corner of the **Animation Blueprint Editor**. 
+The Anim Preview Editor not only displays the variables associated with the Animation Blueprint, it allows you to adjust the values to see how it will affect any logic which may lead to changes in a character's animation state behavior. 
 
+[REGION:lightbox]
+[![](AnimationBlueprint.png)(w:800)](AnimationBlueprint.png)
+[/REGION]
 
+[REGION:caption]
+Click image for full view. 
+[/REGION]
 
+[REGION:note]
+Also included by default in this location on a separate docked tab is the [](Engine\Animation\Persona\AssetBrowser) window. 
+[/REGION]
 
-The **Anim Preview Editor** panel in Persona lists the variables and variable categories of the **Animation Blueprint**.
+Variables created inside an Animation Blueprint can be assigned to categories for organizational purposes (a **Default** category is used if no category is selected). 
 
-![](BlueprintDefaults.png)
+![](NewVariable.png)
 
-When you create or modify a Blueprint variable, it is assigned to the **Default** category by default. To organize your assets, you can assign it to a different category by typing in a new category name in the **Category** box.
+You can type in a name for a category for the variable to be assigned to or select a previously created category from the drop-down menu. 
 
-![](AB_Details.png)
+## Anim Preview Editor Modes
 
-For the most part, an Animation Blueprint's **Anim Preview Editor** panel can be modified in the same way as the Blueprint Editor Defaults mode.  For more information about the Blueprint Editor Defaults mode, see [](Engine/Blueprints\Editor\Modes\Defaults).
+The Anim Preview Editor is split into two modes: **Edit Preview** and **Edit Defaults**
 
-## Previewing
+### Edit Preview
 
-You can preview your various animation assets by manually modifying the values of the variables that control them in the **Anim Preview Editor** panel. As you modify the values of the variables, the Viewport displays the appropriate animation assets based on the new values. 
+The default option for the Anim Preview Editor is the Edit Preview mode which allows you alter variables that affect the preview Skeletal Mesh only. 
 
-![](BlueprintDefaults_Preview.png)
+![](EditPreviewOption.png)
 
-## Edit Modes
+Above, the option **IsInAir?** has been set to **True** and the preview mesh (based on the Animation Blueprint's logic) enters a falling state. 
+Whenever you make changes inside the Anim Preview Editor and you are on the Edit Preview mode, those changes only affect the preview mesh and the changes you make will not be saved unless you click the **Apply** button which will take your values and assign them to the new default values. 
+
+You can also define the [Root Motion Mode](Engine/Animation/RootMotion) for the Animation Blueprint to use as well as define [Optimization](Engine/Animation/Optimization/FastPath) options. 
+
+### Edit Defaults
+
+The Edit Defaults option is very similar to the Edit Preview, however here you are editing the actual default values for your properties. 
+
+![](EditDefaults.png)
+
+When you make changes to the values in this mode, you may notice the preview out of date warning in the viewport (which will go away once you click the **Compile** button). 
+
+The **Public View** option is a toggle that can be used to filter the variables to only show ones that have been marked as **Public** or all variables. 
+
+![](PublicViewSetting.png)
+
+Above, the variables **IsInAir?** and **Reloading** have been defined as **Public** so only those are the only ones that appear in the Anim Preview Editor when Public View is enabled. 
 
-Along the top of the **Anim Preview Editor** panel are two modes that can be switched between and affect the permanency of changes made to variable values in the panel.
 
-| Edit Mode | Description | |
-| --- | ----------- | --- |
-| **Edit Preview** | Changes made to variables in the **Anim Preview Editor** will only reflect in the Viewport preview and will not be saved. Changes will revert to default values when the Blueprint is recompiled. | ![](EditPreview.png) |
-| **Edit Defaults** | Changes made to variables in the **Anim Preview Editor** affect both the preview and the default values of the variables and will be saved.<br>[INCLUDE:#publicview]| ![](EditDefaults.png) |
 
-<!--
-[EXCERPT:publicview]
-* **Public View** - An option to toggle whether or not to display only variables marked as Public. After toggling Public View on or off, the Blueprint must be recompiled before changes can be seen. Only available in **Edit Defaults** mode.
-[/EXCERPT:publicview]
--->
