Copying //UE4/Dev-Documentation to Samples-Main (//UE4/Samples-Main)

diff --git a/Engine/Documentation/Source/Engine/Animation/Persona/BoneManipulation/BoneManipulationMode.INT.udn b/Engine/Documentation/Source/Engine/Animation/Persona/BoneManipulation/BoneManipulationMode.INT.udn
--- a/Engine/Documentation/Source/Engine/Animation/Persona/BoneManipulation/BoneManipulationMode.INT.udn
+++ b/Engine/Documentation/Source/Engine/Animation/Persona/BoneManipulation/BoneManipulationMode.INT.udn
@@ -1,44 +1,64 @@
 Availability:Public
 Title:Bone Manipulation Mode
 Crumbs: %ROOT%, Engine, Engine/Animation, Engine/Animation/Persona
-Description:Mode in Persona that allows you to translate or rotate bones while previewing animation assets.
+Description:Overview of the Bone Manipulation Mode that allows you to translate or rotate bones while previewing animation assets.
 Related: Engine/Content/FBX/Animations
 Related: Engine/Content/Types/SkeletalMeshes
 Related: Engine/Animation/PhysicallyDrivenAnimation
 Related: Engine/Content/Tools/MayaRiggingTool
-version: 4.9
-
-**Bone Manipulation Mode** is a mode in Persona in which technical artists can temporarily translate or rotate bones in order to look for skinning issues in a Skeletal Mesh while previewing animations.
+version: 4.14
+type:reference
+topic-image:bonemanipulationmode.png
 
 ![Bone Manipulation Mode](Bones_TwoViews.png)
 
-**To manipulate bones:**
-1. Open a Skeletal Mesh asset in Persona.
-1. Select an animation asset from the list in the **Asset Browser** to make Bone Manipulation Mode available. You can do this in **Skeleton**, **Mesh**, or **Animation** mode. (You may want to pause the animation before manipulating the bone.)
+Inside the [Animation Editor](Engine/Animation/Persona/Modes/Animation), the **Bone Manipulation Mode** allows technical artists to temporarily translate or rotate bones in order to look for skinning issues in a Skeletal Mesh while previewing animations. 
+Another useful aspect to the Bone Manipulation Mode is that you can translate or rotate a bone (for example, tilting the head of a character up or down) then use the [Timeline's Record](Engine/Animation/Sequences/Editor) option to create a new animation based on the altered bone positioning.  
+
+## Manipulating Bones
+
+To maniuplate bones in a Skeletal Mesh during animation playback: 
+
+1. In the [](Engine/Content/Browser) (or [Asset Browser](Engine/Animation/Persona/AssetBrowser)) open any animation asset. 
+
+1. In the **Animation Editor**, click the **Pause** button on the timeline to pause the animation. 
+
+	 ![](PauseAnimation.png)
 
-	![](Bones_SelectAnimation.png)
-1. Select the bone you want to manipulate in the **Skeleton Tree** panel.
+1. Click the **Skeleton Tree** tab, then select your desired bone to manipulate from the **Tree View** window. 
 
-	![](Bones_SelectBone.png)
-1. At the top of the **Viewport** panel, select the translation or rotation icon to activate its widget.
+	![](BoneSelection.png)
 
-	![Translation Widgets](Bones_Widgets.png)
-1. Using the widget that appears, translate or rotate the bone.
-1. Play the animation and look for any skinning issues.
+1. Use the [Transform Widgets](Engine/Actors/Transform\#InteractiveTransformation) to Translate (**W** key), Rotate (**E** key) or Scale (**R** key) the selected bone. 
 
-For more information about using translation and rotation widgets, see [Interactive Transformation](Engine/Actors/Transform\#InteractiveTransformation).
+	![](ModifyBone.png)
+
+	You can also click one of the transform options from the toolbar across the top of the viewport. 
+
+1. Play the animation with the modified bone transformations. 
+
+	[OBJECT:EmbeddedVideo]
+	[PARAMLITERAL:width]
+	640
+	[/PARAMLITERAL]
+	[PARAMLITERAL:height]
+	360
+	[/PARAMLITERAL]
+	[PARAMLITERAL:videoid]
+	V29k8E0K2jk
+	[/PARAMLITERAL]
+	[/OBJECT]
+
+	In the video above, we adjusted two different bones to see what the animation would look like if the bones were in different positions. 
 
 [REGION:note]
-After you have manipulated the bone, the manipulation persists when you select different animation assets, so you can preview the change with different animations.
+After you have manipulated a bone, the manipulation persists when you select different animation assets, so you can preview the change with different animations. 
+Any modified bones will return to their default translation, rotation or scale when you close and re-open the **Animation Editor**. 
 [/REGION]
 
-**To return a bone to its default translation and rotation values:**
 
-* Select a different bone in the **Skeleton Tree**.
 
-[REGION:tip]
-For an example of **Bone Scaling** applied to a character, you can reference the [Animation Content Examples](Resources\ContentExamples\Animation) page under section 1.7.
-[/REGION]
+
 
 
 
diff --git a/Engine/Documentation/Source/Engine/Animation/Persona/Debug/Debug.INT.udn b/Engine/Documentation/Source/Engine/Animation/Persona/Debug/Debug.INT.udn
--- a/Engine/Documentation/Source/Engine/Animation/Persona/Debug/Debug.INT.udn
+++ b/Engine/Documentation/Source/Engine/Animation/Persona/Debug/Debug.INT.udn
@@ -1,29 +1,40 @@
 Availability:Public
-Title:Persona Debug
+Title:Animation Blueprint Debugging
 Crumbs: %ROOT%, Engine, Engine/Animation, Engine/Animation/Persona
-Description:User guide for the Asset Browser in the Persona animation editor.
-Version:4.9
-
-[VAR:Topic]
-[OBJECT:Topic]
-	[PARAM:image]
-		![%Engine/Animation/Persona/AssetBrowser:title%](Engine/Animation/animation_topic.png)
-	[/PARAM]
-	[PARAM:icon]
-		![](%ROOT%/content_icon.png)(convert:false)
-	[/PARAM]
-	[PARAM:title]
-		%Engine/Animation/Persona/AssetBrowser:title%
-	[/PARAM]
-	[PARAM:description]
-		%Engine/Animation/Persona/AssetBrowser:description%
-	[/PARAM]
-	[PARAM:path]
-		[RELATIVE:Engine/Animation/Persona/AssetBrowser]
-	[/PARAM]
-[/OBJECT]
-[/VAR]
-
-The **Debug** panel shows information designated as important by the designer in the form of breakpoints and watch values as well as a trace stack of all nodes belonging to the Blueprint that have been executed. This window also shows controls for controlling execution of the game when using breakpoints.
-
-The Debug panel in persona is the same panel as found in the Blueprint Editor. To learn more about using this panel, see the [](Engine/Blueprints/UserGuide/Debugging) page.
\ No newline at end of file
+Description:Overview of how to enable the Animation Blueprint Editor's Debug window for debugging your Animation Blueprints. 
+Related: Engine/Blueprints/UserGuide/Debugging
+Related: Engine/Animation/AnimBlueprints
+Related: Engine/Blueprints/Editor/UIComponents/Debug
+version: 4.14
+type:reference
+topic-image:AnimationDebugSmall.png
+tags:Animation
+tags:Blueprints
+tags:Debugging
+
+![](AnimationDebugWindow.png)
+
+After constructing your [](Engine/Animation/AnimBlueprints) and playtesting for functionality, you may encounter errors and want to debug what is causing those errors so that you can fix them. 
+Animation Blueprints share a lot of functionality with the [](Engine/Blueprints) system, including some very useful debugging tools that allow you to watch the value of variables, place **Breakpoints** on nodes in your graph to pause the execution flow, as well as viewing the flow of execution through a trace stack of all nodes belonging to the Blueprint that has been executed. 
+
+By default, the **Debug** window inside the Animation Blueprint Editor is hidden, however, you can enable it from the **Window** option of the toolbar. 
+
+![](WindowOption.png)
+
+This will open the Debug window, similar to below. 
+
+![](DefaultDebugWindow.png)
+
+The content in this window may vary based on what is inside your Animation Blueprint. 
+
+You can use the same [](Engine/Blueprints/UserGuide/Debugging) methods inside your Animation Blueprint to add Breakpoints or add Variables to the Watch List. 
+
+![](ExampleDebugging.png)
+
+Above a Breakpoint has been added to the **AnimNotify_Jump** event, which will cause the script to pause when this node is called to execute. 
+We also watching the value of the variable **Enable Jump**, with these debugging options we could diagnose any problems that may arise when attempting to execute the animation notify attached to our jump. 
+Perhaps _Jump_ is not being set correctly or the **Cast** node we are using failed, we would see that inside the debugger and our Animation Blueprint during a play session. 
+
+From the toolbar, under the **Debug** option you can enable/disable all Breakpoints, delete all Breakpoints or delete all watches (stop watching all variables). 
+
+![](DebugMenuOption.png)
diff --git a/Engine/Documentation/Source/Engine/Animation/Persona/MeshDetails/MeshDetails.INT.udn b/Engine/Documentation/Source/Engine/Animation/Persona/MeshDetails/MeshDetails.INT.udn
--- a/Engine/Documentation/Source/Engine/Animation/Persona/MeshDetails/MeshDetails.INT.udn
+++ b/Engine/Documentation/Source/Engine/Animation/Persona/MeshDetails/MeshDetails.INT.udn
@@ -1,37 +1,161 @@
 Availability:Public
-Title:Persona Mesh Details
+Title:Skeletal Mesh Asset Details
 Crumbs: %ROOT%, Engine, Engine/Animation, Engine/Animation/Persona
 Description:User guide for the Mesh Details panel in the Persona animation editor.
-version: 4.9
+Related: Engine/Content/Types/SkeletalMeshes
+Related: Engine/Rendering/Materials
+Related: Engine/Content/Types/StaticMeshes/HowTo/LODCollision
+Related: Resources/ContentExamples/Cloth
+Related: Engine/Content/FBX/SkeletalMeshes
+version: 4.14
+type:reference
+topic-image:SkeletalMeshAssetDetails.png
+tags:Animation
+tags:Skeletal Mesh
+tags:FBX Importing
 
-[VAR:Topic]
-[OBJECT:Topic]
-	[PARAM:image]
-		![%Engine/Animation/Persona/MeshDetails:title%](Engine/Animation/animation_topic.png)
+![](MeshDetails.png)(w:940)
+
+[EXCERPT:Overview]
+The **Asset Details** panel inside the **Skeletal Mesh Editor** shows specific properties pertaining to the Skeletal Mesh, such as the Materials applied to the surface, as well as skinning and LOD options. 
+You can also assign and adjust APEX clothing or define and associate a **Physics Asset** for the Skeletal Mesh which is needed for creating ragdoll type motions. 
+[/EXCERPT:Overview]
+
+This page highlights each section within the Asset Details panel with links to additional documentation, please see each section below: 
+
+[TOC (start:2)]
+
+## Materials 
+
+![](MaterialsTable.png)
+
+Any materials applied to meshes modeled in external applications will be exported along with the mesh. 
+When importing that mesh into Unreal Engine 4, you can automatically import those materials as part of the [FBX Import Process.](Engine/Content/FBX/Materials)
+Those materials will appear inside the Skeletal Mesh Editor under this **Materials** section and can be overridden by any other Materials that exist within your project. 
+You can add **Material Slots** which can be used to override a LODs section (not the base LOD). 
+
+## LODs
+
+![](LODSection.png)
+
+When importing a Skeletal Mesh, you can also import any [Level of Details](Engine/Content/ImportingContent/ImportingSkeletalLODs/) (**LODs**) that you may have created and associated with your mesh. 
+You can **Highlight** portions of the mesh or **Isolate** portions to only show certain sections by clicking your desired option's checkbox. You can also assign the Material Elements from the Material section to **Material Slots** or override sections on a per LOD basis.
+
+[OBJECT:ComparisonSlider]
+	[PARAM:before]
+	![Highlight On](HighlightOption.png)(w:1000)
+	[/PARAM]
+	[PARAM:after]
+	![Isolate On](IsolateOption.png)(w:1000)
 	[/PARAM]
-	[PARAM:icon]
-		![](%ROOT%/content_icon.png)(convert:false)
+[/OBJECT]
+
+Above we toggle between highlighting the helmet of our character to displaying only the helmet. 
+
+The LODs section also allows you to define when to use your LODs through the **Screen Size** option (how close or far away should the LOD become visible) as well as defining the amount of **LODHysteresis** to avoid "flickering" when on an LOD boundary. 
+Inside the **Viewport**, you can click the **Preview LODs** button (below) to switch between the available LODs (it may also help to enable the **Display Info** for statistical information). 
+
+![](previewLODs.png)
+
+Below using the LOD Previewer, we can see the visual difference between the LODs. 
+
+[OBJECT:ComparisonSlider]
+	[PARAM:before]
+	![LOD 0](LOD0.png)
 	[/PARAM]
-	[PARAM:title]
-		%Engine/Animation/Persona/MeshDetails:title%
+	[PARAM:after]
+	![LOD 2](LOD2.png)
 	[/PARAM]
-	[PARAM:description]
-		%Engine/Animation/Persona/MeshDetails:description%
+[/OBJECT]
+
+With our stats being displayed we can also see statistical information about each LOD. Above, the amount of triangles in LOD 2 is almost half the amount of triangles in LOD 0. 
+
+You can also turn on/off the ability for sections within an LOD to **Cast Shadows** or not. 
+
+[OBJECT:ComparisonSlider]
+	[PARAM:before]
+	![Cast Shadows On](castshadows1.png)(w:1000)
 	[/PARAM]
-	[PARAM:path]
-		[RELATIVE:Engine/Animation/Persona/MeshDetails]
+	[PARAM:after]
+	![Cast Shadows Off](castshadows2.png)(w:1000)
 	[/PARAM]
 [/OBJECT]
-[/VAR]
 
-[EXCERPT:Overview]
-![](MeshDetails.png)
+You can turn on/off casting shadows on a per element, per LOD basis (some areas you may not need to cast shadows). 
 
+## Clothing
 
-The **Mesh Details** panel shows specific properties pertaining to the mesh of a Skeletal Mesh, such as the Materials applied to the surface, as well as skinning and LOD options. This panel is only visible when Persona is in **Mesh** mode. 
-[/EXCERPT:Overview]
+The **Clothing** section allows you to add an [](Engine/Physics/Apex) clothing file to associate with the mesh. 
+
+Click the **Add APEX clothing file...** button to begin importing an .apb or .apx file which follows the same pipeline as importing an FBX file. 
+
+![](AddApex.png)
+
+Once you've imported your Apex file, you can adjust the [Physics Properties](Engine/Physics/PhAT/Reference) to achieve your desired results. 
+
+![](OwenApexCloth.png)
+
+Above our character has two Apex files associated with it, one for the character's tie and the other for the character's coat. 
+
+## Mesh
+
+![](MeshSection.png)
+
+In the **Mesh** section, you can see what [Skeleton](Engine/Animation/Skeleton) asset the mesh is using and browse to it in the **Content Browser**. 
+There are also options for setting **Positive Bounds Extension** and **Negative Bounds Extension** values which you can use to offset the bounds values of the mesh when you import it. 
+The other option, **Use Full Precision UVs** can be turned on to use 32 bit UVs (or left off to use 16 bit UVs to save memory). 
+
+## Physics
+
+![](PhysicsSection.png)
+
+The **Physics** section allows you to **Enable Per Poly Collision** for the mesh which uses skinned data for collision data. Per Poly Collision cannot be used for simulation, in most cases, you are better off using a **Physics Asset** instead. 
+There is also the ability to assign a [Physics Asset](Engine/Physics/PhAT) to the Skeletal Mesh which is used for processing physics and collision information. This is used for per-bone hit detection, accurate bounding box calculation, and ragdoll physics. 
+
+## Lighting
+
+![](Lighting.png)
+
+The **Shadow Physics Asset** option in the Lighting section allows you to specify a Physics Asset whose shapes will be used for shadowing when components have their **Cast Character Capsule Direct Shadow** or **Cast Character Capsule Indirect Shadow** options enabled. 
+Only spheres and sphyl shapes in the Physics Asset are supported and the more shapes used, the higher the cost the capsule shadows will be. 
+
+In-game, indirect capsule shadows are a good way to ground characters to the environment: 
+
+![](IndirectShadows.png)(w:840)
+
+Direct capsule shadows can be used to achieve extremely soft character shadows in an efficient baked lighting environment. 
+
+![](DirectShadows.png)(w:840)
+
+The light's **Source Radius** or **Source Angle** determines how soft the shadows will be. Capsule shadowing is very efficient as it computes shadowing at half resolution with a depth-aware upsample and uses screen tile culling to limit the shadowing work to where it is needed. 
+The GPU cost is proportional to the number of capsules and the number of pixels affected by the casted shadow. 
+
+## Skeletal Mesh
+
+![](SkeletalMesh.png)
+
+In this section, you can assign a **Post Process Anim Blueprint** which is an [Animation Blueprint](Engine/Animation/AnimBlueprints) class to run as a Post Process for the mesh. This Blueprint will be run before Physics, but after the Main Anim Instance for any Skeletal Mesh component using the mesh. 
+To assign an Animation Blueprint, click the drop-down menu and specify your desired Animation Blueprint. 
+
+You can also assign **Asset User Data** from the expanded menu option and click the **+** sign to add elements to the array. Any additions will be stored along with the Skeletal Mesh asset. 
+
+## Transform
+
+![](TransformSection.png)
+
+The **Transform** section inside the Skeletal Mesh Editor allows you to offset the mesh from the world origin when imported to match its position in the source art program through the **Import Translation** option. 
+The **Import Rotation** will rotate the mesh when imported to match its position in the source art program while the **Import Uniform Scale** option will import the uniform scaling of the mesh from the source art program. 
+
+## Thumbnail
+
+![](ThumbnailSection.png)
+
+The **Thumbnail** options allow you to modify how the mesh is framed in the thumbnail that appears in the **Content Browser** or **Editor Toolbar**. 
+
+## Import Settings
+
+![](ImportSettingsSection.png)
 
-Outside the fact that this panel is reserved for specific mesh editing actions within Persona, it is functionally identical to the standard **Details** panel of the main editor. 
+The **Import Settings** section allows you to modify any of the settings used during the initial [FBX Skeletal Mesh Import Process](https://docs.unrealengine.com/latest/INT/Engine/Content/FBX/ImportOptions/#skeletalmeshoptions) for the mesh. 
 
-For more information about the basics of using this panel, please see the [Editor Details Panel](Engine/UI/LevelEditor\Details) documentation.
 
diff --git a/Engine/Documentation/Source/Engine/Animation/Persona/Modes/Animation/Animation.INT.udn b/Engine/Documentation/Source/Engine/Animation/Persona/Modes/Animation/Animation.INT.udn
--- a/Engine/Documentation/Source/Engine/Animation/Persona/Modes/Animation/Animation.INT.udn
+++ b/Engine/Documentation/Source/Engine/Animation/Persona/Modes/Animation/Animation.INT.udn
@@ -1,40 +1,60 @@
 Availability: Public
-Title:Animation Mode
+Title:Animation Editor
 Crumbs: %ROOT%, Engine, Engine/Animation, Engine/Animation/Persona, Engine/Animation/Persona/Modes
-Description:User guide for Animation Mode in the Persona animation editor.
-version: 4.9
-
-[VAR:Topic]
-[OBJECT:Topic]
-	[PARAM:image]
-		![%Engine/Animation/Persona/Modes/Animation:title%](Engine/Animation/animation_topic.png)
-	[/PARAM]
-	[PARAM:icon]
-		![](%ROOT%/content_icon.png)(convert:false)
-	[/PARAM]
-	[PARAM:title]
-		%Engine/Animation/Persona/Modes/Animation:title%
-	[/PARAM]
-	[PARAM:description]
-		%Engine/Animation/Persona/Modes/Animation:description%
-	[/PARAM]
-	[PARAM:path]
-		[RELATIVE:Engine/Animation/Persona/Modes/Animation]
-	[/PARAM]
-[/OBJECT]
-[/VAR]
+Description:Highlights the core features of the Animation Editor in Unreal Engine 4. 
+version: 4.14
+type:landing
+topic-image:animEditor_topic.png
+Related: Engine\Animation\Persona\Modes\Skeleton
+Related: Engine\Animation\Persona\Modes\Mesh
+Related: Engine\Animation\Persona\Modes\Graph
+Related: Engine/Content/FBX/Animations
+tags:Animation
+tags:Skeletal Mesh
+tags:FBX Importing
 
+[EXCERPT:Overview]
+The **Animation Editor** provides easy access to the various animation-centric assets available for a Skeletal Mesh. 
+In the Animation Editor, you can preview playback of animation assets such as Animation Sequences, Blend Spaces, Animation Montages, make edits to animation assets, add and edit curves for Material Parameters or Morph Targets as well as define Animation Notifies (events that occur during designated points in an animation). 
+[/EXCERPT:Overview]
 
+When opening any animation asset through the **Content Browser** (or Editor Toolbar), the Animation Editor will open: 
 
 [REGION:fullwidth]
 ![](AnimationModeWindow.png)
 [/REGION]
 
-The **Animation** mode of **Persona** (accessed by selecting the **Animation** tab from the toolbar window) contains several windows which appear in other modes, however the three windows (denoted above) are the ones used most when working in Animation mode. 
+Please refer to each section below for a breakdown of the Animation Editor user interface: 
+
+## 1. Toolbar
+
+The [Toolbar](Engine/Animation/Persona/Toolbar) inside the Animation Editor provides you with options for saving any changes to your animation or locating it in the Content Browser. It also allows you to create new animation assets, define compression settings, adjust additive layer tracks, or create a static mesh out of the preview's current pose. 
+On the far right of the Toolbar is the **Editor Toolbar** which allows you to switch between the different Animation Tools in Unreal Engine 4. 
+
+## 2. Asset Details / Skeleton Tree
+
+The [Asset Details](Engine\Animation\Persona\AnimAssetDetails) tab provides you with options to define how your animation asset functions and will change based on the type of animation asset you have open. 
+Docked next to the Asset Details is the [](Engine/Animation/Persona/SkeletonTree) which displays the skeletal hierarchy of the current Skeleton asset and allows you to create/edit Sockets and define settings related to animation retargeting. 
+
+## 3. Viewport
+
+The [Viewport](Engine/Animation/Persona/Viewport) window allows you to preview playback of animation assets on your selected Skeletal Mesh and provides information about your assets. You can change lighting modes, show or hide the bones of your skeleton, adjust animation playback speeds, even set your Skeletal Mesh to automatically rotate around on a turntable allowing you view it from all angles. 
+
+## 4. Details
+
+The [Details](Engine/UI/LevelEditor/Details) panel, which is similar to the Main Editor, is primarily used for modifying options that have been added while working in the **Asset Editor**. 
+For example, when you add a notify to an Animation Sequence, clicking on the notify will populate the Details panel with options related to the notify. Adding a new section to an Animation Montage, you can modify the content of the section from the Details panel. 
+
+## 5. Asset Editor
+
+The [Asset Editor](Engine/Animation/Persona/AssetEditor) is a contextual editor that changes interface layout and options based on the type of animation asset opened. You can modify Anim Sequences, Blend Spaces, Anim Montages or any other animation asset inside this window as well as playback animations (or record new animations) from the toolbar located at the bottom. 
+
+## 6. Asset Browser
+
+The [](Engine\Animation\Persona\AssetBrowser) allows you to view all the animation assets that are associated with and can be used by the current Skeleton asset. 
+
+
 
-1. [Anim Asset Details](Engine\Animation\Persona\AnimAssetDetails) - The **Anim Asset Details** panel is a context-sensitive property editor that allows you to change settings on various animation assets, such as Anim Sequences, Blend Spaces, Anim Montages, and so on. By default, the **Anim Asset Details** panel is only visible when Persona is in **Animation** mode.
-1. [Anim Sequence Editor](Engine\Animation\Sequences\Editor) - The **Animation Sequence Editor** provides functionality for previewing and editing individual AnimationSequence assets. The properties, including compression scheme, of the AnimationSequence can be set, and animation notification events - also called Notifies - can be added to allow camera effects, particle effects, sounds, and more to be triggered based on the animation.
-1. [Asset Browser](Engine\Animation\Persona\AssetBrowser) - The **Asset Browser** panel displays the AnimSequences, BlendSpaces, AnimMontages, and AimOffsetBlendSpaces that are valid for the target Skeleton in a mini **Content Browser**-style interface. 
 
 
 
diff --git a/Engine/Documentation/Source/Engine/Animation/Persona/Modes/Graph/Graph.INT.udn b/Engine/Documentation/Source/Engine/Animation/Persona/Modes/Graph/Graph.INT.udn
--- a/Engine/Documentation/Source/Engine/Animation/Persona/Modes/Graph/Graph.INT.udn
+++ b/Engine/Documentation/Source/Engine/Animation/Persona/Modes/Graph/Graph.INT.udn
@@ -1,42 +1,29 @@
 Availability: Public
-Title:Graph Mode
+Title:Anim Blueprint Editor
 Crumbs: %ROOT%, Engine, Engine/Animation, Engine/Animation/Persona, Engine/Animation/Persona/Modes
-Description:User guide for Graph Mode in the Persona animation editor.
-version: 4.9
-
-[VAR:Topic]
-[OBJECT:Topic]
-	[PARAM:image]
-		![%Engine/Animation/Persona/Modes/Graph:title%](Engine/Animation/animation_topic.png)
-	[/PARAM]
-	[PARAM:icon]
-		![](%ROOT%/content_icon.png)(convert:false)
-	[/PARAM]
-	[PARAM:title]
-		%Engine/Animation/Persona/Modes/Graph:title%
-	[/PARAM]
-	[PARAM:description]
-		%Engine/Animation/Persona/Modes/Graph:description%
-	[/PARAM]
-	[PARAM:path]
-		[RELATIVE:Engine/Animation/Persona/Modes/Graph]
-	[/PARAM]
-[/OBJECT]
-[/VAR]
-
+Description:Highlights the core features of the Anim Blueprint Editor in Unreal Engine 4. 
+Related: Engine/Animation/AnimBlueprints
+Related: Engine/Blueprints
+Related: Engine\Animation\Persona\Modes\Animation
+Related: Engine\Animation\Persona\Modes\Skeleton
+Related: Engine\Animation\Persona\Modes\Mesh
+version: 4.14
+type:landing
+topic-image:AnimGraphEditor_topic.png
+tags:Animation
+tags:Blueprints
 
 
 [REGION:fullwidth]
-![](AnimationGraph.png)
+![](AnimBPEditor.png)
 [/REGION]
 
-The **Graph** mode of **Persona** is where you specify the logic behind and hook up animations and animation states based off variables and conditions set within the **AnimGraph** and **EventGraph**. Working with the graphs, is similar to working with Blueprint graphs however there are some key details to Animation Blueprints as outlined below. 
-
-##Animation Blueprints Overview
+The **Animation Blueprint Editor** is where you specify the logic the drives your character's animations and can be accessed by opening any **Animation Blueprint** from the **Content Browser** (or from the **Editor Toolbar**). 
+This editor is very similar to working with the [Blueprint Editor](Engine/Blueprints) as it uses the same visual node based scripting system with an emphasis on animation features like managing animation [](Engine/Animation/StateMachines), using [](Engine/Animation/NodeReference/Blend) to control [](Engine/Animation/AnimationBlending), or direct [Skeletal Control](Engine/Animation/NodeReference/SkeletalControls) through script. 
 
-[INCLUDE:Engine/Animation/AnimBlueprints#Intro]
-
-For more information on working with Animation Blueprints, refer to the [Animation Blueprints](Engine/Animation/AnimBlueprints) documentation. 
+[REGION:note]
+For more information on working with **Animation Blueprints** and the **Animation Blueprint Editor**, refer to the [Animation Blueprints](Engine/Animation/AnimBlueprints) documentation. 
+[/REGION]
 
 
 
diff --git a/Engine/Documentation/Source/Engine/Animation/Persona/Modes/Mesh/Mesh.INT.udn b/Engine/Documentation/Source/Engine/Animation/Persona/Modes/Mesh/Mesh.INT.udn
--- a/Engine/Documentation/Source/Engine/Animation/Persona/Modes/Mesh/Mesh.INT.udn
+++ b/Engine/Documentation/Source/Engine/Animation/Persona/Modes/Mesh/Mesh.INT.udn
@@ -1,155 +1,49 @@
 Availability: Public
-Title:Mesh Mode
+Title:Skeletal Mesh Editor
 Crumbs: %ROOT%, Engine, Engine/Animation, Engine/Animation/Persona, Engine/Animation/Persona/Modes
-Description:User guide for Mesh Mode in the Persona animation editor.
-version: 4.9
-
-[VAR:Topic]
-[OBJECT:Topic]
-	[PARAM:image]
-		![%Engine/Animation/Persona/Modes/Mesh:title%](Engine/Animation/animation_topic.png)
-	[/PARAM]
-	[PARAM:icon]
-		![](%ROOT%/content_icon.png)(convert:false)
-	[/PARAM]
-	[PARAM:title]
-		%Engine/Animation/Persona/Modes/Mesh:title%
-	[/PARAM]
-	[PARAM:description]
-		%Engine/Animation/Persona/Modes/Mesh:description%
-	[/PARAM]
-	[PARAM:path]
-		[RELATIVE:Engine/Animation/Persona/Modes/Mesh]
-	[/PARAM]
-[/OBJECT]
-[/VAR]
-
+Description:Highlights the core features of the Skeletal Mesh Editor in Unreal Engine 4. 
+version: 4.14
+Related: Engine\Animation\Persona\Modes\Skeleton
+Related: Engine\Animation\Persona\Modes\Animation
+Related: Engine\Animation\Persona\Modes\Graph
+Related: Engine/Content/FBX/SkeletalMeshes
+type:landing
+topic-image:skelMeshEditor_topic.png
+tags:Animation
+tags:Skeletal Mesh
+tags:FBX Importing
 
 [EXCERPT:Overview]
-**Mesh** mode of the **Persona** animation editor (accessed by clicking the **Mesh** tab) shares some of the same windows that are available while in other Editing Modes, however there are two windows which are specific to **Mesh** mode (highlighted below): [Mesh Details](Engine/Animation/Persona\MeshDetails) and [Morph Target Previewer](Engine/Animation/Persona\MorphTargetPreviewer). 
-[/EXCERPT:Overview]
+Whenever you open a **Skeletal Mesh** asset inside the **Content Browser** (or from the **Editor Toolbar**), the **Skeletal Mesh Editor** will open. 
+This editor allows you to make changes to the polygonal mesh by setting up **Materials** for the Skeletal Mesh, adding clothing elements, setting up LODs and testing out any Morph Targets applied to the mesh. 
+This editor includes windows that can be found in some of the other Animation Tools such as the Toolbar/Viewport (and others that are hidden by default) but most of your mesh work will be done in the **Asset Details** and **Morph Target Preview** windows. 
 
-[REGION:fullwidth]
-![](MeshModeWindows.png)
-[/REGION]
-
-##Mesh Details
-The **Mesh Details** window allow you to modify properties pertaining to the mesh of the Skeletal Mesh such as **Material** and **Level of Detail (LOD)** settings, assigning **Apex Clothing**, adding meshes to the existing mesh, adjusting **Mirroring** settings or assigning a **Physics Asset** to the mesh. You can also **Reimport** asset data, adjust **Thumbnail Orbit** or **Texture Streaming** settings. 
-
-###LOD Settings
+[/EXCERPT:Overview]
 
 [REGION:fullwidth]
-![](LOD_Settings.png)
+![](SkeletalMeshEditor.png)
 [/REGION]
 
-In this section, you can assign Materials to aspects of the Skeletal Mesh and toggle whether or not those Materials cast shadows or not. 
-
-In the **LOD Info** section, there are settings that allow you to modify how the LOD(s) function: 
-
-| Property| Description |
-|--- | --- |
-| **Display Factor** | Indicates when to use the LOD. A smaller number means use the LOD when further away.|
-| **LODHysteresis** | Used to avoid "flickering" when on LOD boundary. Only taken into account when moving from complex to simple. |
-| **LODMaterial Map** | Specifies the mapping table from the LOD's materials to the USkeletalMesh materials array. |
-| **Triangle Sorting** | Set the current triangle sorting method. |
-| **Custom Left Right Axis** | Set the Left Right Axis.  |
-| **Custom Left Right Bone Name** | Set the Left Right Bone name.  |
-
-###Clothing
-
-![](AddClothing.png)
-
-This section of the **Mesh Details** allows you assign **Apex Cloth** assets to the mesh and view information about the assets assigned, there are also **Physics Properties** that can be adjusted for each asset assigned. 
-
-| Property| Description |
-|--- | --- |
-| **Bend Resistance** | Bending stiffness of the cloth in the range [0,1]. |
-| **Shear Resistance** | Shearing stiffness of the cloth in the range [0,1]. |
-| **Stretch Limit** | Make cloth simulation less stretchy. A value smaller than 1 will turn it off. |
-| **Friction** | Friction coefficient in the range [0,1]. |
-| **Damping** | Spring damping of the cloth in range [0,1].  |
-| **Drag** | Drag coefficient in the range [0,1].  |
-| **Gravity Scale** | Amount of gravity that is applied to the cloth. |
-| **Inertia Blend** | Amount of inertia that is kept when using local space simulation. Internal name is inertia scale. |
-| **Self Collision Thickness** | Minimal amount of distance particles will keep of each other. |
-
-###Additional Meshes
-
-This section allows you to add meshes to the existing mesh or remove meshes that have been added. 
-
-###Mesh
-
-This section lists the Skeleton that is used with the Skeletal Mesh and its location inside the **Content Browser**. 
-
-###Skeletal Mesh
-
-![](SkeletalMeshSection.png)
-
-The **Skeletal Mesh** section allows you to define Materials for any Material Interfaces that exist and whether or not to enable Shadow Casting from the Material. The ability to set whether or not to **Use Full Precision UVs** is available which if true, it will use 32 bit UVs and if false, it will use 16 bit UVs to save memory. 
-
-<!--Commenting this section out as it no longer exists as of 4.9 Preview 4 
-###Mirroring
-
-![](MirroringSection.png)
-
-There are two options available inside the **Mirroring** section that can be set: 
-
-| Property| Description |
-|--- | --- |
-| **Skel Mirror Table** | This lists all the bones that should be mirrored. |
-| **Skel Mirror Axis** | Set the Mirror Axis to X, Y, Z or None. |
-| **Skel Mirror Flip Axis** | Set the Mirror Flip Axis to X, Y, Z or None. |-->
-
-###Physics
-
-In the **Physics** section, you can assign a Physics Asset to the mesh. The location of the asset in the **Content Browser** (if applicable) is also shown. 
-
-###Reimport
-
-![](ReimportSection.png)
-
-This option allows you to reimport the asset being used and adjust the import settings (also shows the currently used import data and settings when applied).
-
-| Property| Description |
-|--- | --- |
-| **Update Skeleton Reference Pose** | If enabled, update the Skeleton (of the mesh being imported)'s reference pose. |
-| **Use T0As Ref Pose** | Enable this option to use frame 0 as the reference pose.  |
-| **Preserve Smoothing Groups** | If checked, triangles with non-matching smoothing groups will be physically split. |
-| **Import Meshes in Bone Hierarchy** | If checked, meshes nested in bone hierarchies will be imported instead of being converted to bones.  |
-| **Import Morph Targets** | If enabled, creates Unreal morph objects for the imported meshes. |
-| **Import Mesh LODs** | If enabled, creates LOD models for Unreal meshes from LODs in the imported file; if not enabled, only the base mesh from the LOD group is imported. |
-| **Normal Import Method** | Enabling this option will read the tangents (tangent, binormal, normal) from FBX files instead of generating them automatically.  |
-| **Source File Path** | This is the path to the resource used to construct the Static Mesh. |
-| **Source File Timestamp** | Date/Time-stamp of the file from the last import. |
-
-###Thumbnail
-
-![](ThumbnailSection.png)
-
-The **Thumbnail** section allows you to set the **Orbit Pitch**, **Orbit Yaw** and **Orbit Zoom** of the camera around the asset.
-
-| Property| Description |
-|--- | --- |
-| **Orbit Pitch** | The pitch of the orbit camera around the asset. |
-| **Orbit Yaw** | The yaw of the orbit camera around the asset. |
-| **Orbit Zoom** | The offset from the bounds sphere distance from the asset. |
-
-###Texture Streaming
+Please refer to each section below for a breakdown of the Skeletal Mesh Editor user interface: 
 
-![](TextureStreamingSection.png)
+## 1. Toolbar
 
-**Texture Streaming** allows artists to adjust the distance where textures using UV 0 are streamed in/out.
+The [Toolbar](Engine/Animation/Persona/Toolbar) inside the Skeletal Mesh Editor provides you with options for saving any changes to your animation or locating it in the Content Browser. 
+It also allows you to create a Static Mesh out of the preview's current pose. 
+On the far right of the Toolbar is the **Editor Toolbar** which allows you to switch between the different Animation Tools in Unreal Engine 4. 
 
-| Property| Description |
-|--- | --- |
-| **Streaming Distance Multiplier** | 1.0 is the default, where a higher value increases the streamed-in resolution. |
+## 2. Asset Details
 
+The [Asset Details](Engine\Animation\Persona\MeshDetails) panel shows specific properties pertaining to the Skeletal Mesh, such as the Materials applied to the surface, as well as skinning and LOD options. 
+You can also assign and adjust APEX clothing or define and associate a **Physics Asset** for the Skeletal Mesh which is needed for creating rag-doll type motions. 
 
+## 3. Viewport
 
-##Morph Target Previewer
+The [Viewport](Engine/Animation/Persona/Viewport) window allows you to preview the changes that you make to the Mesh in the Asset Details panel or to preview any Morph Target settings. 
+From the Viewport, you can also assign a preview animation for the Skeletal Mesh to use, change lighting modes, show or hide the bones of your skeleton, adjust animation playback speeds, even set your Skeletal Mesh to automatically rotate around on a turntable allowing you view it from all angles. 
 
-[INCLUDE:Engine/Animation/Persona/MorphTargetPreviewer#Overview]
+## 4. Morph Target Preview
 
-For more information, refer to the [Morph Target Previewer](Engine/Animation/Persona/MorphTargetPreviewer/) documentation. 
+The Morph Target Preview tab lets you preview any [Morph Targets](Engine/Animation/Persona/MorphTargetPreviewer) (or Blend Shapes) available for the current Skeletal Mesh. 
 
 
diff --git a/Engine/Documentation/Source/Engine/Animation/Persona/Modes/Modes.INT.udn b/Engine/Documentation/Source/Engine/Animation/Persona/Modes/Modes.INT.udn
--- a/Engine/Documentation/Source/Engine/Animation/Persona/Modes/Modes.INT.udn
+++ b/Engine/Documentation/Source/Engine/Animation/Persona/Modes/Modes.INT.udn
@@ -1,114 +1,37 @@
 Availability: Public
-Title:Persona Editing Modes
+Title:Animation Tools
 Crumbs: %ROOT%, Engine, Engine/Animation, Engine/Animation/Persona
-Description:User guide for the editing modes available in the Persona animation editor.
-version: 4.9
-
-[VAR:Topic]
-[OBJECT:Topic]
-	[PARAM:image]
-		![%Engine/Animation/Persona/Modes:title%](Engine/Animation/animation_topic.png)
-	[/PARAM]
-	[PARAM:icon]
-		![](%ROOT%/content_icon.png)(convert:false)
-	[/PARAM]
-	[PARAM:title]
-		%Engine/Animation/Persona/Modes:title%
-	[/PARAM]
-	[PARAM:description]
-		%Engine/Animation/Persona/Modes:description%
-	[/PARAM]
-	[PARAM:path]
-		[RELATIVE:Engine/Animation/Persona/Modes]
-	[/PARAM]
-[/OBJECT]
-[/VAR]
-
-
+Description:User guide for the different Animation Tools and their respective editors in Unreal Engine 4. 
+Related: Engine/Content/FBX/Animations
+Related: Engine/Content/Types/SkeletalMeshes
+Related: Engine/Animation/PhysicallyDrivenAnimation
+Related: Engine/Content/Tools/MayaRiggingTool
+version: 4.14
+type:landing
+topic-image:animation_tools.png
+tags:Animation
+tags:Skeleton Asset
+tags:Animation Blueprint
+tags:Skeletal Mesh
 
 [EXCERPT:Overview]
-It is important to remember that Persona is actually a collection of various animation tools and editors, rather than just a single tool unto itself. It allows for editing of Skeletons, Meshes, Animation, and Animation Blueprints. As such, Persona contains a variety of different modes in which it can operate in order to optimally handle the editing task at hand. 
+Creating animated characters in Unreal Engine 4 involves using several different Animation Tools (or Editors), each of which focus on different aspects of animation. 
 
-In each mode, Persona's UI updates to show a different layout of tools and editors.
+For instance, the **Skeleton Editor** is where everything starts as it is used to manage the bone (or joint hierarchy) that drives the Skeletal Mesh and the animation. 
+The **Skeletal Mesh Editor** is used to modify the Skeletal Mesh that is linked to a Skeleton and is the outward appearance of a character. 
+The **Animation Editor** allows you to create and modify your animation assets and is where the tuning/tweaking of animation assets are done. 
+Finally, the **Animation Blueprint Editor** can be used to create logic that drives what animations a character uses and when as well as how animations are blended together. 
 [/EXCERPT:Overview]
 
-[REGION:note]
-	Not all of the modes are available at all times. For instance,**Graph** mode is only available when an Animation Blueprint is opened from the **Content Browser**. 
+Please refer to each editor for more information: 
 
-	Also, be aware that changes made to editing layouts will be maintained. As an example, if you move a panel while in Skeleton editing mode, then switch away to Animation, when you switch back to Skeleton, the moved panel will keep the position you gave it.
+[EXCERPT:animEditors]
+[REGION:topics half]
+%Engine\Animation\Persona\Modes\Skeleton:TopicCompact%
+%Engine\Animation\Persona\Modes\Mesh:TopicCompact%
+%Engine\Animation\Persona\Modes\Animation:TopicCompact%
+%Engine\Animation\Persona\Modes\Graph:TopicCompact%
 [/REGION]
+[/EXCERPT]
 
 
-[EXCERPT:SkeletonMode]
-## Skeleton Editing Mode
-
-[![](Persona_Mode_Skeleton.png)(w:800)](Engine/Animation/Persona/Modes/Skeleton)
-
-(Click the image to visit the **Skeleton Editing** mode page.)
-
-[EXCERPT:SkeletonModeUI]
-1. [Toolbar](Engine/Animation/Persona\Toolbar) 
-1. [Skeleton Tree](Engine/Animation/Persona\SkeletonTree)
-1. [Viewport](Engine/Animation/Persona\Viewport)
-1. [Details](Engine/UI/LevelEditor\Details) 
-1. [Asset Browser](Engine/Animation/Persona\AssetBrowser)
-[/EXCERPT:SkeletonModeUI]
-[/EXCERPT:SkeletonMode]
-
-[REGION:note]
-The **Details** panel in Persona will only show properties pertaining to animation assets. Functionally, however, it is the same as the standard **Details** panel in the main editor window.
-[/REGION]
-
-[EXCERPT:MeshMode]
-## Mesh Editing Mode
-
-[![](Persona_Mode_Mesh.png)(w:800)](Engine/Animation/Persona/Modes/Mesh)
-
-(Click the image to visit the **Mesh Editing** mode page.)
-
-[EXCERPT:MeshModeUI]
-1. [Toolbar](Engine/Animation/Persona\Toolbar) 
-1. [Skeleton Tree](Engine/Animation/Persona\SkeletonTree)
-1. [Mesh Details](Engine/Animation/Persona\MeshDetails)
-1. [Viewport](Engine/Animation/Persona\Viewport)
-1. [Details](Engine/UI/LevelEditor\Details)/[Asset Browser](Engine/Animation/Persona\AssetBrowser)
-1. [Morph Target Previewer](Engine/Animation/Persona\MorphTargetPreviewer)
-[/EXCERPT:MeshModeUI]
-[/EXCERPT:MeshMode]
-
-
-[EXCERPT:AnimationMode]
-## Animation Editing Mode
-
-[![](Persona_Mode_Anmation.png)(w:800)](Engine/Animation/Persona/Modes/Animation)
-
-(Click the image to visit the **Animation Editing** mode page.)
-
-[EXCERPT:AnimationModeUI]
-1. [Toolbar](Engine/Animation/Persona\Toolbar) 
-1. [Skeleton Tree](Engine/Animation/Persona\SkeletonTree)
-1. [Anim Asset Details](Engine/Animation/Persona\AnimAssetDetails)
-1. [Viewport](Engine/Animation/Persona\Viewport)
-1. [Anim Sequence Editor](Engine/Animation/Sequences\Editor)
-1. [Details](Engine/UI/LevelEditor\Details)
-1. [Asset Browser](Engine/Animation/Persona\AssetBrowser)
-[/EXCERPT:AnimationModeUI]
-[/EXCERPT:AnimationMode]
-
-
-[EXCERPT:GraphMode]
-## Graph Editing Mode
-
-[![](Persona_Mode_Vim.png)(w:800)](Engine/Animation/Persona/Modes/Graph)
-
-(Click the image to visit the **Graph Editing** mode page.)
-
-[EXCERPT:GraphModeUI]
-1. [Toolbar](Engine/Animation/Persona\Toolbar) 
-1. [Viewport](Engine/Animation/Persona\Viewport)
-1. [Preview Settings (Class Defaults)](Engine/Blueprints\Editor\Modes\Defaults)
-1. [Blueprint Graph Panel](Engine/Animation/AnimBlueprints)
-1. [Asset Browser](Engine/Animation/Persona\AssetBrowser)/[My Blueprint](Engine/Blueprints\Editor\UIComponents\MyBlueprint)
-1. [Details](Engine/UI/LevelEditor\Details)
-[/EXCERPT:GraphModeUI]
-[/EXCERPT:GraphMode]
diff --git a/Engine/Documentation/Source/Engine/Animation/Persona/Modes/Skeleton/Skeleton.INT.udn b/Engine/Documentation/Source/Engine/Animation/Persona/Modes/Skeleton/Skeleton.INT.udn
--- a/Engine/Documentation/Source/Engine/Animation/Persona/Modes/Skeleton/Skeleton.INT.udn
+++ b/Engine/Documentation/Source/Engine/Animation/Persona/Modes/Skeleton/Skeleton.INT.udn
@@ -1,26 +1,57 @@
 Availability: Public
-Title:Skeleton Mode
+Title:Skeleton Editor
 Crumbs: %ROOT%, Engine, Engine/Animation, Engine/Animation/Persona
-Description:
-version: 4.9
+Description:Highlights the core features of the Skeleton Editor in Unreal Engine 4. 
+version: 4.14
+Related: Engine\Animation\Persona\Modes\Mesh
+Related: Engine\Animation\Persona\Modes\Animation
+Related: Engine\Animation\Persona\Modes\Graph
+Related: Engine/Content/FBX/SkeletalMeshes
+type:landing
+topic-image:skeletonEdtior_topic.png
+tags:Skeleton
+tags:Skeletal Mesh
+tags:Animation
 
 [EXCERPT:Overview]
-**Skeleton** mode allows for visualization and control over those aspects of a character or Skeletal Mesh involving its bone or joint hierarchy. Example uses include:
-
-* Creating [Sockets](Engine/Content\Types\SkeletalMeshes\Sockets) to attach items to a Skeletal Mesh
-* Editing the relative position and rotation for each [Sockets](Engine/Content\Types\SkeletalMeshes\Sockets)
-* Establishing animation retargeting relationships
-
+The **Skeleton Editor** is a tool used for working with [](Engine/Animation/Skeleton) inside Unreal Engine 4. 
+This editor allows for the visualization and control over the bone or joint hierarchy associated with a Skeletal Mesh. 
+In this editor, you can create [Skeletal Mesh Sockets](Engine/Content\Types\SkeletalMeshes\Sockets) to attach items to a Skeletal Mesh, preview [](Engine/Animation/Sequences/Curves) and track the associated [](Engine/Animation/Sequences/Notifies) associated with the Skeleton. 
+It is in this editor where you can also set up your [](Engine/Animation/AnimationRetargeting) options and manage your retarget sources with the [Retarget Manager.](Engine/Animation/Persona/BasePoseManager/) 
 [/EXCERPT:Overview]
 
 [REGION:fullwidth]
 ![](Persona_MODE_Skeleton_Large.png)
+[/REGION]
 
+Please refer to each section below for a breakdown of the Skeleton Editor user interface: 
 
-[INCLUDE:Engine/Animation/Persona#skeletonui]
+## 1. Toolbar
+
+The [Toolbar](Engine/Animation/Persona/Toolbar) inside the Skeleton Editor provides you with options for saving any changes made to your Skeleton or locating it in the Content Browser. 
+You can view **Anim Notifies**, open the **Retarget Manager**, import a new Skeletal Mesh to associate with the Skeleton and create a Static Mesh out of the preview's current pose. 
+On the far right of the Toolbar is the **Editor Toolbar** which allows you to switch between the different Animation Tools in Unreal Engine 4. 
+
+## 2. Skeleton Tree
+
+The [](Engine/Animation/Persona/SkeletonTree) shows the skeletal hierarchy of the current Skeleton asset and is where you create and edit sockets, preview assets, and control animation retargeting options. 
+
+## 3. Viewport
+
+The [Viewport](Engine/Animation/Persona/Viewport) window allows you to preview the changes that you make and adjust socket positioning or preview any animation curves. 
+From the Viewport, you can also assign a preview animation for the Skeletal Mesh to use, change lighting modes, show or hide the bones of your skeleton, adjust animation playback speeds, even set your Skeletal Mesh to automatically rotate around on a turntable allowing you view it from all angles. 
+
+## 4. Details
+
+The [Details](Engine/UI/LevelEditor/Details) panel, which is similar to the Main Editor, is primarily used for modifying the properties of added elements like Sockets. 
+For example, when you add a Socket to the Skeleton, clicking on the Socket in the Skeleton Tree will populate the Details panel with options related to how the Socket functions. 
+
+## 5. Anim Curves / Animation Notifies
+
+The [Anim Curves](Engine/Animation/Persona/AnimCurves) viewer lets you preview any animation curves available for the current mesh in the Viewport. 
+The [Animation Notifies](Engine/Animation/Persona/AnimationNotifies) window lets you modify your custom created Animation Notifies that are associated with the Skeleton asset. 
 
 
-[/REGION]
 <!--
 
 [EXCERPT:SkeletonUses]
@@ -32,8 +63,5 @@ version: 4.9
 
 [/EXCERPT:SkeletonUI]
 
-
-
-
-
 -->
+
diff --git a/Engine/Documentation/Source/Engine/Animation/Persona/MorphTargetPreviewer/MorphTargetPreviewer.INT.udn b/Engine/Documentation/Source/Engine/Animation/Persona/MorphTargetPreviewer/MorphTargetPreviewer.INT.udn
--- a/Engine/Documentation/Source/Engine/Animation/Persona/MorphTargetPreviewer/MorphTargetPreviewer.INT.udn
+++ b/Engine/Documentation/Source/Engine/Animation/Persona/MorphTargetPreviewer/MorphTargetPreviewer.INT.udn
@@ -1,37 +1,23 @@
 Availability:Public
 Title:Morph Target Previewer
 Crumbs: %ROOT%, Engine, Engine/Animation, Engine/Animation/Persona
-Description:User guide for the editing modes available in the Persona animation editor.
-Related: Engine/Content/FBX/Animations
+Description:User guide for the editing modes available in the Animation Editor.
 Related: Engine/Content/FBX/MorphTargets
 Related: Engine/Content/Types/SkeletalMeshes
 Related: Engine/Content/Tools/MayaRiggingTool
 Related: Engine/Content/Types/StaticMeshes/MorphTargets
-Related: Engine/Animation/AnimBlueprints
-version: 4.9
-[VAR:Topic]
-[OBJECT:Topic]
-	[PARAM:image]
-		![%Engine/Animation/Persona/MorphTargetPreviewer:title%](Engine/Animation/animation_topic.png)
-	[/PARAM]
-	[PARAM:icon]
-		![](%ROOT%/content_icon.png)(convert:false)
-	[/PARAM]
-	[PARAM:title]
-		%Engine/Animation/Persona/MorphTargetPreviewer:title%
-	[/PARAM]
-	[PARAM:description]
-		%Engine/Animation/Persona/MorphTargetPreviewer:description%
-	[/PARAM]
-	[PARAM:path]
-		[RELATIVE:Engine/Animation/Persona/MorphTargetPreviewer]
-	[/PARAM]
-[/OBJECT]
-[/VAR]
-
-[TOC (start:2 end:3)]
-
-
+Related: Resources/ContentExamples/MorphTargets
+version: 4.14
+type:reference
+topic-image:MorphTargetSmall.png
+tags:Animation
+tags:Morph Target
+tags:Skeletal Mesh
+tags:FBX Importing
+
+<!--
+[TOC (start:2)]
+-->
 
 [EXCERPT:Overview]
 
@@ -43,14 +29,14 @@ version: 4.9
 360
 [/PARAMLITERAL]
 [PARAMLITERAL:videoid]
-CfAI0c-Br8Y
+YoEQYVEEcKo
 [/PARAMLITERAL]
 [/OBJECT]
 
 The **Morph Target Previewer** allows you to preview any morph targets (also known as "morphs" and "blend shapes", depending on what 3D animation package you are using) that are applied to your Skeletal Meshes. Each Morph Target is additively blended into the existing Skeletal Mesh geometry. This allows for multiple morph targets to be combined to create complex vertex-driven animation, which is often ideal for such things as facial expressions.
 [/EXCERPT:Overview]
 
-The **Morph Target Previewer** is visible by default when Persona is in **Mesh** mode.
+The **Morph Target Previewer** window is visible by default whenever the [](Engine\Animation\Persona\Modes\Mesh) is open. 
 
 [REGION:tip]
 To see an example of a Morph Target applied to a character, refer to the [Animation Content Examples](Resources\ContentExamples\Animation) page under section 1.6.
@@ -59,30 +45,30 @@ To see an example of a Morph Target applied to a character, refer to the [Animat
 
 ## Interface
 
-The **Morph Target Previewer** panel has a fairly straightforward interface:
+The **Morph Target Previewer** panel has two main sections: 
 
 ![](MorphTargetPreviewerInterface.png)
 
 1. Search Bar
 1. Morph Target List
 
-The Search Bar allows you to quickly filter the list of Morph Targets to quickly find the one you need. Simply type the first few letters of the desired target and the list will begin filtering.
+The Search Bar allows you to filter the list of Morph Targets to quickly find the one you need. To filter, type the first few letters of the desired target and the list will begin filtering.
 
-You can also **Right-click** on a Morph Target which will pop-up a dialog box with additional actions such as **Delete**. 
+You can also **Right-click** on a Morph Target which will pop-up a dialog box with additional actions such as **Delete** or **Copy Morph Target Names**. 
 
 ![](RightClickDelete.png)
 
 ## Creating Morph Targets
 
-Morph Targets are importable as a part of a Skeletal Mesh, as well as independently of a given mesh. This is supported via the FBX file format. For more information on the setup procedure, as well as how to import Morph Targets into Unreal, please see [](Engine/Content\FBX\MorphTargets).
+You can import Morph Targets as a part of a Skeletal Mesh, as well as independently of a given mesh. This is supported via the FBX file format. 
 
-[REGION:tip]
-To delete a Morph Target, **Right-click** its name in the **Morph Target Previewer** window, and click **Delete** on the context menu that appears.
+[REGION:note]
+For more information on the setup procedure, as well as how to import Morph Targets into Unreal, please see [](Engine/Content\FBX\MorphTargets).
 [/REGION]
 
 ## Using Morph Targets
 
-Once your Morph Targets are in place, you will need to set up your Animation Blueprint to utilize them. This is done within your Event Graph via the **Set Morph Target** node.
+Once your Morph Targets are in place, you will need to set up your [Animation Blueprint](Engine/Animation/AnimBlueprints) to utilize them. This is done within your Event Graph via the **Set Morph Target** node.
 
 ![](SetMorphTarget.png)
 
@@ -97,7 +83,18 @@ Once your Morph Targets are in place, you will need to set up your Animation Blu
 | Execution | Passes execution on to the next node. |
 
 
+## Optimization
+
+For target platforms that support Shader Model 5, you can enable Morph Targets being calculated on the GPU. This means that the CPU does not have to perform the calculations.
 
+To enable this setting:
+1. On the file menu, go to **Edit** > **Project Settings**.
+1. Open the **Rendering** section of **Project Settings**.
+1. In the **Optimizations** category, find the **Use GPU for computing morph targets** checkbox and enable it. 
+
+[REGION:lightbox]
+[![](ProjectSettingsMorphGPU.png)(w:600)](ProjectSettingsMorphGPU.png)
+[/REGION]
 
 
 
diff --git a/Engine/Documentation/Source/Engine/Animation/Persona/PersonaEditor.INT.udn b/Engine/Documentation/Source/Engine/Animation/Persona/PersonaEditor.INT.udn
--- a/Engine/Documentation/Source/Engine/Animation/Persona/PersonaEditor.INT.udn
+++ b/Engine/Documentation/Source/Engine/Animation/Persona/PersonaEditor.INT.udn
@@ -1,66 +1,73 @@
 Availability:Public
-Title:Persona Animation Editor Reference
+Title:Animation Editors
 Crumbs: %ROOT%, Engine, Engine/Animation
 Description:Editor for previewing and editing AnimBlueprints, AnimationSequences and other animation assets.
-version: 4.9
-
+version: 4.14
+parent:Engine/Animation
+topic-image:Engine/Animation/animation_topic.png
+reference-image:UpdatedPersonaWindowAnim.gif
+Related: Engine/Content/FBX/Animations
+Related: Engine/Content/Types/SkeletalMeshes
+Related: Engine/Animation/PhysicallyDrivenAnimation
+Related: Engine/Content/Tools/MayaRiggingTool
+tags:Animation
+tags:Skeleton Asset
+tags:Animation Blueprint
+tags:Skeletal Mesh
+type:landing
 
 [VAR:ModeSize]w:450 convert:false[/VAR]
 
-[VAR:Topic]
-[OBJECT:Topic]
-	[PARAM:image]
-		![%Engine/Animation/Persona:title%](Engine/Animation/animation_topic.png)
-	[/PARAM]
-	[PARAM:icon]
-		![](%ROOT%/anim_icon.png)(convert:false)
-	[/PARAM]
-	[PARAM:title]
-		%Engine/Animation/Persona:title%
-	[/PARAM]
-	[PARAM:description]
-		%Engine/Animation/Persona:description%
-	[/PARAM]
-	[PARAM:path]
-		[RELATIVE:Engine/Animation/Persona]
-	[/PARAM]
-[/OBJECT]
-[/VAR]
-
-[REGION:fullwidth]
-![](UpdatedPersonaWindow.png)
+[REGION:asyncgif]
+![](UpdatedPersonaWindowAnim.gif)
 [/REGION]
 
-[TOC (start:2 end:3)]
+[EXCERPT:Intro]
+Unreal Engine 4 contains a four major tools for working with skeletal animations and **Skeletal Meshes**: The Skeleton Editor, the Skeletal Mesh Editor, the Animation Editor, and the Animation Blueprint Editor. Each of these are accessed by editing an associated asset or by using the four navigation buttons at the top of each animation editor.
 
+* The **Skeleton Editor** is for examining and modifying the **Skeleton** of a **Skeletal Mesh**. This is where you'd add **Sockets** to joints or test how your Skeletal Mesh's joints are rotating.
 
+* The **Skeletal Mesh Editor** is where you can import LOD's or assign default materials to your **Skeletal Mesh**. It's also where you'll find ways to attach APEX clothing, modify the Skeletal Mesh's bounds, and even assign a specialized **Shadow Physics Asset**.
 
-[EXCERPT:Intro]
+* The **Animation Editor** is where you can work with **Animation Sequences** and other animation assets, such as **Blend Spaces** and **Animation Montages**. Here you will create Anim Notifies, work with Pose Assets, adjust compression and timing of Animation Sequences, and build your Animation Montages.
 
-[INCLUDE:Shared/Glossary/P/#Glossary_Persona]
+* The final editor is the **Animation Blueprint Editor** where you will create the rules for when and how your animations are played. Here you'll be able to use complex state machines and different blends to bring your characters to life.
 
 [/EXCERPT:Intro]
 
-## Opening Persona
+[REGION:note]
+Before Unreal Engine version 4.14, the combined animation editor was known as **Persona**. For 4.14 it was was broken out into four separate editors. While their functionality remains basically the same, some of the interface has changed, and now you can multiple animation editor windows open in different states. For example, you can now have a tab open for each **Animation Sequence** in your game, instead of having a single tab that switches between them (this behavior still remains if you select an **Animation Sequence** in the **Asset Browser**).
+[/REGION]
+
+<!-- [INCLUDE:Shared/Glossary/P/#Glossary_Persona] -->
 
-[REGION:fullwidth]
+## Opening a Specific Animation Editor
+
+[REGION:raw]
 ![](OpeningPersona.png)
 [/REGION]
 
 
-Persona can be opened from the **Content Browser** by **Double-clicking** on an animation asset, such as a:
+Each animation editor can be opened from the **Content Browser** by **Double-clicking** on an asset that is associated with that editor:
 
-* Animation Blueprint
-* Skeleton
-* Skeletal Mesh
-* Animation Sequence
-* BlendSpace
-* Animation Montage
+| Editor | Asset Types |
+| --- | --- |
+| Skeleton Editor | Skeleton Asset |
+| Skeletal Mesh Editor | Skeletal Mesh |
+| Animation Editor | Animation Sequence, Animation Montage, 1D Blend Space, Blend Space, Aim Offset 1D, Aim Offset, Animation Composite, Pose Asset |
+| Animation Blueprint Editor | Animation Blueprint |
 
 [REGION:tip]
-You can also open Persona by **Right-clicking** on any of the above assets in the **Content Browser** and choosing **Edit** from the context menu!
+You can also open an asset's associated editor by **Right-clicking** on any of the above assets in the **Content Browser** and choosing **Edit** from the context menu!
 [/REGION]
 
+## Animation Editors Reference
+
+See the individual editor pages for information on their interfaces and usage:
+
+[INCLUDE:Engine/Animation/Persona/Modes#animeditors]
+
+<!--
 ## Persona Modes and User Interface
 
 [OBJECT:EmbeddedVideo]
@@ -88,21 +95,38 @@ As a vast collection of individual tools, Persona's interface could easily overw
 
 Each mode corresponds to a specific type of task you will want to perform with the tools. Within each mode, only a certain selection of tools and interface elements are visible and available. Keep in mind, however, there is still tremendous room for flexibility, along with some cross-over for performing various tasks without having to switch modes. 
 
+
+Whether visible or not, all aspects of the interface that can be accessed for each mode are available via Persona's **Window** menu. The full list of panels is below, and each mode's section lists the panels available by default in that mode.
+
+* [Skeleton Tree](Engine/Animation/Persona/SkeletonTree)
+* [Asset Browser](Engine/Animation/Persona/AssetBrowser)
+* [Viewport](Engine/Animation/Persona/Viewport)
+* [Morph Target Previewer](Engine/Animation/Persona/MorphTargetPreviewer)
+* [Anim Curves](Engine/Animation/Persona/AnimCurves)
+* [Animation Notifies](Engine/Animation/Persona/AnimationNotifies)
+* [Retarget Manager](Engine/Animation/Persona\BasePoseManager)
+* [Anim Slot Manager](Engine/Animation/AnimMontage#animslotmanager)
+* Blend Profiles
+* Preview Scene Settings
+* [Details](Engine/UI/LevelEditor/Details)
+* [Anim Asset Details](Engine/Animation/Persona/AnimAssetDetails)
+* [Toolbar](Engine/Animation/Persona/Toolbar)
+
+
+
 [REGION:note]
 The locations below are the default locations for each window. As you move panels around, each interface element may appear in different locations. Be sure to use the tabs for each panel so that you know which panel is which.
 [/REGION]
 
 
-[REGION:tip]
-Whether visible or not, all aspects of the interface that can be accessed for each mode are available via Persona's Window menu.
-[/REGION]
+
 
 [OBJECT:IconSection]
 	[PARAM:icon]
 		[![](Engine\Animation\Persona\Modes\Persona_Mode_Skeleton.png)(%modesize%)](Engine\Animation\Persona\Modes\Persona_Mode_Skeleton.png)
 	[/PARAM]
 	[PARAM:title]
-		Skeleton Mode
+		Skeleton Editor
 	[/PARAM]
 	[PARAM:description]
 		[INCLUDE:#skeletonoverview]
@@ -117,7 +141,7 @@ Whether visible or not, all aspects of the interface that can be accessed for ea
 		[![](Engine\Animation\Persona\Modes\Persona_Mode_Mesh.png)(%modesize%)](Engine\Animation\Persona\Modes\Persona_Mode_Mesh.png)
 	[/PARAM]
 	[PARAM:title]
-		Mesh Mode
+		Skeletal Mesh Editor
 	[/PARAM]
 	[PARAM:description]
 		[INCLUDE:#meshoverview]
@@ -133,7 +157,7 @@ Whether visible or not, all aspects of the interface that can be accessed for ea
 		[![](Engine\Animation\Persona\Modes\Persona_Mode_Anmation.png)(%modesize%)](Engine\Animation\Persona\Modes\Persona_Mode_Anmation.png)
 	[/PARAM]
 	[PARAM:title]
-		Animation Mode
+		Animation Editor
 	[/PARAM]
 	[PARAM:description]
 		[INCLUDE:#animationoverview]
@@ -151,7 +175,7 @@ Whether visible or not, all aspects of the interface that can be accessed for ea
 		[![](Engine\Animation\Persona\Modes\Persona_Mode_Vim.png)(%modesize%)](Engine\Animation\Persona\Modes\Persona_Mode_Vim.png)
 	[/PARAM]
 	[PARAM:title]
-		Graph Mode
+		Animation Blueprint Editor
 	[/PARAM]
 	[PARAM:description]
 		[INCLUDE:#graphoverview]
@@ -166,7 +190,7 @@ Whether visible or not, all aspects of the interface that can be accessed for ea
 [/REGION]
 
 
-
+-->
 
 
 
diff --git a/Engine/Documentation/Source/Engine/Animation/Persona/SkeletonTree/SkeletonTree.INT.udn b/Engine/Documentation/Source/Engine/Animation/Persona/SkeletonTree/SkeletonTree.INT.udn
--- a/Engine/Documentation/Source/Engine/Animation/Persona/SkeletonTree/SkeletonTree.INT.udn
+++ b/Engine/Documentation/Source/Engine/Animation/Persona/SkeletonTree/SkeletonTree.INT.udn
@@ -1,8 +1,8 @@
 Availability:Public
-Title:Persona Skeleton Tree
+Title:Skeleton Tree
 Crumbs: %ROOT%, Engine, Engine/Animation, Engine/Animation/Persona
-Description:User guide for the Skeleton Tree in the Persona animation editor.
-version: 4.11
+Description:User guide for the Skeleton Tree inside the Skeleton Editor. 
+version: 4.14
 related: Engine/Content\Types\SkeletalMeshes\Sockets
 related: Engine/Animation/AnimationRetargeting
 
@@ -29,77 +29,44 @@ related: Engine/Animation/AnimationRetargeting
 ![](SkeletonTree.png)  
 
 [EXCERPT:Overview]
-The **Skeleton Tree** panel shows the skeletal hierarchy of the current skeletal mesh. Within it, you can perform the following functions:
+The **Skeleton Tree** panel shows the skeletal hierarchy of the current [Skeleton](Engine/Animation/Skeleton) asset. Within it, you can perform the following actions: 
 
 * Select and view specific bones in the skeletal hierarchy.
 * Create and edit [Sockets](Engine/Content\Types\SkeletalMeshes\Sockets), which are bone-offset attachment points for props and other items.
 * Create preview assets, which are temporarily attached assets; these are good for previewing how a weapon will look on a character, for example.
-* Control settings for animation retargeting (when **Show Advanced Options** is checked).
+* Control settings for [Animation Retargeting](Engine/Animation/AnimHowTo/Retargeting) (when **Show Advanced Options** is checked).
 [/EXCERPT:Overview]
 
 ## Interface Breakdown
 
 ![](SkeletonTree_UIBreakdown.png)
 
-1. Search Line, Blend Profile, Bone and Socket Buttons, Advanced Options
-1. Bone Hierarchy List
-1. Translation Retargeting Controls (Enable **Show Advanced Options** at the top of the Skeleton Tree panel to show)
+1. [Search Line](#searchline)
+1. [Bone Hierarchy List](#bonehierarchylist)
+1. [Blend Profile and Skeleton Options](#blendprofileandskeletonoptions). 
 
 
-## Search Line and Bone and Socket Buttons
+### Search Line
 
 ![](SearchLine.png)
 
-The **Search** line allows for quick and easy searching of a skeletal hierarchy. Type in some part of the name of the Bone, Socket, or Preview  Object. and the list will filter to what you want to see.
+The **Search** line allows for quick filtering of the Skeletal hierarchy by filtering out the hierarchy to only include names based on the entered text. 
 
 ![](SeachFiltering.png)
 
-
-
-### Bone and Socket Buttons
-
-![](button_BoneAndSocketButtons.png)
-
-The **Bone** and **Socket** buttons are used to control visibility over various bone and Socket types. Each button opens a menu allowing you to select the types of Bones or Sockets you want to see.
-
-[REGION:raw]
-![](BonesMenu.png)
-[/REGION]
-
-| Option					| Description														|
-|---------------------------|-------------------------------------------------------------------|
-|**Show All Bones**			| Shows all bones available in the skeleton.						|
-|**Show Mesh Bones**		| Shows only those bones that are exclusive to this Skeletal Mesh.	|
-|**Show Weighted Bones**	| Shows only those bones which have vertices weighted to them.		|
-|**Hide Bones**				| Hides all bones, showing only Sockets and preview assets.			|
-
-[REGION:raw]
-![](SocketsMenu.png)
+[REGION:caption]
+Above, we have filtered the hierarchy to only include Bone Names, Sockets or Preview Objects that include the text **Spine** in them. 
 [/REGION]
 
+### Bone Hierarchy List
 
-| Option | Description |
-|---|---|
-|**Show Active Sockets**	| Shows only active Sockets. When you customize a Socket (**Right-click** and choose Customize Socket), the Socket is copied from the Skeleton onto the mesh, making a special socket that is exclusive to this Skeletal Mesh and not shared by any other Skeletal Meshes using the same Skeleton. When you do that, the new customized Socket - also known as a Mesh Socket - becomes active and the original Socket, which remains a part of the Skeleton asset, is considered inactive. In effect, this will hide away any Sockets on the Skeleton that have a customized Mesh Socket counterpart. |
-|**Show Mesh Sockets**		| Shows only customized Sockets, which are those Sockets that are exclusive to this Skeletal Mesh and are not used by any other Skeletal Mesh that shares a Skeleton asset with this one.  |
-|**Show Skeleton Sockets**	| Shows only those Sockets which are non-customized, meaning those which are part of the Skeleton asset and are not exclusive to this Skeletal Mesh. |
-|**Show All Sockets**		|Shows all Sockets available on the Skeletal Mesh. |
-|**Hide Sockets**			| Hides all Sockets. |
+![](BoneHierarchyList.png)
 
-### Blend Profile
+This window shows all the Bones, Sockets and Preview Objects currently used by the selected Skeleton Asset. 
+Inside this window is where you can add and assign Sockets and Preview Objects as well as adjust Translation Retargeting options. 
+Most of the adjustments made in the Bone Hierarchy list is handled through **Right-click** context menus where the menu options vary based on the object you right-click on (see the context menu sections below).
 
-A Blend Profile is a set of per-bone scales that can be used in transtions and blend lists to tweak the weights of specific bones. The scales are applied to a normal weight for that bone. 
-
-You can create/name your Blend Profile from the drop-down menu then specific the per-bone scales next to each bone. 
-
-![](BlendProfile.png)
-
-
-## Context Menus
-
-The context menu - accessible by **Right-clicking** within the **Skeleton Tree** - has multiple forms depending on what you are clicking: one for bones, one for Sockets, and one for Preview Assets
-
-### Bone Context Menu
+#### Bone Context Menu
 
 When **Right-clicking** on a bone, you will see the following context menu:
 
@@ -121,7 +88,7 @@ When **Right-clicking** on a bone, you will see the following context menu:
 |**Remove Selected Bone from...**| Places the selected bone and its children on the Bones to Remove list for a certain LOD level within the LOD settings. See LOD Settings for more information.|
 |**Add Remove Selected Bone to...**| Removes the selected bone from the Bones to Remove list for the specified LOD level. See LOD Settings for more information.|
 | [REGION:tablesection]Attached Assets Actions[/REGION]||
-|**Add Preview Asset**| Opens a menu of all eligible assets that can be temporarily attached to a bone for previewing purposes. For instance, if you want to see what a certain piece of armor would look like attached a bone, this is the option you would use. Note, however, that this is not a permanent attachment and will only be visible in Persona. |
+|**Add Preview Asset**| Opens a menu of all eligible assets that can be temporarily attached to a bone for previewing purposes. For instance, if you want to see what a certain piece of armor would look like attached a bone, this is the option you would use. Note, however, that this is not a permanent attachment and will only be visible while previewing. |
 
 ### Socket Context Menu
 
@@ -136,10 +103,10 @@ When **Right-clicking** on a [Socket](Engine/Content\Types\SkeletalMeshes\Socket
 |[REGION:tablesection]Selected Socket Actions[/REGION]||
 |**Copy Selected Sockets**| Copies the selected socket to the clipboard for pasting on to other bones.|
 | **Rename Socket** | Allows you to rename a selected Socket.  |
-|**Create Mesh Socket**| Copies the Socket from the Skeleton asset onto the specific Skeletal mesh being used. In this way, a Socket can be edited slighting for characters using the same Skeleton asset, but having different mesh proportions.|
+|**Create Mesh Socket**| Copies the Socket from the Skeleton asset onto the specific Skeletal mesh being used. In this way, a Socket can be edited for characters using the same Skeleton asset, but having different mesh proportions.|
 |**Delete Mesh Socket** | This is only available when **Right-clicking** on a Socket that has been made into a Mesh Socket. This removes the Mesh copy of the socket, leaving only the original Socket on the Skeleton asset.|
 | [REGION:tablesection]Attached Assets Actions[/REGION]||
-|**Add Preview Asset**| Opens a menu of all eligible assets that can be temporarily attached to a bone for previewing purposes. For instance, if you want to see what a certain piece of armor would look like attached a bone, this is the option you would use. Note, however, that this is not a permanent attachment and will only be visible in Persona. |
+|**Add Preview Asset**| Opens a menu of all eligible assets that can be temporarily attached to a bone for previewing purposes. For instance, if you want to see what a certain piece of armor would look like attached a bone, this is the option you would use. Note, however, that this is not a permanent attachment and will only be visible while previewing. |
 |**Remove All Attached Assets**| Removed all attached Preview Assets from this socket.|
 
 
@@ -157,25 +124,44 @@ When **Right-clicking** on a Preview Asset, you will see the following context m
 | [REGION:tablesection]Selected Asset Actions[/REGION]||
 |**Delete**|Removes the selected Preview Asset.|
 | [REGION:tablesection]Attached Asset Actions[/REGION]||
-|**Add Preview Asset**| Opens a menu of all eligible assets that can be temporarily attached to a bone for previewing purposes. For instance, if you want to see what a certain piece of armor would look like attached to a bone, this is the option you would use. Note, however, that this is not a permanent attachment and will only be visible in Persona. |
+|**Add Preview Asset**| Opens a menu of all eligible assets that can be temporarily attached to a bone for previewing purposes. For instance, if you want to see what a certain piece of armor would look like attached to a bone, this is the option you would use. Note, however, that this is not a permanent attachment and will only be visible while previewing. |
 |**Remove All Preview Assets**|Deletes all Preview Assets currently on the mesh.|
 
 
+## Blend Profile and Skeleton Options
+
+![](BlendProfileSection.png)
+
+At the bottom of the Skeleton Tree window are the Blend Profile and Skeleton Options buttons. 
+
+### Blend Profiles
 
-## Sockets
+A Blend Profile is a set of per-bone scales that can be used in transitions and blend lists to tweak the weights of specific bones. 
+The scales are applied to a normal weight for that bone. 
+You can create/name your Blend Profile from the drop-down menu then specify the per-bone scales next to each bone with the Blend Profile Scale setting.
 
-For more information on Sockets, please see the [Skeletal Mesh Sockets documentation](Engine/Content\Types\SkeletalMeshes\Sockets).
+![](BlendProfile.png)
+
+### Skeleton Options
+
+![](SkeletonOptionsWindow.png)
 
+The Skeleton Options window allow you show or hide Skeleton Retargeting Options as well as Bones and Sockets. 
 
-## Translation Retargeting Controls
+#### Translation Retargeting Controls
 
-The right column of buttons in the **Skeleton Tree** controls how you want to handle translation retargeting on your skeletal mesh. The dropdowns allow you to specify what will be used as the basis for retargeting: the skeleton's bone information, the animation rotation data, or the animation as applied to a scaled character.
+With **Show Retargeting Options** enabled from the **Skeleton Options** menu, the Bone Hierarchy window will update with addition options. 
+The right column of buttons in the **Skeleton Tree** controls how you want to handle translation [retargeting](Engine/Animation/AnimationRetargeting) on your skeletal mesh. The dropdowns allow you to specify what will be used as the basis for retargeting: the skeleton's bone information, the animation rotation data, or the animation as applied to a scaled character.
 
-Each dropdown allows you to control how that specific bone will handle retargeting and what will be used as the basis for that retarget. You can learn more about animation retargeting in the [](Engine/Animation/AnimationRetargeting) documentation.
+Each dropdown allows you to control how that specific bone will handle retargeting and what will be used as the basis for that retarget. 
 
 ![](RetargetingControls.png)
 
-Each button contains the following three options. Note that you can change these settings recursively on each child of a selected bone by **Right-clicking** and choosing the recursive commands in the context menu.
+Each button contains the following three options. 
+
+[REGION:note]
+Note that you can change these settings recursively on each child of a selected bone by **Right-clicking** and selecting the recursive commands in the context menu.
+[/REGION]
 
 |Name|Description|
 |---|---|
@@ -183,3 +169,35 @@ Each button contains the following three options. Note that you can change these
 |**Skeleton**|Use translation from skeleton.|
 |**AnimationScaled**|Use translation from animation, scale length by the Skeleton's proportions. |
 |**AnimationRelative**|Use relative translation from animation similar to an additive animation. |
+
+#### Show/Hide Bones and Sockets
+
+The **Bone** and **Socket** section of the **Skeleton Options** menu are used to control visibility over various bone and Socket types. 
+
+[REGION:raw]
+![](BonesMenu.png)
+[/REGION]
+
+| Option					| Description														|
+|---------------------------|-------------------------------------------------------------------|
+|**Show All Bones**			| Shows all bones available in the skeleton.						|
+|**Show Mesh Bones**		| Shows only those bones that are exclusive to this Skeletal Mesh.	|
+|**Show Weighted Bones**	| Shows only those bones which have vertices weighted to them.		|
+|**Hide Bones**				| Hides all bones, showing only Sockets and preview assets.			|
+
+[REGION:raw]
+![](SocketsMenu.png)
+[/REGION]
+
+
+| Option | Description |
+|---|---|
+|**Show Active Sockets**	| Shows only active Sockets. When you customize a Socket (**Right-click** and choose Customize Socket), the Socket is copied from the Skeleton onto the mesh, making a special socket that is exclusive to this Skeletal Mesh and not shared by any other Skeletal Meshes using the same Skeleton. When you do that, the new customized Socket - also known as a Mesh Socket - becomes active and the original Socket, which remains a part of the Skeleton asset, is considered inactive. In effect, this will hide away any Sockets on the Skeleton that has a customized Mesh Socket counterpart. |
+|**Show Mesh Sockets**		| Shows only customized Sockets, which are those Sockets that are exclusive to this Skeletal Mesh and are not used by any other Skeletal Mesh that shares a Skeleton asset with this one.  |
+|**Show Skeleton Sockets**	| Shows only those Sockets which are non-customized, meaning those which are part of the Skeleton asset and are not exclusive to this Skeletal Mesh. |
+|**Show All Sockets**		|Shows all Sockets available on the Skeletal Mesh. |
+|**Hide Sockets**			| Hides all Sockets. |
+
+
+
+
diff --git a/Engine/Documentation/Source/Engine/Animation/Persona/Toolbar/PersonaToolbar.INT.udn b/Engine/Documentation/Source/Engine/Animation/Persona/Toolbar/PersonaToolbar.INT.udn
--- a/Engine/Documentation/Source/Engine/Animation/Persona/Toolbar/PersonaToolbar.INT.udn
+++ b/Engine/Documentation/Source/Engine/Animation/Persona/Toolbar/PersonaToolbar.INT.udn
@@ -1,169 +1,110 @@
 Availability:Public
-Title:Persona Toolbar
+Title:Animation Toolbars
 Crumbs: %ROOT%, Engine, Engine/Animation, Engine/Animation/Persona
-Description:User guide for the Toolbar in the Persona animation editor.
-version: 4.9
-
-[VAR:Topic]
-[OBJECT:Topic]
-	[PARAM:image]
-		![%Engine/Animation/Persona/Toolbar:title%](Engine/Animation/animation_topic.png)
-	[/PARAM]
-	[PARAM:icon]
-		![](%ROOT%/content_icon.png)(convert:false)
-	[/PARAM]
-	[PARAM:title]
-		%Engine/Animation/Persona/Toolbar:title%
-	[/PARAM]
-	[PARAM:description]
-		%Engine/Animation/Persona/Toolbar:description%
-	[/PARAM]
-	[PARAM:path]
-		[RELATIVE:Engine/Animation/Persona/Toolbar]
-	[/PARAM]
-[/OBJECT]
-[/VAR]
-
-[REGION:fullwidth]
-![](PersonaToolbar.png)
-
-[TOC(start:2)]
-
+Description:Overview of the different Toolbars used in each of the Animation Tool Editors. 
+Related: Engine/Animation/Persona/Modes/Skeleton
+Related: Engine/Animation/Persona/Modes/Mesh
+Related: Engine/Animation/Persona/Modes/Animation
+Related: Engine/Animation/Persona/Modes/Graph
+version: 4.14
+type:reference
+topic-image:Engine/Animation/animation_topic.png
+tags:animation
 
 
 [EXCERPT:Overview]
-The **Toolbar** is displayed along the top of the Persona interface window by default. It contains a variety of common functions you will need when editing animation assets.
+Regardless of what Animation Tool you are working in, each Editor will feature a **Toolbar** which can be located at the top, however, the options on the Toolbar will vary based on the type of Editor you are currently working in.
+The Toolbar is comprised of two main areas, the **Asset Toolbar** which contains several options related to the Editor you are working with and the **Editor Toolbar** which allows you to quickly switch between the different Animation Tool Editors. 
 [/EXCERPT:Overview]
 
-[REGION:tip]
-If you do not see the Toolbar, you can show it using the **Window** menu within Persona.
-[/REGION]
-
-## Toolbar Modes
 
-As you may know, Persona has four different _modes_ that it uses to selectively display aspects of its interface based on what the user is currently doing. This also means that the common commands you will use within Persona will change along with the modes. The Toolbar intuitively updates to hide away any commands that would not be useful in the current mode. For instance:
+[TOC(start:2)]
 
-When in **Skeleton**, **Mesh**, or **Animation** modes, the Toolbar will look something like this:
+ 
+## Asset Toolbar
 
+Located along the left side of the Toolbar are the options that are available based on the type of Editor you are working with, refer to the sections below: 
 
-![](Toolbar_01.png)
+### Skeleton Editor Toolbar
 
- However, when in **Graph** mode, the toolbar expands with further features, looking like this:
+![](SkeletonToolbar.png)
 
-![](PersonaToolbar.png)
+The Skeleton Editor Toolbar is shown whenever you open [](Engine/Animation/Skeleton) and contains the following options: 
 
-You may notice that many of the additional buttons are the same as those used in the Blueprint Editor window. 
+|Option | Description |
+| ------- | ------------ |
+| ![](saveIcon.png) | Saves any changes made to the current Skeleton asset.  |
+| ![](findincbIcon.png) | Opens the **Content Browser** and automatically navigates to the current Skeleton asset. |
+| ![](animnotifiesIcon.png) | Opens the [Anim Notifies](Engine/Animation/Persona/AnimationNotifies) window that displays all the custom Animation Notifies for the current Skeleton. |
+| ![](retargetmanagerIcon.png) | Opens the [](Engine/Animation/Persona/BasePoseManager) which allows you to manage your retarget sources, set up rigs and define retarged base poses for use with Animation Retargeting. |
+| ![](importmeshIcon.png) | Begins the [](Engine/Content/FBX/SkeletalMeshes) and allows you to import a new Skeletal Mesh for the current Skeleton.  |
+| ![](makestaticmeshIcon.png) | Allows you to create a **Static Mesh** out of the current preview pose. |
 
+### Skeletal Mesh Editor Toolbar
 
-## Toolbar Buttons
+![](MeshToolbar.png)
 
-Persona toolbar buttons provide easy access to switching between asset modes, such as Skeleton, Mesh, and Animations. It also contains common commands needed when editing Animation Blueprints.
+The Skeletal Mesh Editor Toolbar is shown whenever you open a [Skeletal Mesh](Engine/Content/Types/SkeletalMeshes) asset and contains the following options: 
 
+|Option | Description |
+| ------- | ------------ |
+| ![](saveIcon.png) | Saves any changes made to the current Skeletal Mesh asset.  |
+| ![](findincbIcon.png) | Opens the **Content Browser** and automatically navigates to the current Skeletal Mesh asset. |
+| ![](reimportmeshIcon.png) | Allows you to reimport the current Skeletal Mesh and starts the FBX Import Process.  |
+| ![](makestaticmeshIcon.png) | Allows you to create a **Static Mesh** out of the current preview pose. |
 
-| Button																						| Description														|
-| --------------------------------------------------------------------------------------------- | ----------------------------------------------------------------- | 
-| ![Compile Successful Button](Engine/Blueprints/Editor/UIComponents/Toolbar/bp_compile_successful.png)		| [INCLUDE:Engine/Blueprints\Editor/UIComponents\Toolbar#compilesuccessful]	|
-| ![Compile Needed Button](Engine/Blueprints/Editor/UIComponents/Toolbar/bp_needs_compile.png)				| [INCLUDE:Engine/Blueprints\Editor/UIComponents\Toolbar#compileneeded]		|
-| ![Compile Warning Button](Engine/Blueprints/Editor/UIComponents/Toolbar/compile_warning.png)				| [INCLUDE:Engine/Blueprints\Editor/UIComponents\Toolbar#compilewarning]		|
-| ![Compile Failed Button](Engine/Blueprints/Editor/UIComponents/Toolbar/compile_failed.png)				| [INCLUDE:Engine/Blueprints\Editor/UIComponents\Toolbar#compilefailed]		|
-| ![Save Button](Engine/Blueprints/Editor/UIComponents/Toolbar/bp_save.png)									| [INCLUDE:Engine/Blueprints\Editor/UIComponents\Toolbar#save]					|
-| ![Find in Content Browser Button](Engine/Blueprints/Editor/UIComponents/Toolbar/bp_find_in_cb.png)		| [INCLUDE:Engine/Blueprints\Editor/UIComponents\Toolbar#findincb]				| 
-| ![Blueprint Properties Button](Engine/Blueprints/Editor/UIComponents/Toolbar/bp_properties.png)			| [INCLUDE:Engine/Blueprints\Editor/UIComponents\Toolbar#blueprintprops]		|
-| ![Blueprint Properties Button](Engine/Blueprints/Editor/UIComponents/Toolbar/defaults.png)				| Shows the Class Defaults Panel in the Details Tab |
-| ![Search Button](Engine/Blueprints/Editor/UIComponents/Toolbar/bp_search.png)								| [INCLUDE:Engine/Blueprints\Editor/UIComponents\Toolbar#search]				|
-| ![Simulate Button](Engine/Blueprints/Editor/UIComponents/Toolbar/bp_simulate.png)							| [INCLUDE:Engine/UI/LevelEditor/Toolbar#simulate]			|
-| ![Play In Editor Button](Engine/Blueprints/Editor/UIComponents/Toolbar/bp_play_in.png)					| [INCLUDE:Engine/UI/LevelEditor/Toolbar#playin]				|
-| ![Pause Button](Engine/Blueprints/Editor/UIComponents/Toolbar/bp_pause.png)								| [INCLUDE:Engine/Blueprints\Editor/UIComponents\Toolbar#pause]				|
-| ![Resume Button](Engine/Blueprints/Editor/UIComponents/Toolbar/bp_resume.png)								| [INCLUDE:Engine/Blueprints\Editor/UIComponents\Toolbar#resume]				|
-| ![Frame Skip Button](Engine/Blueprints/Editor/UIComponents/Toolbar/bp_frameskip.png)						| [INCLUDE:Engine/Blueprints\Editor/UIComponents\Toolbar#frameskip]			|
-| ![Stop Button](Engine/Blueprints/Editor/UIComponents/Toolbar/bp_stop.png)									| [INCLUDE:Engine/Blueprints\Editor/UIComponents\Toolbar#stop]					|
-| ![Possess Button](Engine/Blueprints/Editor/UIComponents/Toolbar/bp_possess.png)							| [INCLUDE:Engine/Blueprints\Editor/UIComponents\Toolbar#possess]				|
-| ![Eject Button](Engine/Blueprints/Editor/UIComponents/Toolbar/bp_eject.png)								| [INCLUDE:Engine/Blueprints\Editor/UIComponents\Toolbar#eject]				|
-| ![Find Node Button](Engine/Blueprints/Editor/UIComponents/Toolbar/bp_find_node.png)						| [INCLUDE:Engine/Blueprints\Editor/UIComponents\Toolbar#findnode]				|
-| ![Step Button](Engine/Blueprints/Editor/UIComponents/Toolbar/bp_step.png)									| [INCLUDE:Engine/Blueprints\Editor/UIComponents\Toolbar#step]					|
-| ![Debug Dropdown](Engine/Blueprints/Editor/UIComponents/Toolbar/bp_debug_dropdown.png)					| [INCLUDE:Engine/Blueprints\Editor/UIComponents\Toolbar#debugobject]			|
-| ![](button_Preview.png)																					| Displays a preview of the currently selected animation asset. |
-| ![](button_RefPose.png)																					| Displays the skeletal mesh in it's reference pose. |
-| ![](button_Import.png)																					| Brings up the Import dialogue. |
-| ![](button_ImportMesh.png)																				| Brings up the Import dialogue. |
-| ![](button_Reimport.png)																					| Brings up the Browse dialogue to select an asset to reimport. |
-| ![](button_ReimportMesh.png)																				| Brings up the Browse dialogue to select an asset to reimport. |
-| ![](button_Export.png)																					| Brings up the Export dialogue. |
-| ![](button_Record.png)																					| Allows for the recording of an animation to be saved at another location. |
-| ![](button_CreateAsset.png)																				| Brings a menu with multiple options for creating a new asset. |
-| ![](button_Compression.png)																				| Brings multiple options of compression for the currently selected animation. |
-| ![](button_AnimNotifies.png)																				| Brings up the Anim Notifies window, which displays all Anim Notifies. |
-| ![](button_Key.png)																						| Adds a key to the Additive Layer Track of the currently selected bone for the currently selected animation. |
-| ![](button_Apply.png)																						| Applies added keys to the currently selected animation. |
+### Animation Editor Toolbar
 
-## Mode Switching Tabs
+![](AnimationToolbar.png)
 
-![](ModeTabPersona.png)
+The Animation Editor Toolbar is shown whenever you open any animation assets (Animation Sequences, Anim Montages, Blend Spaces, etc.) and contains the following options: 
 
-The Persona Toolbar also contains buttons to switch between each of its four modes. These buttons have some functionality in and of themselves that can be useful when trying to track down a particular animation asset. You change Persona modes by simply clicking on the button for the desired mode.
+|Option | Description |
+| ------- | ------------ |
+| ![](saveIcon.png) | Saves any changes made to the current animation asset.  |
+| ![](findincbIcon.png) | Opens the **Content Browser** and automatically navigates to the current animation asset. |
+| ![](createassetIcon.png) | Allows you to create new animation assets based on the current Skeleton asset. |
+| ![](compressionIcon.png) | Allows you to apply different [Compression](Engine/Animation/Sequences) settings to the animation (available for Animation Sequences). |
+| ![](keyIcon.png) | Allows you to add the current Bone Transform of a selected Bone to an [Additive Layer Track.](Engine/Animation/AnimHowTo/LayerEditing)  |
+| ![](applyIcon.png) | Applies [Additive Layer Tracks](Engine/Animation/AnimHowTo/LayerEditing) to Runtime Animation Data. |
+| ![](makestaticmeshIcon.png) | Allows you to create a **Static Mesh** out of the current preview pose. |
 
-| Button | Description |
-|-------|------|
-| ![](button_SkeletonMode.png)| [INCLUDE:#skeletonbutton] |
-| ![](button_MeshMode.png) | [INCLUDE:#meshbutton] |
-| ![](button_AnimationMode.png) | [INCLUDE:#animationbutton] |
-| ![](button_GraphMode.png) | [INCLUDE:#graphbutton] |
+### Animation Blueprint Editor Toolbar
 
+![](AnimBPToolbar.png)
 
+The Animation Blueprint Editor Toolbar is shown whenever you open any [](Engine/Animation/AnimBlueprints) and allows you to Save, Find the Animation Blueprint in the Content Browser or create a Static Mesh from the current preview pose. 
+The rest of the options utilize the same functionality found in the normal **Blueprint Editor Toolbar**. 
 
+[REGION:note]
+Please see the [Blueprint Editor Toolbar](Engine/Blueprints/Editor/UIComponents/Toolbar) page for more information on each option.
+[/REGION]
 
 
-<!--
-[EXCERPT:SkeletonButton]
-This switches Persona to [Skeleton mode](Engine/Animation/Persona#PersonaModesandUserInterface).
+## Editor Toolbar
 
-[REGION:simpletable]
-| Element | Description |
-|---|---|
-|![](button_SkeletonField.png)| Shows the name of the Skeleton asset currently used by Persona. |
-%search%
-[/REGION]
-[/EXCERPT:SkeletonButton]
+Located along the right side of each editor toolbar is the **Editor Toolbar** which allows you to switch between each of the Animation Tool Editors. 
 
+The currently open editor will be highlighted: 
+  
+ ![](skeletoneditoropen.png)
+ 
+ When clicking one of the other Editors, the new tab will open and the new editor will be highlighted: 
 
-[EXCERPT:MeshButton]
-This switches Persona to [Mesh mode](Engine/Animation/Persona#PersonaModesandUserInterface).
+ ![](mesheditoropen.png)
 
-[REGION:simpletable]
-| Element | Description |
-|---|---|
-|![](button_MeshField.png)| Shows the name of the Skeletal Mesh asset currently being edited by Persona. |
-%pickasset%
-%search%
-[/REGION]
-[/EXCERPT:MeshButton]
+ The **Animation Editor** includes an animation selection drop-down button next to the icon to quickly change the preview animation: 
 
-[EXCERPT:AnimationButton]
-This switches Persona to [Animation mode](Engine/Animation/Persona#PersonaModesandUserInterface).
+![](animationdropdownmenu.png)
 
-[REGION:simpletable]
-| Element | Description |
-|---|---|
-|![](button_AnimationField.png)| Shows the name of the animation asset currently being edited within Persona. |
-%pickasset%
-%search%
-[/REGION]
-[/EXCERPT:AnimationButton]
+An ** * ** symbol in the lower left corner of the icon indicates that the asset has changes that have not been saved. 
 
-[EXCERPT:GraphButton]
-This switches Persona to [Graph mode](Engine/Animation/Persona#PersonaModesandUserInterface).
+![](unsavedasset.png)
 
-[REGION:simpletable]
-| Element | Description |
-|---|---|
-|![](button_GraphField.png)| Shows the name of the Animation Blueprint asset currently being edited by Persona. |
-%search%
-[/REGION]
-[/EXCERPT:GraphButton]
+Based on the assets in your project, some of the editors may not appear on the toolbar. 
 
-[VAR:PickAsset]|![](button_AssetSelection.png)|Allows you to quickly select any compatible assets in the current project. |[/VAR]
-[VAR:Search]|![](button_Search.png)| Clicking this brings up the **Content Browser** and shows and highlights the asset currently shown in the respective field. |[/VAR]
+For example, if you have imported a Skeletal Mesh and created a Skeleton, but have no animations or Animation Blueprints, your toolbar may look like below. 
 
--->
+![](limitedtoolbar.png)
 
-[/REGION]
+Creating an Animation Blueprint with the associated Skeleton and importing animations linked to the Skeleton will allow you to access the other editors for that asset. 
