Copying //UE4/Dev-Documentation to Samples-Main (//UE4/Samples-Main)

diff --git a/Engine/Documentation/Source/Engine/Animation/Persona/Viewport/PersonaViewport.INT.udn b/Engine/Documentation/Source/Engine/Animation/Persona/Viewport/PersonaViewport.INT.udn
--- a/Engine/Documentation/Source/Engine/Animation/Persona/Viewport/PersonaViewport.INT.udn
+++ b/Engine/Documentation/Source/Engine/Animation/Persona/Viewport/PersonaViewport.INT.udn
@@ -1,91 +1,68 @@
 Availability:Public
-Title:Persona Viewport
+Title:Animation Viewports
 Crumbs: %ROOT%, Engine, Engine/Animation, Engine/Animation/Persona
-Description:User guide for the Viewport preview in the Persona animation editor.
-version: 4.9
-
-[VAR:Topic]
-[OBJECT:Topic]
-	[PARAM:image]
-		![%Engine/Animation/Persona/Viewport:title%](Engine/Animation/animation_topic.png)
-	[/PARAM]
-	[PARAM:icon]
-		![](%ROOT%/content_icon.png)(convert:false)
-	[/PARAM]
-	[PARAM:title]
-		%Engine/Animation/Persona/Viewport:title%
-	[/PARAM]
-	[PARAM:description]
-		%Engine/Animation/Persona/Viewport:description%
-	[/PARAM]
-	[PARAM:path]
-		[RELATIVE:Engine/Animation/Persona/Viewport]
-	[/PARAM]
-[/OBJECT]
-[/VAR]
-
-[TOC (start:2 end:3)]
+Description:Overview of the Viewports displayed in each of the Animation Tools. 
+Related: Engine\Animation\Persona\Modes\Skeleton
+Related: Engine\Animation\Persona\Modes\Mesh
+Related: Engine\Animation\Persona\Modes\Animation
+Related: Engine\Animation\Persona\Modes\Graph
+version: 4.14
+type:reference
+topic-image:viewportsmall.png
+tags:animation
 
 
+<!--
+[TOC (start:2 end:3)]
+-->
 [EXCERPT:Overview]
-The **Viewport** panel allows Animation Sequences, Blendspaces, and other aspects of the Animation Blueprint to be previewed on a specific Skeletal Mesh. The panel consists of a 3D perspective viewport, a movement toolbar, a view options toolbar, and a timeline/control bar for controlling playback of the preview.
+Each of the Animation Tools contains a **Viewport** window which allows you to preview playback of your animation assets. 
 [/EXCERPT:Overview]
+![](viewport01.png)
 
-![Viewport Panel with Callouts](ViewportPanel.png)(w:720)
-
-1. Movement Toolbar
-1. View Options Toolbar 
-1. Viewport 
-1. Timeline and Control Bar
-
-## Movement Toolbar
-[EXCERPT:movement]
-The **Movement Toolbar** contains a collection of tools for moving, rotating, or scaling bones as well as how they are moved, rotated, or scaled with Grid Snapping. This toolbar also has an option for setting the relative coordinate system for movement or rotation as well as an option to set the camera movement speed when moving around in the viewport. 
-
-| Option | Description |
-| ------ | ----------- |
-| ![](Persona_Translate.png)| Translates skeletal bones. |
-| ![](Persona_Rotate.png)| Rotates skeletal bones. |
-| ![](Persona_Scale.png) | The last option scales skeletal bones (if applicable, grayed out by default).  |
-| ![](button_Persona_Local.png) ![](Persona_World.png)| Sets the relative coordinate system for moving and rotating the mesh (**World** or **Local**).|
-| ![](Persona_GridSnap.png) | Enables or disables snapping to the grid when dragging objects around and the Position Grid Snap value. |
-| ![](Persona_RotationSnap.png)| Enables or disables snapping to the grid when rotating objects around and the Rotation Grid Snap value.  |
-| ![](Persona_ScaleSnapping.png) | Enables or disables snapping to the grid when scaling objects and the Scale Grid Snap value.  |
-| ![](Persona_CameraMoveSpeed.png) | Adjust the speed of the camera while moving around in the viewport.|
+The Viewport consists of a **View Options** toolbar (1), a **Movement** toolbar (2), and a 3D perspective viewport of your currently selected Skeletal Mesh asset (3). 
 
-[/EXCERPT:movement]
 
 ## View Options Toolbar
 
 [EXCERPT:viewoptions]
-The **View Options Toolbar** contains a series of buttons that hold options for what is displayed in the viewport.
+The **View Options Toolbar** contains a series of buttons that hold options for what is displayed (and how it is displayed) in the viewport.
 
-### Preview Mode
+### Scene Setup
 
-The **Preview** Mode Button ![](Persona_PreviewMode1.png) provides the following options: 
+The **Scene Setup** button provides the following options: 
+
+![](SceneSetup.png)
 
 | Option | Description |
 | ------ | ----------- |
-| **Auto** | Previews whatever is active in the viewport. |
-| **Lock** | Locks the viewport to what is currently being previewed. This allows you to preview one item while viewing or editing another. |
-| **Camera Follow** | Follow the bounds of the mesh. |
-| **In-game Bound** | Uses in-game bounds on the preview mesh. |
+| **Preview Scene Settings** | Allows you to alter lighting and post processing settings for the viewport. |
+| **Scene Setup** | Allows you to change the current preview animation and animation mode. Also allows you to change the current preview mesh or add meshes. |
+| **Turn Table** | Allows you to set up auto-rotation of the preview mesh and set rotation speeds. |
+| **Camera Follow** | Allows you to determine if the camera should or should not follow the bounds of the current mesh. |
 
 ### Perspective
 
-The **Perspective** Button ![](Persona_Perspective.png) provides the following options: 
+The **Perspective** button provides the following options: 
+
+![](perspectivemode.png)
 
 | Option | Description |
 | ------ | ----------- |
 | **Perspective** | Switches the viewport to perspective view (this is the default setting). |
-| **Top** | Switches the viewport to top view. |
-| **Side** | Switches the viewport to side view. |
-| **Front** | Switches the viewport to front view. |
+| **Top** | Switches the viewport to top orthographic view. |
+| **Bottom** | Switches the viewport to bottom orthographic view. |
+| **Left** | Switches the viewport to left orthographic view. |
+| **Right** | Switches the viewport to right orthographic view. |
+| **Front** | Switches the viewport to front orthographic view.  |
+| **Back** | Switches the viewport to back orthographic view.  |
 
 
 ### View Mode
 
-The **View** Mode Button ![](Persona_LitOptions.png) provides the following options: 
+The **View** Mode Button provides the following options:
+
+![](viewmodebutton.png) 
 
 | Option | Description |
 | ------ | ----------- |
@@ -94,104 +71,55 @@ The **View** Mode Button ![](Persona_LitOptions.png) provides the following opti
 | **Wireframe** | Renders the scene in brush wireframe. |
 | **Detail Lighting** | Renders the scene with detail lighting only. |
 | **Lighting Only** | Renders the scene with lights only, no textures. |
-| **Light Complexity** | Renders the scene with light complexity visualization. |
-| **Shader Complexity** | Renders the scene with shader complexity visualization. |
-| **Stationary Light Overlap** | Visualizes overlap of stationary lights. |
-| **Lightmap Density** | Renders the scene with lightmap density visualization. |
 | **Reflections** | Renders the scene with reflections only. |
+| **Optimization View Modes** | Allows you to select from different optimization visualizers. |
 | **Exposure** | Use automatic exposure or set a fixed exposure value.  |
 
-
 ### Show
 
-The **Show** Button ![](Persona_Show.png) provides the following options: 
+The **Show** button provides the following options: 
+
+![](showmenu.png)
 
 | Option | Description |
 | ------ | ----------- |
 | **Field of View** | Adjust the viewports field of view. |
 | **Mute Audio** | Mutes audio from the preview. |
-| **Reference Pose** | Shows the reference pose on the preview mesh. |
-| **Bound** | Shows bounds on the preview mesh. |
+| **Process Root Motion** | Moves the preview based on the animation's root motion. |
+| **Retarget Base Pose** | Shows the Retargeted Base Pose on the preview mesh.  |
+| **Bound** | Shows the Bounds volume on the preview mesh.  |
+| **In-game Bound** | Whether or not to use the in-game bounds on the preview mesh.  |
 | **Mesh** | Shows the preview mesh. |
+| **MorphTargets** | Whether or not to display any applied Morph Targets on the preview mesh. |
 | **Uncompressed Animation** | Displays the skeleton with uncompressed animation data. |
 | **NonRetargeted Animation** | Displays the skeleton with non-retargeted animation data.|
 | **Additive Base** | Displays the skeleton in an additive base pose. |
+| **Source Animation** | Display the Skeleton with uncompressed animation data.  |
+| **Baked Animation** | Display the Skeleton with non-retargeted animation data. |
+| **Bones** | The method in which to show or hide the bones from the preview mesh. |
 | **Sockets** | Display socket hitpoints. |
-| **Bones** | Display skeleton in viewport. |
 | **Bone Names** | Display Bone Names of the skeleton in the viewport. |
-| **View Selected Bone Weight** | Displays the weights of the selected bones in the viewport. |
+| **Mesh Overlay** | Used to display Bone Weighting or Morph Target Vertices on the preview mesh. |
 | **Display Info** | Displays information about the mesh in the viewport. |
-
-**Scene Setup**
-
-| Option | Description |
-| ------ | ----------- |
-| **Floor** | Toggle the display of the floor.  |
-| **Sky** | Toggle the display of the sky. |
-| **Floor Height Offset** | Height offset for the floor mesh (stored per-mesh). |
-| **Grid** | Toggle the display of the grid. |
-| **Highlight Origin** | Highlight the origin lines on the grid. |
-| **Background Brightness** | Change background color and assign a brightness value. |
-
-**Advanced**
-
-| Option | Description |
-| ------ | ----------- |
-| **UV Visualization** | Toggles the display of the mesh's UVs for the specified channel. |
-| **Normals** | Toggles the display of vertex normals in the Preview Pane. |
-| **Tangents** | Toggles the display of vertex tangents in the Preview Pane. |
-| **Binormals** | Toggles the display of vertex binormals (orthogonal vector to normal and tangent) in the Preview Pane. |
-| **Hierarchy Local Axes** | Show or hide local bone axis (All, Selected Hierarchy or None). |
-
-
-<!--
-[EXCERPT:ShowBones]
-[REGION:imagetable]
-| ![](viewport_wire.png)(h:256 convert:false) | ![](viewport_wire_show_bones.png)(h:256) |
-| ------ | ------ |
-| Wireframe Without Bones Displayed | Wireframe With Bones Displayed |
-[/REGION]
-[/EXCERPT:ShowBones]
-
-[EXCERPT:Lit]
-![](viewport_lit.png)(h:256)
-[/EXCERPT:Lit]
-
-[EXCERPT:Unlit]
-![](viewport_unlit.png)(h:256)
-[/EXCERPT:Unlit]
-
-[EXCERPT:Wireframe]
-![](viewport_wire.png)(h:256)
-[/EXCERPT:Wireframe]
--->
+| **Scene Setup** | Allows you to define settings such as auto-aligning the floor to the mesh, toggling the grid or the floor height offset.  |
+| **Advanced** | Allows for UV Visualization, CPU Skinning, Vertex Normal Visualization or how the bone axes are displayed.   |
 
 ### Preview LODs
 
-The **Preview LODs** Button ![](Persona_LODAuto.png) provides the following options: 
+The **Preview LODs** button provides the following options: 
+
+![](lodmenu.png)
 
 | Option | Description |
 | ------ | ----------- |
 | **LOD Auto** | Automatically selects the Level Of Detail (LOD) (this is the default setting). |
 | **LOD 0** | Selects LOD 0 as the LOD. |
-<!--| **Level of Detail Settings** | Opens a dialog box in where you can adjust LOD Settings (pictured below). |
-
-![](LODSettingsMenu.png)
-
-The LOD Settings menu contains two sections, the first section **Current LODs** displays information about the currently applied LODs or allows for importing of FBX assets for the LOD. 
-
-The second section **Generate LOD Chain** allows you to generate an LOD (or chain generate multiple LODs) by adjusting settings pertaining to the LOD. 
-
-When an LOD is applied, either through the **Import LOD** option or through the **Generate LOD** option, they will be added to the **Current LODs** section as pictured below. 
-
-![](LODAdded.png)
-
-_The added LOD(s) can either be removed or an FBX asset can be imported for the LOD._
--->
 
 ### Playback Speed
 
-The **Playback Speed** Button ![](Persona_PlaybackSpeed.png) provides the following options: 
+The **Playback Speed** button provides the following options: 
+
+![](playbackspeed.png)
 
 | Option | Description |
 | ------ | ----------- |
@@ -204,28 +132,33 @@ The **Playback Speed** Button ![](Persona_PlaybackSpeed.png) provides the follow
 | **x10.0** | Sets the playback speed to 10.0. |
 [/EXCERPT:viewoptions]
 
-### Turn Table Speed
 
-The **Turn Table Speed** button ![](Persona_TurnTableSpeed.png) provides the following options:
+## Movement Toolbar
+
+[EXCERPT:movement]
+The **Movement Toolbar** contains a collection of tools for moving, rotating, or scaling bones as well as how they are moved, rotated, or scaled with Grid Snapping. This toolbar also has an option for setting the relative coordinate system for movement or rotation as well as an option to set the camera movement speed when moving around in the viewport. 
+
 | Option | Description |
 | ------ | ----------- |
-| **Play** | Plays the Turn Table (Spins the skeletal mesh) at the selected speed |
-| **Pause** | Pauses the Turn Table at the current rotation |
-| **Stop** | Stops the Turn Table and resets the rotation to the original |
-| **x0.1** | Sets the Turn Table speed to 0.1.|
-| **x0.25** | Sets the Turn Table speed to 0.25. |
-| **x0.5** | Sets the Turn Table speed to 0.5. |
-| **x1.0** | Sets the Turn Table speed to 1.0 (this is the default setting). |
-| **x2.0** | Sets the Turn Table speed to 2.0. |
-| **x5.0** | Sets the Turn Table speed to 5.0. |
-| **x10.0** | Sets the Turn Table speed to 10.0. |
+| ![](Persona_Translate.png)| Translates skeletal bones. |
+| ![](Persona_Rotate.png)| Rotates skeletal bones. |
+| ![](Persona_Scale.png) | The last option scales skeletal bones (if applicable).  |
+| ![](Persona_World.png) ![](button_Persona_Local.png)| Sets the relative coordinate system for moving and rotating the mesh (**World** or **Local**).|
+| ![](Persona_GridSnap.png) | Enables or disables snapping to the grid when dragging objects around and the Position Grid Snap value. |
+| ![](Persona_RotationSnap.png)| Enables or disables snapping to the grid when rotating objects around and the Rotation Grid Snap value.  |
+| ![](Persona_ScaleSnapping.png) | Enables or disables snapping to the grid when scaling objects and the Scale Grid Snap value.  |
+| ![](Persona_CameraMoveSpeed.png) | Adjust the speed of the camera while moving around in the viewport.|
+[/EXCERPT:movement]
+
+
 
 ## Viewport
 
-The **Viewport** displays a preview of the Skeletal Mesh, any options set in the Toolbar and menu bar, and any selected animation.
+The **Viewport** displays a preview of the Skeletal Mesh and any options set in the Viewport Options Toolbar. 
 
 ![Viewport Panel](ViewportPanel_Viewport.png)
 
+<!-- No longer contained in the viewport, leaving for now if it is included on another page. 
 ## Timeline and Control Bar
 
 [EXCERPT:PreviewControls]
@@ -248,18 +181,7 @@ When opening a Blend Space or Aim Offset, the animation will play automatically
 
 [/EXCERPT:PreviewControls]
 
+-->
 
 
-<!---- I don't see this in the viewport or anywhere else. 
-## Anim Preview Setup Panel
-
-![Persona: Anim Preview Setup Panel](CE_anim_preview_setup_panel.png)
-
-The **Anim Preview Setup** panel is used to configure the preview viewport. It displays a list of all of the 
-Skeletal Meshes that use the Skeleton the VimBlueprint being edited was created from. To set a mesh as the 
-preview mesh used in the preview viewport, select the mesh in the list and click the **Use as preview mesh**. 
-To remove the current mesh being used in the preview viewport, click the **Clear preview** button. The list of 
-Skeletal Meshes can be filtered using the text box above the list.
-
------------>
 
diff --git a/Engine/Documentation/Source/Engine/Animation/RetargetingDifferentSkeletons/RetargetingDifferentSkeletons.INT.udn b/Engine/Documentation/Source/Engine/Animation/RetargetingDifferentSkeletons/RetargetingDifferentSkeletons.INT.udn
--- a/Engine/Documentation/Source/Engine/Animation/RetargetingDifferentSkeletons/RetargetingDifferentSkeletons.INT.udn
+++ b/Engine/Documentation/Source/Engine/Animation/RetargetingDifferentSkeletons/RetargetingDifferentSkeletons.INT.udn
@@ -1,4 +1,4 @@
-Availability:Public
+Availability:Docs
 Title: Animation Retargeting (Different Skeletons)
 Crumbs: %ROOT%, Engine, Engine/Animation
 Description:Retargeting animations for use between different skeletons. 
diff --git a/Engine/Documentation/Source/Engine/Animation/RootMotion/RootMotion.INT.udn b/Engine/Documentation/Source/Engine/Animation/RootMotion/RootMotion.INT.udn
--- a/Engine/Documentation/Source/Engine/Animation/RootMotion/RootMotion.INT.udn
+++ b/Engine/Documentation/Source/Engine/Animation/RootMotion/RootMotion.INT.udn
@@ -47,7 +47,7 @@ In order to properly use Root Motion, it is important to note that the Root Bone
 [/REGION]
 
 
-The Root Motion of an animation can be visualized in Persona. Simply open up any animation in which the root bone is moving, and in the viewport choose **Show > Bones**. If you have not selected **Enable Root Motion** in the animation properties, you will notice a red line that appears as the character's root moves. This illustrates the Root Motion of the animation.
+The Root Motion of an animation can be visualized during playback. Simply open up any animation in which the root bone is moving, and in the viewport choose **Show > Bones**. If you have not selected **Enable Root Motion** in the animation properties, you will notice a red line that appears as the character's root moves. This illustrates the Root Motion of the animation.
 
 [REGION:fullwidth]
 ![](RootMotionBone.png)
@@ -64,7 +64,7 @@ So, what does this mean? Now that the system is taking the Root Motion of the ch
 
 ## Enabling Root Motion
 
-Root Motion can be enabled inside the **Anim Asset Details** panel of **Persona** for your [Animation Sequence](Engine\Animation\Sequences). 
+Root Motion can be enabled inside the **Asset Details** panel in the **Animation Editor** for your [Animation Sequence](Engine\Animation\Sequences). 
 
 ![](RootMotion.png)
 
@@ -74,14 +74,10 @@ For networked games however, Root Motion still requires the use of an Animation
 
 [/REGION]
 
-While defining if Root Motion is enabled in done inside the Animation Sequence, you will still need to determine how that sequence is handled inside of an [Animation Blueprint]( Engine/Animation/AnimBlueprints). There are several different ways in which you can handle Root Motion within Animation Blueprints. You can define this through the **Root Motion Mode** drop down menu inside the Anim Blueprint Editor inside Persona. 
+While defining if Root Motion is enabled in done inside the Animation Sequence, you will still need to determine how that sequence is handled inside of an [Animation Blueprint]( Engine/Animation/AnimBlueprints). There are several different ways in which you can handle Root Motion within Animation Blueprints. You can define this through the **Root Motion Mode** drop down in the **Details** panel of the **Anim Blueprint Editor**. 
 
 ![](RootMotionMode.png)
 
-[REGION:note]
-Make sure that you are setting this inside the **Edit Defaults** tab and not the **Edit Preview**.
-[/REGION]
-
 The options are as follows: 
 
 | **Property** | **Description** |
diff --git a/Engine/Documentation/Source/Engine/Animation/Sequences/AnimationSequences.INT.udn b/Engine/Documentation/Source/Engine/Animation/Sequences/AnimationSequences.INT.udn
--- a/Engine/Documentation/Source/Engine/Animation/Sequences/AnimationSequences.INT.udn
+++ b/Engine/Documentation/Source/Engine/Animation/Sequences/AnimationSequences.INT.udn
@@ -8,7 +8,7 @@ Related: Engine/Animation/PhysicallyDrivenAnimation
 Related: Engine/Animation/StateMachines
 Related: Gameplay/AI
 Related: Engine/Content/Tools/MayaRiggingTool
-version: 4.9
+version: 4.14
 
 [VAR:Topic]
 [OBJECT:Topic]
@@ -48,7 +48,7 @@ Qo4hH0h9ZZA
 [/PARAMLITERAL]
 [/OBJECT]
 
-Each Animation Sequence asset targets a specific Skeleton and can only be played on that Skeleton. This means, in order to share animations between multiple Skeletal Meshes, each of the meshes must use the same Skeleton asset.
+Each Animation Sequence asset targets a specific [Skeleton](Engine\Animation\Skeleton) and can only be played on that Skeleton. This means, in order to share animations between multiple Skeletal Meshes, each of the meshes must use the same Skeleton asset.
 
 [REGION:tip]
 An example of a Skeleton using an Animation Sequence can also be found on the [Animation Content Examples](Resources\ContentExamples\Animation) page under section 1.1. 
@@ -58,32 +58,38 @@ An example of a Skeleton using an Animation Sequence can also be found on the [A
 
 [INCLUDE:Engine/Animation/Sequences/Editor#overview]
 
+[REGION:note]
 See [](Engine/Animation/Sequences/Editor) for more details.
+[/REGION]
 
 ## Playback
 
 [INCLUDE:Engine/Animation/Sequences/UserGuide#Intro]
 
+[REGION:note]
 See the [](Engine/Animation/Sequences/UserGuide) for instructions on using Animation Sequences within Animation Blueprints.
+[/REGION]
 
 ## Notifies
 
 [INCLUDE:Engine/Animation/Sequences/Notifies#Intro]
 
+[REGION:note]
 See [](Engine/Animation/Sequences/Notifies) for descriptions and usage information for the various types of notifications available.
-
-Creation and editing of notifies is performed using the [Animation Sequence Editor](Engine/Animation/Sequences/Editor) in Persona.
+[/REGION]
 
 ## Curves
 
 [INCLUDE:Engine/Animation/Sequences/Curves#Intro]
 
-See [](Engine/Animation/Sequences/Curves) for descriptions and usage information on how to set up and apply curves in Persona. 
+[REGION:note]
+See [](Engine/Animation/Sequences/Curves) for descriptions and usage information on how to set up and apply curves in the Animation Editor. 
+[/REGION]
 
 
 ## Compression
 
-Animation Sequence assets when opened up in **Persona** can have compression applied to them. 
+Animation Sequence assets when opened up in the **Animation Editor** can have compression applied to them. 
 
 To do this, from the **Menu Bar** select **Asset** then select **Apply Compression** (as seen below). 
 
@@ -106,12 +112,13 @@ Each option provides a different compression method as outlined in the table bel
 | **Remove Trivial Keys** | Removes trivial frames of tracks when position or orientation is constant over the entire animation from the raw animation data. |
 
 
-## Anim Asset Details
+## Asset Details
 
-When an Animation Sequence asset is opened in **Persona**, there are several properties that can be adjusted inside the **Anim Asset Details** panel which can determine how the asset is handled. These properties include things such as tweaking playback speed, enabling/disabling Root Motion, assigning any Meta Data to the asset and others which are outlined in the tables below. 
+When an Animation Sequence asset is opened in the **Animation Editor**, there are several properties that can be adjusted inside the **Asset Details** panel which can determine how the asset is handled. These properties include things such as tweaking playback speed, enabling/disabling Root Motion, assigning any Meta Data to the asset and others which are outlined in the tables below. 
 
 ![](AnimAssetDetails.png)
 
+<!-- This was moved with the 4.14 Animation Edtior Revamp
 ### Linear Key Removal
 
 |---|---|
@@ -124,20 +131,19 @@ When an Animation Sequence asset is opened in **Persona**, there are several pro
 |**Parent Key Scale**| A scale value which increases the likelihood that a bone will retain a key if its parent also had a key at the same time position. Higher values can remove shaking artifacts from the animation at the cost of compression. |
 |**Retarget**| If true, as the animation is compressed, adjust animated nodes to compensate for compression error. If false, do not adjust animated nodes. |
 |**Actually Filter Linear Keys**| Controls whether the final filtering step will occur, or only the retargeting after bitwise compression. If both this and **Retarget** are false, then the linear compressor will do no better than the underlying bitwise compressor, extremely slowly.|
-
+-->
 
 ### Animation
 
 |---|---|
 |**Retarget Source**| When retargeting, this pose will be used as a base of animation.|
+|**Preview Pose Asset**| The default Skeletal Mesh to use when previewing this asset - this only applies when you open the Animation Editor using this asset. |
+|**Interpolation**| This defines how values between keys are calculated. |
 |**Rate Scale**| Number for tweaking playback rate of this animation globally. |
 |**Skeleton**| Contains the skeleton associated with this Montage. Cannot be changed in the editor.|
-
-### Compression
-
-|---|---|
-|**Description**| Name of Compression Scheme used for this asset.|
-|**Do Not Override Compression**| Do not attempt to override compression scheme when running Compress Animations commandlet. Some high frequency animations are too sensitive and shouldn't be changed. |
+|**Parent Asset**| During cooking, this data will be used to bake out to Normal Asset. Cannot be changed in the editor. |
+|**Asset Mapping Table**| This is used when Parent Asset is set and works with Asset User Data mapping table. Cannot be changed in the editor. |
+|**Asset User Data**| Array of user defined data stored with the asset. |
 
 ### Additive Settings
 
@@ -145,11 +151,20 @@ When an Animation Sequence asset is opened in **Persona**, there are several pro
 |**Additive Anim Type**| The type of additive animation type to use: No Additive, Local Space or Mesh Space. Local space is additive and delta is calculated by local space. Mesh Space is additive and delta will be applied in component space. |
 |**Base Pose Type**| Decides how the additive delta is to be calculated: None (No Base Pose), Reference Pose (The mesh's default reference pose), Selected Animation Scaled (choose another animation and scale the animation), Selected Animation Frame (choose another animation and only based on the frame number). Base Pose Type is only available if Additive Anim Type is set to Local or Mesh Space. |
 
+### Compression
+
+|---|---|
+|**Compression Scheme**| Name of Compression Scheme used for this asset.|
+|**Do Not Override Compression**| Do not attempt to override compression scheme when running Compress Animations commandlet. Some high frequency animations are too sensitive and shouldn't be changed. |
+|**Edit Compression Settings**| Click this button to view and edit the current Compression Settings for this animation. Same as clicking **Asset** from the **Menu Bar** and selecting **Apply Compression**. |
+
+
 ### Root Motion
 
 |---|---|
 |**Enable Root Motion**| If this is on, it will allow extracting of [Root Motion](Engine/Animation/RootMotion). |
-|**RootMotion Root Lock**| Root Bone will be locked to that position specified when extracting root motion. |
+|**Root Motion Root Lock**| Root Bone will be locked to that position specified when extracting root motion. |
+|**Force Root Lock**| This will force Root Bone lock even if Root Motion is not enabled. |
 
 [REGION:note]
 With regards to Root Motion, there are some considerations you may need to make such as whether or not you need to replicate Root Motion in networked games. If you need to replicate Root Motion, you will need to use an [Animation Montage](Engine\Animation\AnimMontage). 
@@ -159,7 +174,25 @@ See [Enabling Root Motion](Engine\Animation\RootMotion#enablingrootmotion) for m
 ### Import Settings
 
 |---|---|
-|**Asset Import Data**| The importing data and options used for the mesh. |
+|**Animation Length**| Which animation range to import. The one defined at Exported Time, at Aniamted Time or using a range set manually. |
+|**Source File**| Browse to the asset's source file used during import. |
+|**Import Meshes in Bone Hierarchy**| If checked, meshes nested in bone hierarchies will be imported instead of being converted to bones. |
+|**Frame Import Range**| Frame range used when Set Range is used in Animation Length. |
+|**Use Default Sample Rate**| Samples all animation curves to 30 fps. |
+|**Import Custom Attribute**| Import custom attribute as a curve within the animation. |
+|**Set Material Curve Type**| Set Material Curve type for all custom attributes that exist. |
+|**Material Curve Suffixes**| Set Material Curve Type for the custom attribute with the following suffixes. This doesn't matter if Set Material Curve Type is true. |
+|**Remove Redundant Keys**| When importing custom attributes as curves, remove redundant keys. |
+|**Delete Existing Morph Target Curves**| This will delete any existing morph target curves from the FBX during import. |
+|**Do not import curves with only 0 values**| When importing custom attribute or morph target as curves, do not import if the value is 0. This is to avoid adding extra curves to evaluate. |
+|**Preserve Local Transform**| This will import a curve within the animation. |
+
+### Transform
+
+|---|---|
+|**Import Translation**| Imports the tranlsation data from the FBX.  |
+|**Import Rotation**| Imports the rotation data from the FBX. |
+|**Import Uniform Scale**| Imports Uniform Scale data from the FBX. |
 
 
 ### Meta Data
@@ -182,3 +215,6 @@ See [Enabling Root Motion](Engine\Animation\RootMotion#enablingrootmotion) for m
 
 
 
+
+
+
diff --git a/Engine/Documentation/Source/Engine/Animation/Sequences/Curves/Curves.INT.udn b/Engine/Documentation/Source/Engine/Animation/Sequences/Curves/Curves.INT.udn
--- a/Engine/Documentation/Source/Engine/Animation/Sequences/Curves/Curves.INT.udn
+++ b/Engine/Documentation/Source/Engine/Animation/Sequences/Curves/Curves.INT.udn
@@ -2,22 +2,31 @@ Availability:Public
 Title:Animation Curves
 Crumbs: %ROOT%, Engine, Engine/Animation, Engine/Animation/Sequences
 Description:Curve system for driving properties of Morph Targets and Materials throughout animation playback.
-version: 4.9
+version: 4.14
+Related:Engine/Content\FBX\MorphTargets
+Related:Engine\Rendering\Materials
+Related:Engine/Animation/Persona/AnimCurves
+
+
 
 [EXCERPT:Intro]
-**Animation Curves** provide a way to change the value of a Material parameter or a Morph Target while an animation is playing back. Their workflow is very simple, requiring that you simply designate the asset you want to modify (either a Material or a Morph Target), name the curve accordingly, and then adjust keyframe values across the duration of an animation.
+**Animation Curves** provide a way to change the value of a Material parameter or a Morph Target while an animation is playing back. Their workflow requires that you designate the asset you want to modify (either a Material or a Morph Target), name the curve accordingly, and then adjust keyframe values across the duration of an animation.
 [/EXCERPT:Intro]
 
-
 When you import Morph Target animations into UE4 via FBX, Morph Target Curves are automatically generated for that AnimSequence. However, they are hidden by default. This is by design; hiding away what could be many different Morph Target Curves will prevent the editor from slowing down.
 
 It should be noted that the ability to adjust curves for Morph Targets, while useful, is more of a convenience feature than anything else. If you are handling a very complex Morph Target animation, you will probably have a much easier time if you perform the animation in the 3D animation package from which you exported out the original Morph Target.
 
 ## Setting Up a Material Curve
 
-In this example, we will show how to set up a simple Material Curve that causes a character to blush during an animation. The setup for a Morph Target curve would be similar, requiring only that the name of the curve match the exact name of the Morph Target that was imported. For more information on Morph Targets, please see [](Engine/Content\FBX\MorphTargets) and [](Engine/Animation/Persona\MorphTargetPreviewer).
+In this example, we will show how to set up a Material Curve that causes our attacking bear's eyes to glow red with rage as they perform a melee attack. 
+
+[REGION:note]
+The setup for a Morph Target curve would be similar, requiring only that the name of the curve match the exact name of the Morph Target that was imported. 
+For more information on Morph Targets, please see [](Engine/Content\FBX\MorphTargets) and [](Engine/Animation/Persona\MorphTargetPreviewer).
+[/REGION]
 
-We will start with a simple animation in which our character is lunging forward with a hammer. We will then set up a simple Material curve that will cause his face to blush heavily while he does this. To begin, let us look at the animation:
+For this example, we will start with our bear character's melee attack where we want to have the bear's eyes glow redder as they are attacking. 
 
 [OBJECT:EmbeddedVideo]
 [PARAMLITERAL:width]
@@ -27,47 +36,52 @@ We will start with a simple animation in which our character is lunging forward
 360
 [/PARAMLITERAL]
 [PARAMLITERAL:videoid]
-TKAcNubFOH8
+5Qt6xJlWRPA
 [/PARAMLITERAL]
 [/OBJECT]
 
-We can open that Animation Sequence in Persona and look at it by **Double-clicking** it in the **Content Browser**.
+If we open our bear's Material, we can see that there is a **Material Expression Scalar Parameter** that controls the intensity of the eyes glowing with a parameter named **Eye Glow Intensity**.
+You will want to take note of the parameter name that you wish to drive with a Material Curve as that name must be the same as the one used inside the Material to affect it.
+The Eye Glow Intensity has a default value of 20 and it is this value that we will drive with our Material Curve to produce a brighter red glow during the animation. 
 
-[REGION:fullwidth]
-![](Persona_HammerSlam.png)
+[REGION:lightbox]
+[![](BearMaterial.png)(w:800)](BearMaterial.png)
 [/REGION]
 
-The character's head has a simple Material with a parameter named **Blush**. Here is a look at the Material Graph:
-
-[REGION:fullwidth]
-![](BlushMaterial.png)
-[/REGION]
-
-[REGION:note]
-We have also made a [Material Instance](Engine/Rendering/Materials\MaterialInstances) from this Material and applied it to our character.
+[REGION:caption]
+Click image for full view of the Material. 
 [/REGION]
 
-In Persona, we start by clicking the ![](button_AddFloatTrack.png) button located in the Curves area. From this, navigate to the **Add Variable Curve...** field and select **Create Curve**. We name the Curve "Blush". This name must coincide with the name used in the Material Parameter.
+Inside the **Animation Editor** for our melee attack, we start by clicking the ![](button_AddFloatTrack.png) button located in the Curves area, then select **Add Variable Curve...** and **Create Curve**. 
 
-![](AddingBlushCurve.png)
+![](CreateCurve.png)
 
-This creates the Curve we need, but we must also make sure to identify it as a Material Curve. This is done using the dropdown at the far right of the curve.
+When prompted to enter a name for the curve, we enter **Eye Glow Intensity** to match the name of the parameter inside our Material. 
 
-![](SettingCurveToMaterial.png)
+![](NameCurve.png)
 
-You can expand the curve area size by clicking the small eye icon next to the dropdown.
+This creates the Curve we need, but we must also make sure to identify it as a **Material Curve**. This is done using through the **Anim Curves** menu under **Window**.  
 
-![](ExpandingCurve.png)
+![](OpenAnimCurvesWindow.png)
 
-You can **Shift+Left-click** in the curve area to create a new keyframe. Keyframe values and times can be adjusted using the fields that appear when a keyframe is selected.
+In the **Anim Curves** window, we click the option under **Type** to make it a Material Curve. 
 
-![](CreateMaterialCurve.png)
+![](ConvertToMaterialCurve2.png)
 
-You can also **Right-click** on a keyframe and change its tangent type:
+We can now add keys in the curve window (**Shift + Left-click**) to affect how/when to adjust the parameter's value (we added 4 keys as shown in the video below). 
 
+[OBJECT:EmbeddedVideo]
+[PARAMLITERAL:width]
+640
+[/PARAMLITERAL]
+[PARAMLITERAL:height]
+360
+[/PARAMLITERAL]
+[PARAMLITERAL:videoid]
+9XMBcDlBXR4
+[/PARAMLITERAL]
+[/OBJECT]
 
-![](CurveType.png)
+Our end result is that the curve now drives the associated Material parameter, causing the bear's eyes to glow as it performs its melee attack. 
 
-Our end result is that the curve now drives the associated Material parameter, causing the character to blush while swinging his hammer.
 
-![](BlushInAction.png)
diff --git a/Engine/Documentation/Source/Engine/Animation/Sequences/Editor/AnimSequenceEditor.INT.udn b/Engine/Documentation/Source/Engine/Animation/Sequences/Editor/AnimSequenceEditor.INT.udn
--- a/Engine/Documentation/Source/Engine/Animation/Sequences/Editor/AnimSequenceEditor.INT.udn
+++ b/Engine/Documentation/Source/Engine/Animation/Sequences/Editor/AnimSequenceEditor.INT.udn
@@ -1,105 +1,132 @@
 Availability:Public
-Title:Editing Animation Sequences
+Title:Animation Sequence Editor
 Crumbs: %ROOT%, Engine, Engine/Animation, Engine/Animation/Sequences
-Description:Guide to using Persona to edit Animation Sequences and add animation notification events.
-version: 4.9
-
-[REGION:banner]
-![](persona_sequence_banner.png)(convert:false)
-[/REGION]
+Description:Guide to using the Animation Editor to edit Animation Sequences and add animation notification events.
+version: 4.14
+Related:Engine/Animation/AnimHowTo/LayerEditing
+Related:Engine/Animation/Sequences\Notifies
+Related:Engine/Animation/Sequences
+Related:Engine\Animation\AnimMontage
 
 [TOC(start:2)]
 
 [EXCERPT:overview]
-The **Animation Sequence Editor** provides functionality for previewing and editing individual **Animation Sequence** assets. The properties, including compression scheme, of the Animation Sequence can be set, and animation notification events - also called Notifies - can be added to allow camera effects, particle effects, sounds, and more to be triggered based on the animation.
+The **Animation Editor** provides functionality for previewing and editing individual **Animation Sequence** assets. The properties, including compression scheme, of the Animation Sequence can be set, and animation notification events - also called Notifies - can be added to allow camera effects, particle effects, sounds, and more to be triggered based on the animation.
 [/EXCERPT:overview]
 
 ## Interface
 
-The **Animation Sequence Editor** consists of the following:
+The **Animation Editor** consists of the following sections:
 
 ![](AnimationSequenceEditor.png)
 
 1. **Name** - Displays the name of the AnimationSequence currently being edited.
 1. **Notifies** - Editable list of tracks for creating/editing notifies. See below for instructions on creating and removing tracks as well as creating and working with notifies.
 1. **Curves** - A graph-like area for adding and editing animation curves
-1. **Tracks** - A field used for editing Animation Layers. For more information, please visit the [](Engine/Animation/AnimHowTo/LayerEditing) page.
-1. **Timeline** - Displays information about, and provides controls for, playback of the preview in the **Viewport** panel of **Persona**.
+1. **Tracks** - A field used for editing [Animation Layers](Engine/Animation/AnimHowTo/LayerEditing). 
+1. **Timeline** - Displays information about, and provides controls for, playback of the preview in the viewport of the **Animation Editor**.
 
 
 ## Notifies
 
 An **Animation Notify** is an event along a particular animation at which you can cause something to happen. For instance, the moment a character's foot hits the floor, you can create a Notify along the timeline to play a footstep sound. Other notifies are available as well, allowing for particle, camera, and sound effects. You may also create custom notifies that trigger your own events.
 
+[REGION:note]
 For more information on Animation Notifies, please see the [](Engine/Animation/Sequences\Notifies) page.
+[/REGION]
 
 
 ## Notify Tracks
 
-Notifies are stored in **Tracks** within Animation Sequences. A track is essentially a container that can hold any number of Notifies, and an Animation Sequences can have any number of tracks. Using multiple tracks makes it easier to organize and visualize Notifies, especially in cases where an Animation Sequences requires large numbers or several different types of Notifies.
+Notifies are stored in **Tracks** within Animation Sequences. A track is essentially a container that can hold any number of Notifies, and Animation Sequences can have any number of tracks. Using multiple tracks makes it easier to organize and visualize Notifies, especially in cases where an Animation Sequences requires large numbers or several different types of Notifies.
+
+Every Animation Sequence has a single notify track by default. Tracks can be added to the sequence or existing tracks can be removed using the controls to the right of the tracks.
+
+To add a new Track:
 
-Every Animation Sequences has a single track by default; however, additional tracks can be added to the sequence and existing tracks can be removed using the controls to the right of the tracks.
+* In the **Track List**, click the ![](button_track_add.png) button next to an existing track:
 
-To add a new Track; In the **Track List**, click the ![](button_track_add.png) button next to an existing track:
+	![](AddTrack.png)
 
-![](AddTrack.png)
+	A new track is added above the existing track in the list:
 
-A new track is added above the existing track in the list:
+	![](TrackAdded.png)
 
-![](TrackAdded.png)
+To remove an existing Track:
 
-To remove an existing Track; In the **Track List**, click the ![](button_track_remove.png) button next to the track you want to remove:
+* In the **Track List**, click the ![](button_track_remove.png) button next to the track you want to remove:
 
-![](RemoveTrack.png)
+	![](RemoveTrack.png)
 
-The track is removed from the list:
+	The track is removed from the list:
 
-![](TrackRemoved.png)
+	![](TrackRemoved.png)
 
 ## Working with Notifies
 
-Notifies can be added to tracks, moved around inside a track, and moved back and forth between different tracks within the **Track List**. This makes it extremely easy to tweak the position of a Notify or quickly arrange them logically.
+Notifies can be added to tracks, moved around inside a track, and moved back and forth between different tracks within the **Track List**. 
+This allows you to quickly tweak the position of a Notify or quickly arrange them logically.
 
 ### Creating a Notify
 
-Notifies of any existing type or custom Notifies can be added through the context menu. **Right-clicking** on a track displays the context menu. All available types of Notifies are listed in the **Add Notify...** context menu. Select the type of notify to add to the track:
+Notifies of any existing type or custom Notifies can be added through the context menu. 
 
-![](AddNotify.png)
+To create a new Notify:
 
-The new Notify is added to the track at the location where you **Right-clicked**:
+* **Right-click** on a track to display the context menu then select the type of notify to add. 
 
-![](NotifyAdded.png)
+	![](AddNotify.png)
+
+	The new Notify is added to the track at the location where you **Right-clicked**:
+
+	![](NotifyAdded.png)
 
 ### Positioning a Notify
 
-The position of a Notify within the track is extremely important as it determines when the event that corresponds to the Notify is triggered. Moving a Notify around is easy; just click and hold the **Left Mouse Button** on the Notify and drag the mouse.
+The position of a Notify within the track is extremely important as it determines when the event that corresponds to the Notify is triggered. 
+
+To move a Notify around: 
 
-![](MovingNotify.png)
+* **Left-click** and hold the **Left Mouse Button** on the Notify and drag to the desired location.
 
-Release the **Left Mouse Button** to drop the Notify at that location:
+	![](MovingNotify.png)
 
-![](NotifyMoved.png)
+* Release the **Left Mouse Button** to drop the Notify at that location:
+
+	![](NotifyMoved.png)
 
 ### Switching Tracks
 
-Moving notifies between tracks is extremely easy. In fact, it is essentially the same as the process for positioning Notifies described above. The only difference is you release the notify on a different track than where it started.
+Moving notifies between tracks is essentially the same as the process for positioning Notifies described above. 
+
+To switch Notify Tracks:
 
-![](TrackSwitch1.png)
+* **Left-click** and hold the **Left Mouse Button** on the Notify to move to a new track. 
 
-![](TrackSwitch2.png)
+	![](TrackSwitch1.png)
 
-![](TrackSwitch3.png) 
+* The Notify will turn yellow at which point you can drag and drop it on a different track. 
+
+	![](TrackSwitch2.png)
+
+	The Notify will turn red when you release the **Left Mouse Button** at the new location.
+
+	![](TrackSwitch3.png) 
 
 
 ### Copy/Pasting Notifies
 
-By **Right-clicking** on a notify, you can copy it to the clipboard for pasting elsewhere along the animation, or even into entirely different Animation Sequences. 
+You can copy a Notify to the clipboard for pasting elsewhere along the animation, or inside different Animation Sequences. 
+
+To copy a Notify: 
+
+* **Right-click** on the Notify and select **Copy**. 
 
-![](CopyAnimNotify.png)
+	![](CopyAnimNotify.png)
 
-You can paste by also **Right-clicking** within a notify track. When pasting, you may choose between 2 options:
+* **Right-click** and select the desired **Paste** option:
 
-![](PasteOptions.png)
+	![](PasteOptions.png)
 
 * **Paste** - Paste the notify at the current location of the time slider.
 * **Paste at Original Time** - Pastes the notify at the same time as when it was copied. 
@@ -129,8 +156,9 @@ Additional Notes:
 
 ## Timeline
 
-![](timeline.png)
 
+![](timeline.png)
+[EXCERPT:timeline]
 The information bar in the **Timeline** displays the following playback information:
 
 | Info | Description |
@@ -140,4 +168,18 @@ The information bar in the **Timeline** displays the following playback informat
 | **Current Time** | The current playback position expressed in seconds. |
 | **Current Frame** | The current playback position expressed in frames. |
 
-[INCLUDE:Engine/Animation/Persona/Viewport#PreviewControls]
+The playback controls perform the following actions: 
+
+| Option | Description |
+| ---- | ----------- |
+|![](TimelineWindow.png) | The timeline. This shows and allows scrubbing through the frames of the AnimationSequence or AnimationMontage. A red bar is displayed at the current position. **Drag** the red bar to scrub. |
+|![](FirstFrameButton.png) | Skips to the first frame in the AnimationSequence or AnimationMontage. |
+|![](StepBackButton.png)| Skips to the previous frame in the AnimationSequence or AnimationMontage.|
+|![](ReversePlay.png)![](PauseButton.png) | Begins playback of the AnimationSequence or AnimationMontage in reverse. When playing, the pause button is displayed. |
+|![](RecordButton.png)| Creates a new animation object in the specified location. Click to start recording, click again to stop the recording.  |
+|![](PlayButton.png) ![](PauseButton.png)| Begins playback of the AnimationSequence or AnimationMontage. When playing, the pause button is displayed. |
+|![](StepForward.png) | Skips to the next frame in the AnimationSequence or AnimationMontage.  |
+|![](SkipToEndButton.png)| Skips to the last frame in the AnimationSequence or AnimationMontage.|
+|![](LoopButton.png)| Toggles playback between play once and looping. When looping is enabled, the button turns red.|
+[/EXCERPT:timeline]
+
diff --git a/Engine/Documentation/Source/Engine/Animation/Sequences/UserGuide/AnimSequenceUserGuide.INT.udn b/Engine/Documentation/Source/Engine/Animation/Sequences/UserGuide/AnimSequenceUserGuide.INT.udn
--- a/Engine/Documentation/Source/Engine/Animation/Sequences/UserGuide/AnimSequenceUserGuide.INT.udn
+++ b/Engine/Documentation/Source/Engine/Animation/Sequences/UserGuide/AnimSequenceUserGuide.INT.udn
@@ -2,22 +2,30 @@ Availability:Public
 Title:Animation Sequence User Guide
 Crumbs: %ROOT%, Engine, Engine/Animation, Engine/Animation/Sequences
 Description:Guide to placing and sampling AnimationSequences in AnimBlueprints
-version: 4.9
+Related: Engine/Animation/AnimBlueprints/AnimGraph
+Related: Engine/Animation/AnimBlueprints/EventGraph
+Related: Engine/Animation/AnimBlueprints
+Related: Engine/Animation/Sequences
+version: 4.14
+type:reference
+tags:Animation
+
+
 
 [EXCERPT:Intro]
 Playback of **Animation Sequences** can be performed through code, but is generally handled within an instance 
-of an Animation Blueprint applied to a SkeletalMeshComponent. All of the Animation Sequences for the Skeleton 
-the Animation Blueprint was created from are available to be sampled within the **AnimGraph**. These can be 
-placed as sequence player nodes whose output is a pose generated by sampling the Animation Sequences.
+of an Animation Blueprint applied to a Skeletal Mesh Component. All of the Animation Sequences for the Skeleton 
+the Animation Blueprint was created from are available to be sampled within the **AnimGraph**. 
+These can be placed as sequence player nodes whose output is a pose generated by sampling the Animation Sequences.
 [/EXCERPT:Intro]
 
 **To place a sequence player:**
 
-1. **Right-click** in the Animation Graph and expand the **Animations** category.
+1. **Right-click** in the AnimGraph of your [Animation Blueprint](Engine/Animation/AnimBlueprints) and expand the **Animations** category.
 
 	![Using Animation Sequences - List](CE_anim_list.png)
 
-1. Select the AnimationSequence to be sampled from the list. A node is placed in the graph for the sequence.
+1. Select the Animation Sequence to be sampled from the list. A node is placed in the graph for the sequence.
 
 	![Using Animation Sequences - Node](CE_anim_node.png)
 
diff --git a/Engine/Documentation/Source/Engine/Animation/Skeleton/Skeleton.INT.udn b/Engine/Documentation/Source/Engine/Animation/Skeleton/Skeleton.INT.udn
--- a/Engine/Documentation/Source/Engine/Animation/Skeleton/Skeleton.INT.udn
+++ b/Engine/Documentation/Source/Engine/Animation/Skeleton/Skeleton.INT.udn
@@ -21,11 +21,14 @@ If you are familiar with digital content creation, you may already feel like you
 A typical character skeleton as you might see in Maya.
 [/REGION]
 
-
 An important distinction is that a Skeleton asset is not the same thing as the bone hierarchy found within a Skeletal Mesh. In Unreal Engine 4, the Skeleton asset is less like a 3D representation of an actual skeleton and more like a list which holds bone and hierarchy information for a specific type of character or Skeletal Mesh. It is through this list that animation associations are made. 
 
 In terms of its place in the animation system, it is probably best to think of a Skeleton asset as an intermediary between Skeletal Meshes and animations. Multiple Skeletal Meshes can use the same Skeleton asset, and any Skeletal Meshes using the same Skeleton asset will be able to share animations.
 
+[REGION:tip]
+See the [Working with Skeletons](https://docs.unrealengine.com/latest/INT/Videos/PLZlv_N0_O1gbwdyIm78w42fZ1t8dDClsI/index.html) tutorial series for more information and what you can do with Skeleton assets.
+[/REGION]
+
 ## Basic Rules
 
 There are some essential things you should know about Skeleton assets when animating for Unreal Engine 4:
@@ -80,3 +83,5 @@ Any time animation data is received for a bone that is not included in the Skele
 There is very little performance overhead with this system, though in terms of production, it is generally a good idea to only share Skeleton assets between Skeletal Meshes with similar bone hierarchies.
 
 
+
+
diff --git a/Engine/Documentation/Source/Engine/Animation/StateMachines/CreatingStateMachines/CreatingStateMachines.INT.udn b/Engine/Documentation/Source/Engine/Animation/StateMachines/CreatingStateMachines/CreatingStateMachines.INT.udn
--- a/Engine/Documentation/Source/Engine/Animation/StateMachines/CreatingStateMachines/CreatingStateMachines.INT.udn
+++ b/Engine/Documentation/Source/Engine/Animation/StateMachines/CreatingStateMachines/CreatingStateMachines.INT.udn
@@ -2,55 +2,65 @@ Availability:Public
 Title: Creating State Machines
 Crumbs: %ROOT%, Engine, Engine/Animation, Engine/Animation/StateMachines
 Description:Guide to creating new State Machines for use in AnimGraphs of animation assets.
-version: 4.12
+version: 4.14
 related:Engine\Animation\StateMachines\EditingStateMachines
 related:Engine\Animation\StateMachines\TransitionRules
 related:Engine\Animation\AnimBlueprints
 
 
-Before actually creating your [**State Machine**](Engine\Animation\StateMachines\Overview), it is a good idea to take a little time and plan out exactly which states you will need, as well as the circumstances that will drive transitions between states. This will give you a solid plan of attack for not only the State Machine, but also the variables you will need to drive it and which animations will be a part of each state.
+Before actually creating your [State Machine](Engine\Animation\StateMachines\Overview), it is a good idea to take a little time and plan out exactly which states you will need, as well as the circumstances that will drive transitions between states. This will give you a solid plan of attack for not only the State Machine but also the variables you will need to drive it and which animations will be a part of each state.
 
-Creating State Machines is done within the AnimGraph of an Animation Blueprint. 
+You can create a **New State Machine** inside the **AnimGraph** of an **Animation Blueprint Editor**: 
 
-**Right-click** within the AnimGraph and choose **State Machines > New State Machine** from the context menu.
+1. **Right-click** within the AnimGraph of your [Animation Blueprint](Engine\Animation\AnimBlueprints) and choose **State Machines > New State Machine** from the context menu.
 
-![](CreateNewStateMachine.png)
+	![](CreateNewStateMachine.png)
 
-Generally, it is a good idea to name your State Machine immediately after creation. You can **Right-click** the new State Machine in the **My Blueprint** window and enter a new name in the name field. In this case, the State Machine has been named "Locomotion."
+	Generally, it is a good idea to name your State Machine immediately after creation. 
+	
+1.	**Right-click** the new State Machine in the **My Blueprint** window and enter a new name in the name field. 
 
-![](RenameStateMachine.png)
+	![](RenameStateMachine.png)
 
-![](LocomotionSMNode.png)
+	In this case, the State Machine has been named "Locomotion."
 
-When first created, a State Machine is empty and may give a warning about this. 
-You can open it to define its States and Transitions by **Double-clicking** it. This will open a new graph tab with nothing but an Entry point, which essentially kicks off the animation.
+	![](LocomotionSMNode.png)
 
+1.	You can open the State Machine to define its States and Transitions by **Double-clicking** it. 
 
-![](StateMachineEntryPoint.png)
+	![](StateMachineEntryPoint.png)
+
+	This will open a new graph tab with nothing but an Entry point, which essentially kicks off the animation.
 
 ## Creating States and Conduits
 
-Creating a new State is similar to creating new nodes within a Blueprint. You can either **Right-click** and create one out of the context menu, or you can drag a wire into an empty portion of the graph and release the mouse, which will also expose the context menu.
+Creating a new State is similar to creating new nodes within a Blueprint. 
+
+1. You can **Right-click** and create one from the context menu. 
+
+	![](AddState_CM.png)
+
+	 You can also drag a wire into an empty portion of the graph and release the mouse, which will also expose the context menu.
 
-![](AddState_CM.png)
+	![](AddStateDragWire.png)
 
-![](AddStateDragWire.png)
+1.	You can rename your States by **Right-clicking** the node in the graph. 
 
-Once a new State has been created, it should immediately be named. This can be done immediately after creation by clicking on the name of the new State, or by **Right-clicking** the node in the graph, as well as within the **My Blueprint** panel.
+	![](StateDefault.png)
 
-![](StateDefault.png)
+	Above our State has been created and renamed to **Idle** below. 	
 
-![](StateRenamed.png)
+	![](StateRenamed.png)
 
-Conduits are created in a similar manner, by either **Right-clicking** in the graph or dragging a new wire into an empty space on the graph.
+1.	Conduits are created in a similar manner, by either **Right-clicking** in the graph or dragging a new wire into an empty space on the graph.
 
-![](CreateConduit.png)
+	![](CreateConduit.png)
 
 
 
 ## Creating Transition Rules
 
-[**Transition Rules**](Engine\Animation\StateMachines\TransitionRules) are created automatically along the wires that connect States and Conduits. Graphically, Transition Rules appear as a small circular directional icon, like so:
+[Transition Rules](Engine\Animation\StateMachines\TransitionRules) are created automatically along the wires that connect States and Conduits. Graphically, Transition Rules appear as a small circular directional icon, like so:
 
 ![](NewTransitionRule.png)
 
@@ -58,4 +68,10 @@ This means that when you create a new State or Conduit by dragging a wire into a
 
 ![](CreateTransition.png)
 
+You can **Double-click** on a Transition Rule to open a new graph that will allow you to define the success or fail conditions for the transition. 
+
+![](TranstionRule.png)
+
+Above if the Boolean **IsSecondaryAttacking?** is **True**, then the transition into the new state can occur. 
+
 
diff --git a/Engine/Documentation/Source/Engine/Animation/StateMachines/EditingStateMachines/EditingStateMachines.INT.udn b/Engine/Documentation/Source/Engine/Animation/StateMachines/EditingStateMachines/EditingStateMachines.INT.udn
--- a/Engine/Documentation/Source/Engine/Animation/StateMachines/EditingStateMachines/EditingStateMachines.INT.udn
+++ b/Engine/Documentation/Source/Engine/Animation/StateMachines/EditingStateMachines/EditingStateMachines.INT.udn
@@ -2,24 +2,24 @@ Availability:Public
 Title: Editing State Machines
 Crumbs: %ROOT%, Engine, Engine/Animation, Engine/Animation/StateMachines
 Description:Guide to editing animation State Machines.
-version: 4.12
+version: 4.14
 related:Engine\Animation\StateMachines\CreatingStateMachines
 related:Engine\Animation\StateMachines\TransitionRules
 related:Engine\Animation\AnimBlueprints
 
-Although a [**State Machine**](Engine\Animation\StateMachines\Overview) serves as a convenient way to break down the flow of animation transitions in a Skeletal Mesh, the most important aspect of them is what actually happens within each node of a State Machine. With this in mind, it is important to remember the primary purposes of each of a State Machine's node types:
+Although a [State Machine](Engine\Animation\StateMachines\Overview) serves as a convenient way to break down the flow of animation transitions in a Skeletal Mesh, the most important aspect of them is what actually happens within each node of a State Machine. With this in mind, it is important to remember the primary purposes of each of a State Machine's node types:
 
 | **Node** | **Purpose** |
 | --- | --- |
 | **State** | Output a pose or animation result. |
-| **Transition Rule** | Govern whether animation will flow from one state into another. Uses a Boolean value to determine if transition will take place. |
-| **Conduit** | Also governs the flow of animation, but with the ability to share the transition cause among many different states. Uses a Boolean value to determine if transition will take place. |
+| **Transition Rule** | Govern whether animation will flow from one state into another. Uses a Boolean value to determine if the transition will take place. |
+| **Conduit** | Also governs the flow of animation, but with the ability to share the transition cause among many different states. Uses a Boolean value to determine if the transition will take place. |
 
 Keeping this in mind will make it much easier to determine what types of nodes to place into each part of the State Machine.
 
 ## Editing States
 
-A State's job is to produce a final animation or pose. This can be something as simple as a single [**AnimSequence**](Engine/Animation/Sequences), a [**BlendSpace**](Engine/Animation/Blendspaces), or any type of network so long as it produces a final animation or pose. 
+A State's job is to produce a final animation or pose. This can be a single [AnimSequence](Engine/Animation/Sequences), a [BlendSpace](Engine/Animation/Blendspaces), or any type of network so long as it produces a final animation or pose. 
 
 ![](PlayAnimSequence.png)
 
@@ -42,7 +42,7 @@ Above a different AnimSequence is selected based on the value of the Boolean Is
 
 ### Adding Animation Elements
 
-To add an AnimSequence or BlendSpace into a State, you can drag-and-drop from the Persona **Asset Browser**: 
+To add an AnimSequence or BlendSpace into a State's graph, you can drag-and-drop it from the **Asset Browser** in the **Animation Blueprint Editor**:
 
 ![](DragDropAnimSeqState.png)(w:800)
 
@@ -54,13 +54,14 @@ Or use the **Right-click** context menu and select an animation:
 
 To automatically create a State in a State Machine, drag an AnimSequence or Blendspace asset into the graph. 
 
-![](AutoStateCreationDrag.png)(w:505)
+![](AutoStateCreationDrag.png)
 
 Once dragged in, the State is automatically created and named after the asset.
 
-![](AutoStateCreation.png)(w:505)
+![](AutoStateCreation.png)
 
 Inside, you can see that the animation asset has been connected to the Final Animation Pose result automatically.
 
-![](AutoStateCreationInside.png)(w:505)
+![](AutoStateCreationInside.png)
+
 
diff --git a/Engine/Documentation/Source/Engine/Animation/StateMachines/Overview/StateMachinesOverview.INT.udn b/Engine/Documentation/Source/Engine/Animation/StateMachines/Overview/StateMachinesOverview.INT.udn
--- a/Engine/Documentation/Source/Engine/Animation/StateMachines/Overview/StateMachinesOverview.INT.udn
+++ b/Engine/Documentation/Source/Engine/Animation/StateMachines/Overview/StateMachinesOverview.INT.udn
@@ -2,7 +2,7 @@ Availability:Public
 Title: Overview of State Machines
 Crumbs: %ROOT%, Engine, Engine/Animation, Engine/Animation/StateMachines
 Description:Breaks down how State Machines work and the core systems used within them.
-version: 4.12
+version: 4.14
 related:Engine/Animation/StateMachines/CreatingStateMachines
 related:Engine/Animation/StateMachines/EditingStateMachines
 related:Engine/Animation/StateMachines/TransitionRules
@@ -45,7 +45,7 @@ related:Engine/Animation/AnimBlueprints
 
 [TOC (start:2)]
 
-Animation **State Machines** allow you to think of your [**Animation Blueprint**](Engine/Animation/AnimBlueprints/) in a more modular sense than you would without it. 
+Animation **State Machines** allow you to think of your [Animation Blueprint](Engine/Animation/AnimBlueprints/) in a more modular sense than you would without it. 
 Users define a series of various states in which the character or Skeletal Mesh could possibly exist. 
 Next, they define the circumstances in which the character or Skeletal Mesh may enter and exit each state, similar to a flow chart. 
 
@@ -62,7 +62,7 @@ If you were to set up a State Machine for the above flow, it might end up lookin
 
 ## Anatomy of a State Machine
 
-State Machines are comprised primarily of a network of states and the rules that define transitions into and out of them. Each state, along with each rule set, is its own condensed [**Blueprint**](Engine/Blueprints) network. 
+State Machines are comprised primarily of a network of states and the rules that define transitions into and out of them. Each state, along with each rule set, is its own condensed [Blueprint](Engine/Blueprints) network. 
 This makes it very easy to handle complex animation blending without having to use an overly complicated **AnimGraph**. 
 
 ### States
@@ -82,7 +82,7 @@ Once the character started moving however, they could transition into a Movement
 
 ![](MoveState.png)
 
-Within, this state may contain a [**BlendSpace**](Engine/Animation/Blendspaces) that is being defined by the character's speed and direction of movement.
+Within, this state may contain a [BlendSpace](Engine/Animation/Blendspaces) that is being defined by the character's speed and direction of movement.
 
 ![](MovementInside.png)
 
diff --git a/Engine/Documentation/Source/Engine/Animation/StateMachines/StateMachines.INT.udn b/Engine/Documentation/Source/Engine/Animation/StateMachines/StateMachines.INT.udn
--- a/Engine/Documentation/Source/Engine/Animation/StateMachines/StateMachines.INT.udn
+++ b/Engine/Documentation/Source/Engine/Animation/StateMachines/StateMachines.INT.udn
@@ -41,7 +41,9 @@ For example, a character may have a state of "Idle" when not moving. Once the ch
 
 -->
 
+[EXCERPT:Overview]
 **State Machines** provide a graphical way to break the animation of a Skeletal Mesh into a series of States. These states are then governed by **Transition Rules** that control how to blend from one state to another. As a tool, they greatly simplify the design process for Skeletal Mesh animation, in that you can create a graph that easily controls how your characters can flow between types of animation without having to create a complex Blueprint network.
+[/EXCERPT:Overview]
 
 Below you will find the information and tutorials for working with State Machines in Unreal Engine 4. 
 
diff --git a/Engine/Documentation/Source/Engine/Animation/StateMachines/TransitionRules/TransitionRules.INT.udn b/Engine/Documentation/Source/Engine/Animation/StateMachines/TransitionRules/TransitionRules.INT.udn
--- a/Engine/Documentation/Source/Engine/Animation/StateMachines/TransitionRules/TransitionRules.INT.udn
+++ b/Engine/Documentation/Source/Engine/Animation/StateMachines/TransitionRules/TransitionRules.INT.udn
@@ -2,7 +2,7 @@
 Title: Transition Rules
 Crumbs: %ROOT%, Engine, Engine/Animation, Engine/Animation/StateMachines
 Description:Guide to rules that govern State Machine Transitions
-version: 4.9
+version: 4.14
 Related: Engine/Animation/StateMachines
 Related: Engine/Animation/AnimBlueprints
 Related: Engine/Blueprints/GettingStarted
@@ -41,7 +41,7 @@ This means that a rule can be set to exit one state to move to another and a rul
 
 ![](EditRule.png)(w:840)
 
-Similar to **Blueprints**, you can specify variables in the **My Blueprint** window and drag them into the graph. By connecting the variables to various conditions which are then connected to the **Result** node, you can define what must be _true_ in order for the transition to occur. 
+Similar to **Blueprints**, you can specify variables in the **My Blueprint** window and drag them into the graph. By connecting the variables to various conditions which are then connected to the **Result** node, you can define what must be _true_ (or _false_) in order for the transition to occur. 
 
 [REGION:note]
 For information on getting started with Blueprints, refer to the [Introduction to Blueprints](Engine/Blueprints/GettingStarted) documentation. 
@@ -58,24 +58,24 @@ When **Left-clicking** on the Transition Rule, the **Details** panel will update
 | **Transition Priority Order** | The priority order of the transition. If multiple transitions out of a state are _true_ at the same time, the one with the lowest number for priority order will take precedent. |
 | **Bidirectional** | Determines if the transition can go both directions. |
 | **Blend Logic** | Specifies what transition logic to use: Standard Blend or [Custom Blend](#customblendnodes). Selecting Custom will enable you to enter the Blend Graph to set a custom Blend Transition. |
-| **Transition Rule Sharing** | This option allows you to set whether or not the rule can be shared with other transitions or to use rule from another transition.  |
+| **Transition Rule Sharing** | This option allows you to set whether or not the rule can be shared with other transitions or to use the rule from another transition.  |
+| **Automatic Rule Based on Sequence Player in State** | This will try to set the rule automatically based on most relevant player node's remaining time and the Crossfade Duration of the transition, ignoring the internal time. |
 | **Transition Interrupt** | Set the Trigger Weight Threshold and specify a Notify State Class.  |
 | **Transition Crossfade Sharing** | Use an existing Transition Crossfade Blend or promote the existing Transition Crossfade to shared.  |
 | **Duration** | The duration to apply to the cross-fade. |
 | **Mode** | The type of blending applied to the cross-fade. See [](Engine\Animation\NonLinearBlends) for more information. |
+| **Custom Blend Curve** | Takes in a user specified blend curve and applies it during the blend. |
+| **Blend Profile** | The blend profile used to evaluate this transition per-bone.  |
 | **Start Transition Event / Custom Blueprint Event** | Call and execute an Event from within the notified Blueprint at the start of the transition. |
-| **Blueprint Notify** | Blueprint that is notified when the transition starts.  |
 | **End Transition Event / Custom Blueprint Event** | Call and execute an Event from within the notified Blueprint at the end of the transition. |
-| **Blueprint Notify** | Blueprint that is notified when the transition ends.  |
 | **Interrupt Transition Event / Custom Blueprint Event** | Call and execute an Event from within the notified Blueprint when the transition is interrupted. |
-| **Blueprint Notify** | Blueprint that is notified when the transition is interrupted. |
 
 
 ## Additional Considerations
 
-When working with child [Animation Blueprints](Engine/Animation/AnimBlueprints) and overriding animations in a State Machine, animations referenced by transition rules do not have to be the same length in every child Animation Blueprint. When using nodes like the **Time Remaining** node (which can be used to determine how much time is remaining in an animation and can be used as part of a Transition Rule), these type of nodes correctly respond to changes in animation length when animations are overridden in child Animation Blueprints. 
+When working with child [Animation Blueprints](Engine/Animation/AnimBlueprints) and overriding animations in a State Machine, animations referenced by transition rules do not have to be the same length in every child Animation Blueprint. When using nodes like the **Time Remaining** node (which can be used to determine how much time is remaining in an animation and can be used as part of a Transition Rule), these type of nodes correctly responds to changes in animation length when animations are overridden in child Animation Blueprints. 
 
-Also, rather than targeting a specific animation as part of your Transition Rule that might change or may be overridden, you can use several class based get nodes that always pick the highest weighted animation from the source state (pictured below).
+Also, rather than targeting a specific animation as part of your Transition Rule that might change or may be overridden, you can use several class-based get nodes that always pick the highest weighted animation from the source state (pictured below).
 
 ![](GetRelevantNodes.png)
 
@@ -103,10 +103,11 @@ In addition to the nodes cited above, there is also information exposed inside C
 You can also use get nodes that reference the asset players, just like a Transition Rule. 
 
 [REGION:note]
-If you have an understanding of Animation Programming in Unreal, [**Native Getters**](Programming/Animation/AnimNodes) can also be used to improve performance.
+If you have an understanding of Animation Programming in Unreal, [Native Getters](Programming/Animation/AnimNodes) can also be used to improve performance.
 [/REGION]
 
 
 
 
 
+
diff --git a/Engine/Documentation/Source/Engine/Audio/Audio.INT.udn b/Engine/Documentation/Source/Engine/Audio/Audio.INT.udn
--- a/Engine/Documentation/Source/Engine/Audio/Audio.INT.udn
+++ b/Engine/Documentation/Source/Engine/Audio/Audio.INT.udn
@@ -5,7 +5,7 @@ Description:System for playing in-game sounds, including the use of SoundCue nod
 Navigation:topic
 version: 4.9
 parent:Engine
-order:4
+order:5
 
 [VAR:Topic]
 [OBJECT:Topic]
@@ -71,6 +71,7 @@ spoken dialog from the characters, the audio in the game can make or break the u
 		* [](Engine/Audio/SoundCues/NodeReference "%Engine/Audio/SoundCues/NodeReference:description%")
 		* [](Engine/Audio/SoundActors "%Engine/Audio/SoundActors:description%")
 		* [](Engine/Audio/DistanceModelAttenuation "%Engine/Audio/DistanceModelAttenuation:description%")
+		* [](Engine/Audio/Dialogue "%Engine/Audio/Dialogue:description%")
 	[/PARAM]
 [/OBJECT]
 
diff --git a/Engine/Documentation/Source/Engine/Audio/Dialogue/1/Dialogue_1.INT.udn b/Engine/Documentation/Source/Engine/Audio/Dialogue/1/Dialogue_1.INT.udn
new file mode 100644
--- /dev/null
+++ b/Engine/Documentation/Source/Engine/Audio/Dialogue/1/Dialogue_1.INT.udn
@@ -0,0 +1,93 @@
+Availability: Public
+Title:1. Create Dialogue Voices
+Crumbs: %ROOT%, Engine, Engine/Audio
+Description:Setting up the voices for each of our characters
+version: 4.13
+parent:Engine/Audio/Dialogue
+order:1
+type:quick start
+checkpoint:dialogueqs
+prereq: Engine/Blueprints/QuickStart
+related:Engine/Audio
+related:Engine/Audio/Overview
+related:Engine/Audio/SoundCues/Editor
+related:Engine/Audio/SoundCues/NodeReference
+
+
+[NAV]
+
+
+In this example, we are using characters from the [Mixamo Animation Pack](https://www.unrealengine.com/marketplace/mixamo-animation-pack), available for free on the Unreal Engine Marketplace. 
+In the Mixamo Animation Pack, all the characters inherit from a general character Blueprint. Working in your own project, or with a 
+C++ setup, you would need to adjust the following steps, but the primary concepts still apply.
+
+Each character has their own value for a Dialogue Voice variable. Even if you don't localize your audio content, having a distinct voice per character
+means that you can associate a given voice actor's recordings with that voice, so Zoe will always sound like Zoe. 
+
+
+1. In the **MixamoAnimPack** folder in the **Content Browser**, **double-click** on **MixamoCharacter_Master** to open the Blueprint.
+
+	![](MixamoMaster.png)(w:700)
+
+1. Add a new [variable](Engine/Blueprints/UserGuide/Variables) to the Blueprint. 
+	
+	![](NewVariable.png)
+
+1. Name the new variable **DialogueVoice**, then set its type to a **DialogueVoice Reference**. 
+
+	![](DialogueVoiceRef.png)
+
+1. Compile and save your Blueprint.
+
+	Now, we need to create the Dialogue Voice assets to use for each of our three characters.
+
+1. Go back to the **Content Browser**, and create a new **Dialogue Voice** asset. 
+
+	![](DialogueVoiceNew.png)
+
+1. Name the new asset **QuestNPC**, then **double-click** to open it.
+
+1. Referring to our design document, we know that the quest NPC's voice is **Feminine** and **Singular**. Use the drop-down menus to set the **Gender** and **Plurality** accordingly.
+	
+	![](QuestNPC.png)
+
+1. Repeat the process two more times to create a **Dialogue Voice** asset for Adam that is **Masculine** and **Singular**, and a **Dialogue Voice** asset for Zoe that is **Feminine** and **Singular**.
+
+	![](AdamVoice.png)
+
+	![](ZoeVoice.png)
+
+1. Save and close all your new **Dialogue Voice** assets.
+
+Now that our **Dialogue Voices** have been created, we can associate them with our characters.
+
+1. Back in the **Content Browser**, in the **Mixamo_SWAT** folder, open the **Mixamo_SWAT** Blueprint.
+
+	![](MixamoSWAT.png)
+
+1. If your Blueprint is not a data-only Blueprint and the defaults are not already visible, click on the **Class Defaults** button in the **Toolbar** to open the Blueprint's default properties.
+
+	![](ClassDefaults.png)
+
+1. Set **Dialogue Voice** in the **Details** panel to **Quest NPC**.
+
+	![](DialogueVoiceSetSwat.png)
+
+1. Repeat the process to set the **Mixamo_Adam** Blueprint's **Dialogue Voice** to **Adam**, and to set the **Mixamo_Zoe** Blueprint's **Dialogue Voice** to **Zoe**.
+
+	![](AdamVoiceSet.png)
+
+	![](ZoeVoiceSet.png)
+
+[REGION:tip]
+To help identify your Dialogue Voice characters:
+
+1. Place your character in the world, then frame them in the viewport. 
+1. Right-click on the **Dialogue Voice** asset in the **Content Browser**.
+1. Select **Asset Actions >> Capture Thumbnail**.
+
+	![](ThumbnailCapture.png)(w:500)
+
+[/REGION]
+
+[NAV]
\ No newline at end of file
diff --git a/Engine/Documentation/Source/Engine/Audio/Dialogue/2/Dialogue_2.INT.udn b/Engine/Documentation/Source/Engine/Audio/Dialogue/2/Dialogue_2.INT.udn
new file mode 100644
--- /dev/null
+++ b/Engine/Documentation/Source/Engine/Audio/Dialogue/2/Dialogue_2.INT.udn
@@ -0,0 +1,62 @@
+Availability: Public
+Title:2. Build Dialogue Waves
+Crumbs: %ROOT%, Engine, Engine/Audio
+Description:Building the Dialogue Waves to store dialogue data
+version: 4.13
+parent:Engine/Audio/Dialogue
+order:2
+type:quick start
+checkpoint:dialogueqs
+prereq: Engine/Blueprints/QuickStart
+related:Engine/Audio
+related:Engine/Audio/Overview
+related:Engine/Audio/SoundCues/Editor
+related:Engine/Audio/SoundCues/NodeReference
+
+[NAV]
+
+
+For this example, we are just going to implement a greeting from Quest NPC to Adam and Zoe. For more information about complex dialogue, see the [Sound Cues and Dialogue section](#soundcuesanddialogue).
+Each line of dialogue needs a **Dialogue Wave** asset associated with it.
+
+1. Go back to the **Content Browser**, and create a new **Dialogue Wave** asset. 
+
+	![](DialogueWaveNew.png)
+
+1. Name the new asset **QuestGreeting**, then **double-click** to open it.
+
+1. For the **Spoken Text**, enter "Hey! Could you come over here? I need your help with something important."
+
+	Although the same text would be spoken when the NPC is talking to Adam and Zoe, the game design document gives us some hints that the way the text is spoken might be different for each listener. That means we
+	should set up two different **Dialogue Contexts**.
+
+1. One context is already created by default. Fill in the **Speaker** entry with **Quest NPC**.
+
+	![](Context1Speaker.png)
+1. Click on the plus sign icon to add a listener. 
+
+	![](Context1ListenerAdd.png)
+1. Set the **Directed At** entry to **Adam**.
+
+	![](Context1Listener.png)
+
+1. Click on **Add Dialogue Context** to add a new **Dialogue Context**.
+
+	![](AddDialogueContext.png)
+
+1. Set this context's **Speaker** to **Quest NPC**, and **Directed At** to **Zoe**.
+
+	![](ZoeContext.png)
+
+We could also make some notes about how the voice actor should sound friendlier toward Zoe, as they have a shared military background, and be more abrupt with Adam, who she doesn't trust because of his mercenary past. These would go in the **Voice Actor Direction** field.
+Finally, after the voice actor recordings come back, we would import those as **Sound Waves** and set them in the **Sound Wave** field for each context.
+In this example, we are not going to create **Sound Waves**, but you could use **Sound Waves** from the **Starter Content** to test. 
+
+![](GreetingAudio.png)
+
+
+[REGION:tip]
+The **Dialogue Wave** is also where you can set a **Subtitle Override**. This is useful for effort sounds, as well as some other cases, such as characters speaking in a foreign language not known to the player.
+[/REGION]
+
+[NAV]
\ No newline at end of file
