Copying //UE4/Dev-Documentation to Samples-Main (//UE4/Samples-Main)

diff --git a/Engine/Documentation/Source/Engine/Audio/Dialogue/3/Dialogue_3.INT.udn b/Engine/Documentation/Source/Engine/Audio/Dialogue/3/Dialogue_3.INT.udn
new file mode 100644
--- /dev/null
+++ b/Engine/Documentation/Source/Engine/Audio/Dialogue/3/Dialogue_3.INT.udn
@@ -0,0 +1,184 @@
+Availability: Public
+Title:3. Set the Context
+Crumbs: %ROOT%, Engine, Engine/Audio
+Description:Building a set of contexts, each with its own speaker and listener
+version: 4.13
+parent:Engine/Audio/Dialogue
+order:3
+type:quick start
+checkpoint:dialogueqs
+prereq: Engine/Blueprints/QuickStart
+related:Engine/Audio
+related:Engine/Audio/Overview
+related:Engine/Audio/SoundCues/Editor
+related:Engine/Audio/SoundCues/NodeReference
+
+[NAV]
+
+Now that we have our **Dialogue Voices** and **Dialogue Wave** set up, we can give our quest NPC some logic to create the right context for her greeting.
+Again, a **Dialogue Context** involves at least two **Dialogue Voices**, a **Speaker** and at least one **Dialogue Voice** that the dialogue is **Directed At**.
+
+1. Add a **Box** component to the **Mixamo_SWAT** Blueprint.
+	
+	![](AddBoxComponent.png)
+1. In the **Viewport**, adjust the **Box** component so it is around the height of the character, and extends some distance in front of her. Make sure it doesn't overlap her, or set her Capsule component to not generate overlap events, or she will trigger the 
+conversation logic herself.
+
+	![](TriggerVolume.png)
+
+1. Add a new **Dialogue Wave** variable named **Greeting**.
+
+1. Compile your Blueprint and set the default value to the **QuestGreeting** **Dialogue Wave** you made previously.
+
+	![](QuestGreetingDetails.png)
+
+1. Set up the following Blueprint graph. 
+
+	[INCLUDE:#graph]
+
+	This graph:
+	* Fires after something overlaps the Box component
+	* Casts the **Overlapping Actor** to our **MixamoCharacter_Master**, where we added the **Dialogue Voice** variable
+	* Gets the voice from the **Overlapping Actor** and uses it as the **Target** for our **Dialogue Context**.
+	* Uses the Quest NPC's voice for the **Speaker** for our **Dialogue Context**.
+	* Plays the **Dialogue Wave** named **Greeting** at our Quest NPC's location, using the correct **Dialogue Context**
+
+	The Mixamo Animation pack comes with a Game Mode we can use to test our dialogue setup.
+
+1. Open **World Settings**.
+
+1. Set the **GameMode Override** to **MixamoGame**.
+
+1. Set the **Default Pawn Class** to **Mixamo_Adam** or **Mixamo_Zoe**.
+
+	![](SetDefaultPawn.png)
+
+	Now, if the QuestNPC is placed in the level, you can run up to them with either Adam or Zoe and trigger the dialogue. 
+
+[NAV]
+
+
+<!--
+[EXCERPT:graph]
+[OBJECT:Snippet]
+[PARAMLITERAL:id]
+	context_graph
+[/PARAMLITERAL]
+[PARAMLITERAL:type]
+	blueprint
+[/PARAMLITERAL]
+[PARAM:image]
+	![](DialogueContextBP.png)
+[/PARAM]
+[PARAM:snippet]
+			Begin Object Class=K2Node_ComponentBoundEvent Name="K2Node_ComponentBoundEvent_0"
+				DelegatePropertyName="OnComponentBeginOverlap"
+				DelegateOwnerClass=Class'/Script/Engine.PrimitiveComponent'
+				ComponentPropertyName="Box"
+				EventReference=(MemberParent=Package'"/Script/Engine"',MemberName="ComponentBeginOverlapSignature__DelegateSignature")
+				bInternalEvent=True
+				CustomFunctionName="BndEvt__Box_K2Node_ComponentBoundEvent_0_ComponentBeginOverlapSignature__DelegateSignature"
+				NodePosX=-112
+				NodePosY=-448
+				NodeGuid=D087184B4C2B1B05E6A2A1AB00A26C42
+				CustomProperties Pin (PinId=E866117345318568604A69A3076763A1,PinName="OutputDelegate",Direction="EGPD_Output",PinType.PinCategory="delegate",PinType.PinSubCategoryObject=None,PinType.PinSubCategoryMemberReference=(MemberParent=BlueprintGeneratedClass'/Game/MixamoAnimPack/Mixamo_SWAT/Mixamo_Swat.Mixamo_SWAT_C',MemberName="BndEvt__Box_K2Node_ComponentBoundEvent_0_ComponentBeginOverlapSignature__DelegateSignature"),PinType.bIsArray=False,PinType.bIsReference=False,PinType.bIsConst=False,PinType.bIsWeakPointer=False,PersistentGuid=00000000000000000000000000000000,bHidden=False,bNotConnectable=False,bDefaultValueIsReadOnly=False,bDefaultValueIsIgnored=False,bAdvancedView=False,)
+				CustomProperties Pin (PinId=6DEDBA6643D23A32098E0C9C21918E7B,PinName="then",Direction="EGPD_Output",PinType.PinCategory="exec",PinType.PinSubCategoryObject=None,PinType.PinSubCategoryMemberReference=(),PinType.bIsArray=False,PinType.bIsReference=False,PinType.bIsConst=False,PinType.bIsWeakPointer=False,LinkedTo=(K2Node_DynamicCast_0 9FA3C5F145140730615EAC8FE8B066D4,),PersistentGuid=00000000000000000000000000000000,bHidden=False,bNotConnectable=False,bDefaultValueIsReadOnly=False,bDefaultValueIsIgnored=False,bAdvancedView=False,)
+				CustomProperties Pin (PinId=6B8FDAB54EA97A386349329F683972B5,PinName="OverlappedComponent",PinToolTip="Overlapped Component\nPrimitive Component Reference",Direction="EGPD_Output",PinType.PinCategory="object",PinType.PinSubCategoryObject=Class'/Script/Engine.PrimitiveComponent',PinType.PinSubCategoryMemberReference=(),PinType.bIsArray=False,PinType.bIsReference=False,PinType.bIsConst=False,PinType.bIsWeakPointer=False,PersistentGuid=00000000000000000000000000000000,bHidden=False,bNotConnectable=False,bDefaultValueIsReadOnly=False,bDefaultValueIsIgnored=False,bAdvancedView=False,)
+				CustomProperties Pin (PinId=FA98EBC24CE57779C26D37B6C8A5A07E,PinName="OtherActor",PinToolTip="Other Actor\nActor Reference",Direction="EGPD_Output",PinType.PinCategory="object",PinType.PinSubCategoryObject=Class'/Script/Engine.Actor',PinType.PinSubCategoryMemberReference=(),PinType.bIsArray=False,PinType.bIsReference=False,PinType.bIsConst=False,PinType.bIsWeakPointer=False,LinkedTo=(K2Node_DynamicCast_0 297553CF423272FD6F7651A9862FE558,),PersistentGuid=00000000000000000000000000000000,bHidden=False,bNotConnectable=False,bDefaultValueIsReadOnly=False,bDefaultValueIsIgnored=False,bAdvancedView=False,)
+				CustomProperties Pin (PinId=1249BE8B4E4CDA85BDD5C185C6184A71,PinName="OtherComp",PinToolTip="Other Comp\nPrimitive Component Reference",Direction="EGPD_Output",PinType.PinCategory="object",PinType.PinSubCategoryObject=Class'/Script/Engine.PrimitiveComponent',PinType.PinSubCategoryMemberReference=(),PinType.bIsArray=False,PinType.bIsReference=False,PinType.bIsConst=False,PinType.bIsWeakPointer=False,PersistentGuid=00000000000000000000000000000000,bHidden=False,bNotConnectable=False,bDefaultValueIsReadOnly=False,bDefaultValueIsIgnored=False,bAdvancedView=False,)
+				CustomProperties Pin (PinId=291C0338444A141C042D07AAC8465708,PinName="OtherBodyIndex",PinToolTip="Other Body Index\nInteger",Direction="EGPD_Output",PinType.PinCategory="int",PinType.PinSubCategoryObject=None,PinType.PinSubCategoryMemberReference=(),PinType.bIsArray=False,PinType.bIsReference=False,PinType.bIsConst=False,PinType.bIsWeakPointer=False,DefaultValue="0",AutogeneratedDefaultValue="0",PersistentGuid=00000000000000000000000000000000,bHidden=False,bNotConnectable=False,bDefaultValueIsReadOnly=False,bDefaultValueIsIgnored=False,bAdvancedView=False,)
+				CustomProperties Pin (PinId=A1DCF200485AB0D9206183B1796790B2,PinName="bFromSweep",PinToolTip="From Sweep\nBoolean",Direction="EGPD_Output",PinType.PinCategory="bool",PinType.PinSubCategoryObject=None,PinType.PinSubCategoryMemberReference=(),PinType.bIsArray=False,PinType.bIsReference=False,PinType.bIsConst=False,PinType.bIsWeakPointer=False,DefaultValue="false",AutogeneratedDefaultValue="false",PersistentGuid=00000000000000000000000000000000,bHidden=False,bNotConnectable=False,bDefaultValueIsReadOnly=False,bDefaultValueIsIgnored=False,bAdvancedView=False,)
+				CustomProperties Pin (PinId=0C0FB5914D618708589BB1AB046B2ED9,PinName="SweepResult",PinToolTip="Sweep Result\nHit Result Structure (by ref)",Direction="EGPD_Output",PinType.PinCategory="struct",PinType.PinSubCategoryObject=ScriptStruct'/Script/Engine.HitResult',PinType.PinSubCategoryMemberReference=(),PinType.bIsArray=False,PinType.bIsReference=True,PinType.bIsConst=True,PinType.bIsWeakPointer=False,PersistentGuid=00000000000000000000000000000000,bHidden=False,bNotConnectable=False,bDefaultValueIsReadOnly=False,bDefaultValueIsIgnored=False,bAdvancedView=False,)
+			End Object
+			Begin Object Class=K2Node_DynamicCast Name="K2Node_DynamicCast_0"
+				TargetType=BlueprintGeneratedClass'/Game/MixamoAnimPack/MixamoCharacter_Master.MixamoCharacter_Master_C'
+				NodePosX=288
+				NodePosY=-432
+				NodeGuid=C3BBF9AF4063E0563F4598A1A94ACC9F
+				CustomProperties Pin (PinId=9FA3C5F145140730615EAC8FE8B066D4,PinName="execute",PinType.PinCategory="exec",PinType.PinSubCategoryObject=None,PinType.PinSubCategoryMemberReference=(),PinType.bIsArray=False,PinType.bIsReference=False,PinType.bIsConst=False,PinType.bIsWeakPointer=False,LinkedTo=(K2Node_ComponentBoundEvent_0 6DEDBA6643D23A32098E0C9C21918E7B,),PersistentGuid=00000000000000000000000000000000,bHidden=False,bNotConnectable=False,bDefaultValueIsReadOnly=False,bDefaultValueIsIgnored=False,bAdvancedView=False,)
+				CustomProperties Pin (PinId=2FF491354C6DDCFAA3AFABB985033BDD,PinName="then",Direction="EGPD_Output",PinType.PinCategory="exec",PinType.PinSubCategoryObject=None,PinType.PinSubCategoryMemberReference=(),PinType.bIsArray=False,PinType.bIsReference=False,PinType.bIsConst=False,PinType.bIsWeakPointer=False,LinkedTo=(K2Node_CallFunction_1 BB5259D94F82BDBC0A387183C5F5E437,),PersistentGuid=00000000000000000000000000000000,bHidden=False,bNotConnectable=False,bDefaultValueIsReadOnly=False,bDefaultValueIsIgnored=False,bAdvancedView=False,)
+				CustomProperties Pin (PinId=68A7707049401F66992BE38B1BF9CABC,PinName="CastFailed",Direction="EGPD_Output",PinType.PinCategory="exec",PinType.PinSubCategoryObject=None,PinType.PinSubCategoryMemberReference=(),PinType.bIsArray=False,PinType.bIsReference=False,PinType.bIsConst=False,PinType.bIsWeakPointer=False,PersistentGuid=00000000000000000000000000000000,bHidden=False,bNotConnectable=False,bDefaultValueIsReadOnly=False,bDefaultValueIsIgnored=False,bAdvancedView=False,)
+				CustomProperties Pin (PinId=297553CF423272FD6F7651A9862FE558,PinName="Object",PinType.PinCategory="object",PinType.PinSubCategoryObject=Class'/Script/CoreUObject.Object',PinType.PinSubCategoryMemberReference=(),PinType.bIsArray=False,PinType.bIsReference=False,PinType.bIsConst=False,PinType.bIsWeakPointer=False,LinkedTo=(K2Node_ComponentBoundEvent_0 FA98EBC24CE57779C26D37B6C8A5A07E,),PersistentGuid=00000000000000000000000000000000,bHidden=False,bNotConnectable=False,bDefaultValueIsReadOnly=False,bDefaultValueIsIgnored=False,bAdvancedView=False,)
+				CustomProperties Pin (PinId=14F1798D43AD8DD3D7DE159D0864B636,PinName="AsMixamo Character Master",Direction="EGPD_Output",PinType.PinCategory="object",PinType.PinSubCategoryObject=BlueprintGeneratedClass'/Game/MixamoAnimPack/MixamoCharacter_Master.MixamoCharacter_Master_C',PinType.PinSubCategoryMemberReference=(),PinType.bIsArray=False,PinType.bIsReference=False,PinType.bIsConst=False,PinType.bIsWeakPointer=False,LinkedTo=(K2Node_VariableGet_1 9F01E92D48B086458B3D4AB930E0154D,),PersistentGuid=00000000000000000000000000000000,bHidden=False,bNotConnectable=False,bDefaultValueIsReadOnly=False,bDefaultValueIsIgnored=False,bAdvancedView=False,)
+				CustomProperties Pin (PinId=1BB8547443D862EAE219CABB96DB3895,PinName="bSuccess",Direction="EGPD_Output",PinType.PinCategory="bool",PinType.PinSubCategoryObject=None,PinType.PinSubCategoryMemberReference=(),PinType.bIsArray=False,PinType.bIsReference=False,PinType.bIsConst=False,PinType.bIsWeakPointer=False,PersistentGuid=00000000000000000000000000000000,bHidden=True,bNotConnectable=False,bDefaultValueIsReadOnly=False,bDefaultValueIsIgnored=False,bAdvancedView=False,)
+			End Object
+			Begin Object Class=K2Node_VariableGet Name="K2Node_VariableGet_1"
+				VariableReference=(MemberParent=BlueprintGeneratedClass'/Game/MixamoAnimPack/MixamoCharacter_Master.MixamoCharacter_Master_C',MemberName="DialogueVoice",MemberGuid=C94B9FF543714FB54D60D1A32479775C)
+				NodePosX=640
+				NodePosY=-352
+				NodeGuid=39D6824C4D9E98B608CEE5BF3C0D3070
+				CustomProperties Pin (PinId=E01C87EB4530E79DF0A6148EE6CFD3EE,PinName="DialogueVoice",Direction="EGPD_Output",PinType.PinCategory="object",PinType.PinSubCategoryObject=Class'/Script/Engine.DialogueVoice',PinType.PinSubCategoryMemberReference=(),PinType.bIsArray=False,PinType.bIsReference=False,PinType.bIsConst=False,PinType.bIsWeakPointer=False,LinkedTo=(K2Node_MakeArray_1 B1C3DD614FDE14143191C2B23A0CB676,),PersistentGuid=00000000000000000000000000000000,bHidden=False,bNotConnectable=False,bDefaultValueIsReadOnly=False,bDefaultValueIsIgnored=False,bAdvancedView=False,)
+				CustomProperties Pin (PinId=9F01E92D48B086458B3D4AB930E0154D,PinName="self",PinFriendlyName=NSLOCTEXT("K2Node", "Target", "Target"),PinType.PinCategory="object",PinType.PinSubCategoryObject=BlueprintGeneratedClass'/Game/MixamoAnimPack/MixamoCharacter_Master.MixamoCharacter_Master_C',PinType.PinSubCategoryMemberReference=(),PinType.bIsArray=False,PinType.bIsReference=False,PinType.bIsConst=False,PinType.bIsWeakPointer=False,LinkedTo=(K2Node_DynamicCast_0 14F1798D43AD8DD3D7DE159D0864B636,),PersistentGuid=00000000000000000000000000000000,bHidden=False,bNotConnectable=False,bDefaultValueIsReadOnly=False,bDefaultValueIsIgnored=False,bAdvancedView=False,)
+			End Object
+			Begin Object Class=K2Node_CallFunction Name="K2Node_CallFunction_1"
+				FunctionReference=(MemberParent=Class'/Script/Engine.GameplayStatics',MemberName="PlayDialogueAtLocation")
+				NodePosX=1392
+				NodePosY=-448
+				NodeGuid=7082B74841AB3F5AF934AF895B47BCD6
+				AdvancedPinDisplay=Hidden
+				CustomProperties Pin (PinId=BB5259D94F82BDBC0A387183C5F5E437,PinName="execute",PinToolTip="\nExec",PinType.PinCategory="exec",PinType.PinSubCategoryObject=None,PinType.PinSubCategoryMemberReference=(),PinType.bIsArray=False,PinType.bIsReference=False,PinType.bIsConst=False,PinType.bIsWeakPointer=False,LinkedTo=(K2Node_DynamicCast_0 2FF491354C6DDCFAA3AFABB985033BDD,),PersistentGuid=00000000000000000000000000000000,bHidden=False,bNotConnectable=False,bDefaultValueIsReadOnly=False,bDefaultValueIsIgnored=False,bAdvancedView=False,)
+				CustomProperties Pin (PinId=B7B051DF4D5FE3697A3D5F87E7DAF3D2,PinName="then",PinToolTip="\nExec",Direction="EGPD_Output",PinType.PinCategory="exec",PinType.PinSubCategoryObject=None,PinType.PinSubCategoryMemberReference=(),PinType.bIsArray=False,PinType.bIsReference=False,PinType.bIsConst=False,PinType.bIsWeakPointer=False,PersistentGuid=00000000000000000000000000000000,bHidden=False,bNotConnectable=False,bDefaultValueIsReadOnly=False,bDefaultValueIsIgnored=False,bAdvancedView=False,)
+				CustomProperties Pin (PinId=0E3069E043FE55FAE21403996D817775,PinName="self",PinFriendlyName=NSLOCTEXT("K2Node", "Target", "Target"),PinToolTip="Target\nGameplay Statics Reference",PinType.PinCategory="object",PinType.PinSubCategoryObject=Class'/Script/Engine.GameplayStatics',PinType.PinSubCategoryMemberReference=(),PinType.bIsArray=False,PinType.bIsReference=False,PinType.bIsConst=False,PinType.bIsWeakPointer=False,DefaultObject="/Script/Engine.Default__GameplayStatics",PersistentGuid=00000000000000000000000000000000,bHidden=True,bNotConnectable=False,bDefaultValueIsReadOnly=False,bDefaultValueIsIgnored=False,bAdvancedView=False,)
+				CustomProperties Pin (PinId=B43B784D4F93500A969171A1E78A600F,PinName="WorldContextObject",PinToolTip="World Context Object\nObject Reference",PinType.PinCategory="object",PinType.PinSubCategoryObject=Class'/Script/CoreUObject.Object',PinType.PinSubCategoryMemberReference=(),PinType.bIsArray=False,PinType.bIsReference=False,PinType.bIsConst=False,PinType.bIsWeakPointer=False,PersistentGuid=00000000000000000000000000000000,bHidden=True,bNotConnectable=False,bDefaultValueIsReadOnly=False,bDefaultValueIsIgnored=False,bAdvancedView=False,)
+				CustomProperties Pin (PinId=B3E56AC842A70ED970E39BB7E3B34B30,PinName="Dialogue",PinToolTip="Dialogue\nDialogue Wave Reference\n\ndialogue to play",PinType.PinCategory="object",PinType.PinSubCategoryObject=Class'/Script/Engine.DialogueWave',PinType.PinSubCategoryMemberReference=(),PinType.bIsArray=False,PinType.bIsReference=False,PinType.bIsConst=False,PinType.bIsWeakPointer=False,LinkedTo=(K2Node_VariableGet_125 53A041D1499B24649BB4B98FD26436A8,),PersistentGuid=00000000000000000000000000000000,bHidden=False,bNotConnectable=False,bDefaultValueIsReadOnly=False,bDefaultValueIsIgnored=False,bAdvancedView=False,)
+				CustomProperties Pin (PinId=67183F194B2F52BBAFF414BE16C3DDF5,PinName="Context",PinToolTip="Context\nDialogue Context Structure (by ref)\n\ncontext the dialogue is to play in",PinType.PinCategory="struct",PinType.PinSubCategoryObject=ScriptStruct'/Script/Engine.DialogueContext',PinType.PinSubCategoryMemberReference=(),PinType.bIsArray=False,PinType.bIsReference=True,PinType.bIsConst=True,PinType.bIsWeakPointer=False,LinkedTo=(K2Node_MakeStruct_0 F991E59A41605C37B2DFCEA0BC272194,),PersistentGuid=00000000000000000000000000000000,bHidden=False,bNotConnectable=False,bDefaultValueIsReadOnly=False,bDefaultValueIsIgnored=True,bAdvancedView=False,)
+				CustomProperties Pin (PinId=B7C2ACA343F80406F380C78747F40AAD,PinName="Location",PinToolTip="Location\nVector \n\nWorld position to play dialogue at",PinType.PinCategory="struct",PinType.PinSubCategoryObject=ScriptStruct'/Script/CoreUObject.Vector',PinType.PinSubCategoryMemberReference=(),PinType.bIsArray=False,PinType.bIsReference=False,PinType.bIsConst=False,PinType.bIsWeakPointer=False,DefaultValue="0, 0, 0",AutogeneratedDefaultValue="0, 0, 0",LinkedTo=(K2Node_CallFunction_2 002C2B1D4C40B9581A095C84138B928D,),PersistentGuid=00000000000000000000000000000000,bHidden=False,bNotConnectable=False,bDefaultValueIsReadOnly=False,bDefaultValueIsIgnored=False,bAdvancedView=False,)
+				CustomProperties Pin (PinId=DC41F40E45DED62F748A888B2080937C,PinName="Rotation",PinToolTip="Rotation\nRotator \n\nWorld rotation to play dialogue at",PinType.PinCategory="struct",PinType.PinSubCategoryObject=ScriptStruct'/Script/CoreUObject.Rotator',PinType.PinSubCategoryMemberReference=(),PinType.bIsArray=False,PinType.bIsReference=False,PinType.bIsConst=False,PinType.bIsWeakPointer=False,DefaultValue="0, 0, 0",AutogeneratedDefaultValue="0, 0, 0",PersistentGuid=00000000000000000000000000000000,bHidden=False,bNotConnectable=False,bDefaultValueIsReadOnly=False,bDefaultValueIsIgnored=False,bAdvancedView=True,)
+				CustomProperties Pin (PinId=159A66C44DFE832C98FD6E80ACD80374,PinName="VolumeMultiplier",PinToolTip="Volume Multiplier\nFloat\n\nVolume multiplier",PinType.PinCategory="float",PinType.PinSubCategoryObject=None,PinType.PinSubCategoryMemberReference=(),PinType.bIsArray=False,PinType.bIsReference=False,PinType.bIsConst=False,PinType.bIsWeakPointer=False,DefaultValue="1.000000",AutogeneratedDefaultValue="1.000000",PersistentGuid=00000000000000000000000000000000,bHidden=False,bNotConnectable=False,bDefaultValueIsReadOnly=False,bDefaultValueIsIgnored=False,bAdvancedView=True,)
+				CustomProperties Pin (PinId=D9F0955F46B638BADA1AFEBE0F96FD97,PinName="PitchMultiplier",PinToolTip="Pitch Multiplier\nFloat\n\nPitch multiplier",PinType.PinCategory="float",PinType.PinSubCategoryObject=None,PinType.PinSubCategoryMemberReference=(),PinType.bIsArray=False,PinType.bIsReference=False,PinType.bIsConst=False,PinType.bIsWeakPointer=False,DefaultValue="1.000000",AutogeneratedDefaultValue="1.000000",PersistentGuid=00000000000000000000000000000000,bHidden=False,bNotConnectable=False,bDefaultValueIsReadOnly=False,bDefaultValueIsIgnored=False,bAdvancedView=True,)
+				CustomProperties Pin (PinId=AF1FB93341C9C31DD7B8088ACE6E4025,PinName="StartTime",PinToolTip="Start Time\nFloat\n\nHow far in to the dialogue to begin playback at",PinType.PinCategory="float",PinType.PinSubCategoryObject=None,PinType.PinSubCategoryMemberReference=(),PinType.bIsArray=False,PinType.bIsReference=False,PinType.bIsConst=False,PinType.bIsWeakPointer=False,DefaultValue="0.000000",AutogeneratedDefaultValue="0.000000",PersistentGuid=00000000000000000000000000000000,bHidden=False,bNotConnectable=False,bDefaultValueIsReadOnly=False,bDefaultValueIsIgnored=False,bAdvancedView=True,)
+				CustomProperties Pin (PinId=654B34F748F75705DAC4AE930D4892CD,PinName="AttenuationSettings",PinToolTip="Attenuation Settings\nSound Attenuation Reference\n\nOverride attenuation settings package to play sound with",PinType.PinCategory="object",PinType.PinSubCategoryObject=Class'/Script/Engine.SoundAttenuation',PinType.PinSubCategoryMemberReference=(),PinType.bIsArray=False,PinType.bIsReference=False,PinType.bIsConst=False,PinType.bIsWeakPointer=False,PersistentGuid=00000000000000000000000000000000,bHidden=False,bNotConnectable=False,bDefaultValueIsReadOnly=False,bDefaultValueIsIgnored=False,bAdvancedView=True,)
+			End Object
+			Begin Object Class=K2Node_CallFunction Name="K2Node_CallFunction_2"
+				bIsPureFunc=True
+				bIsConstFunc=True
+				FunctionReference=(MemberName="K2_GetActorLocation",bSelfContext=True)
+				NodePosX=464
+				NodePosY=-160
+				NodeGuid=EFFECE1D446CA14B120F28AFF854FF5C
+				CustomProperties Pin (PinId=B50B101C42A61708AAC0398DC15687AD,PinName="self",PinFriendlyName=NSLOCTEXT("K2Node", "Target", "Target"),PinToolTip="Target\nActor Reference",PinType.PinCategory="object",PinType.PinSubCategoryObject=Class'/Script/Engine.Actor',PinType.PinSubCategoryMemberReference=(),PinType.bIsArray=False,PinType.bIsReference=False,PinType.bIsConst=False,PinType.bIsWeakPointer=False,PersistentGuid=00000000000000000000000000000000,bHidden=False,bNotConnectable=False,bDefaultValueIsReadOnly=False,bDefaultValueIsIgnored=False,bAdvancedView=False,)
+				CustomProperties Pin (PinId=002C2B1D4C40B9581A095C84138B928D,PinName="ReturnValue",PinToolTip="Return Value\nVector ",Direction="EGPD_Output",PinType.PinCategory="struct",PinType.PinSubCategoryObject=ScriptStruct'/Script/CoreUObject.Vector',PinType.PinSubCategoryMemberReference=(),PinType.bIsArray=False,PinType.bIsReference=False,PinType.bIsConst=False,PinType.bIsWeakPointer=False,AutogeneratedDefaultValue="0, 0, 0",LinkedTo=(K2Node_CallFunction_1 B7C2ACA343F80406F380C78747F40AAD,),PersistentGuid=00000000000000000000000000000000,bHidden=False,bNotConnectable=False,bDefaultValueIsReadOnly=False,bDefaultValueIsIgnored=False,bAdvancedView=False,)
+			End Object
+			Begin Object Class=K2Node_MakeStruct Name="K2Node_MakeStruct_0"
+				bMadeAfterOverridePinRemoval=True
+				ShowPinForProperties(0)=(PropertyName="Speaker",PropertyFriendlyName="Speaker",PropertyTooltip=NSLOCTEXT("UObjectToolTips", "DialogueContext:Speaker", "The person speaking the dialogue."),bShowPin=True,bCanToggleVisibility=True,bPropertyIsCustomized=True,CategoryName="DialogueContext")
+				ShowPinForProperties(1)=(PropertyName="Targets",PropertyFriendlyName="Targets",PropertyTooltip=NSLOCTEXT("UObjectToolTips", "DialogueContext:Targets", "The people being spoken to."),bShowPin=True,bCanToggleVisibility=True,bPropertyIsCustomized=True,CategoryName="DialogueContext")
+				StructType=ScriptStruct'/Script/Engine.DialogueContext'
+				NodePosX=976
+				NodePosY=-384
+				bHasCompilerMessage=True
+				ErrorType=3
+				ErrorMsg="Warning  Speaker  - the native property is not tagged as BlueprintReadWrite, the pin will be removed in a future release.\nWarning  Targets  - the native property is not tagged as BlueprintReadWrite, the pin will be removed in a future release.\nWarning  K2Node_MakeStruct_0  has no property tagged as BlueprintReadWrite. The node will be removed in a future release."
+				NodeGuid=008F6A574C6FA4B2561B57A77BCA1E48
+				CustomProperties Pin (PinId=F991E59A41605C37B2DFCEA0BC272194,PinName="DialogueContext",Direction="EGPD_Output",PinType.PinCategory="struct",PinType.PinSubCategoryObject=ScriptStruct'/Script/Engine.DialogueContext',PinType.PinSubCategoryMemberReference=(),PinType.bIsArray=False,PinType.bIsReference=False,PinType.bIsConst=False,PinType.bIsWeakPointer=False,LinkedTo=(K2Node_CallFunction_1 67183F194B2F52BBAFF414BE16C3DDF5,),PersistentGuid=00000000000000000000000000000000,bHidden=False,bNotConnectable=False,bDefaultValueIsReadOnly=False,bDefaultValueIsIgnored=False,bAdvancedView=False,)
+				CustomProperties Pin (PinId=CDFE76BB4A0C7A3A699B4D8BEB7B684C,PinName="Speaker",PinFriendlyName=NSLOCTEXT("", "A6BD052743422BBCD8BA208C6569C2A5", "Speaker"),PinToolTip="Speaker\nDialogue Voice Reference\n\nThe person speaking the dialogue.",PinType.PinCategory="object",PinType.PinSubCategoryObject=Class'/Script/Engine.DialogueVoice',PinType.PinSubCategoryMemberReference=(),PinType.bIsArray=False,PinType.bIsReference=False,PinType.bIsConst=False,PinType.bIsWeakPointer=False,LinkedTo=(K2Node_VariableGet_5 D6C1A07046B65B3E8CF5AFB691087C9D,),PersistentGuid=00000000000000000000000000000000,bHidden=False,bNotConnectable=False,bDefaultValueIsReadOnly=False,bDefaultValueIsIgnored=False,bAdvancedView=False,)
+				CustomProperties Pin (PinId=9FB6EB844C8A89056EFFD4A68EB33E82,PinName="Targets",PinFriendlyName=NSLOCTEXT("", "08F2369A40811DC5A47D24ABD7F2FB7A", "Targets"),PinToolTip="Targets\nArray of Dialogue Voice References\n\nThe people being spoken to.",PinType.PinCategory="object",PinType.PinSubCategoryObject=Class'/Script/Engine.DialogueVoice',PinType.PinSubCategoryMemberReference=(),PinType.bIsArray=True,PinType.bIsReference=False,PinType.bIsConst=False,PinType.bIsWeakPointer=False,LinkedTo=(K2Node_MakeArray_1 375D77E547627D97D0861B909893B26C,),PersistentGuid=00000000000000000000000000000000,bHidden=False,bNotConnectable=False,bDefaultValueIsReadOnly=False,bDefaultValueIsIgnored=False,bAdvancedView=False,)
+			End Object
+			Begin Object Class=K2Node_VariableGet Name="K2Node_VariableGet_5"
+				VariableReference=(MemberName="DialogueVoice",MemberGuid=C94B9FF543714FB54D60D1A32479775C,bSelfContext=True)
+				NodePosX=656
+				NodePosY=-544
+				NodeGuid=9495246E4FB8C1F1C96AE4AC015124BA
+				CustomProperties Pin (PinId=D6C1A07046B65B3E8CF5AFB691087C9D,PinName="DialogueVoice",Direction="EGPD_Output",PinType.PinCategory="object",PinType.PinSubCategoryObject=Class'/Script/Engine.DialogueVoice',PinType.PinSubCategoryMemberReference=(),PinType.bIsArray=False,PinType.bIsReference=False,PinType.bIsConst=False,PinType.bIsWeakPointer=False,LinkedTo=(K2Node_MakeStruct_0 CDFE76BB4A0C7A3A699B4D8BEB7B684C,),PersistentGuid=00000000000000000000000000000000,bHidden=False,bNotConnectable=False,bDefaultValueIsReadOnly=False,bDefaultValueIsIgnored=False,bAdvancedView=False,)
+				CustomProperties Pin (PinId=6AB7071343FD459EF40370B7AEAA2F0D,PinName="self",PinFriendlyName=NSLOCTEXT("K2Node", "Target", "Target"),PinType.PinCategory="object",PinType.PinSubCategoryObject=BlueprintGeneratedClass'/Game/MixamoAnimPack/MixamoCharacter_Master.MixamoCharacter_Master_C',PinType.PinSubCategoryMemberReference=(),PinType.bIsArray=False,PinType.bIsReference=False,PinType.bIsConst=False,PinType.bIsWeakPointer=False,PersistentGuid=00000000000000000000000000000000,bHidden=True,bNotConnectable=False,bDefaultValueIsReadOnly=False,bDefaultValueIsIgnored=False,bAdvancedView=False,)
+			End Object
+			Begin Object Class=K2Node_MakeArray Name="K2Node_MakeArray_1"
+				NodePosX=672
+				NodePosY=-320
+				NodeGuid=3DBCE66345A3AC43437A4BB8AC264C4B
+				CustomProperties Pin (PinId=375D77E547627D97D0861B909893B26C,PinName="Array",Direction="EGPD_Output",PinType.PinCategory="object",PinType.PinSubCategoryObject=Class'/Script/Engine.DialogueVoice',PinType.PinSubCategoryMemberReference=(),PinType.bIsArray=True,PinType.bIsReference=False,PinType.bIsConst=False,PinType.bIsWeakPointer=False,LinkedTo=(K2Node_MakeStruct_0 9FB6EB844C8A89056EFFD4A68EB33E82,),PersistentGuid=00000000000000000000000000000000,bHidden=False,bNotConnectable=False,bDefaultValueIsReadOnly=False,bDefaultValueIsIgnored=False,bAdvancedView=False,)
+				CustomProperties Pin (PinId=B1C3DD614FDE14143191C2B23A0CB676,PinName="[0]",PinType.PinCategory="object",PinType.PinSubCategoryObject=Class'/Script/Engine.DialogueVoice',PinType.PinSubCategoryMemberReference=(),PinType.bIsArray=False,PinType.bIsReference=False,PinType.bIsConst=False,PinType.bIsWeakPointer=False,LinkedTo=(K2Node_VariableGet_1 E01C87EB4530E79DF0A6148EE6CFD3EE,),PersistentGuid=00000000000000000000000000000000,bHidden=False,bNotConnectable=False,bDefaultValueIsReadOnly=False,bDefaultValueIsIgnored=False,bAdvancedView=False,)
+			End Object
+			Begin Object Class=K2Node_VariableGet Name="K2Node_VariableGet_125"
+				VariableReference=(MemberName="Greeting",bSelfContext=True)
+				NodePosX=1200
+				NodePosY=-448
+				NodeGuid=DA2B19AA4859BAE664346FB8F9552B60
+				CustomProperties Pin (PinId=53A041D1499B24649BB4B98FD26436A8,PinName="Greeting",Direction="EGPD_Output",PinType.PinCategory="object",PinType.PinSubCategoryObject=Class'/Script/Engine.DialogueWave',PinType.PinSubCategoryMemberReference=(),PinType.bIsArray=False,PinType.bIsReference=False,PinType.bIsConst=False,PinType.bIsWeakPointer=False,LinkedTo=(K2Node_CallFunction_1 B3E56AC842A70ED970E39BB7E3B34B30,),PersistentGuid=00000000000000000000000000000000,bHidden=False,bNotConnectable=False,bDefaultValueIsReadOnly=False,bDefaultValueIsIgnored=False,bAdvancedView=False,)
+				CustomProperties Pin (PinId=570FAEAE4E3FFD616AB85B92D72DF118,PinName="self",PinFriendlyName=NSLOCTEXT("K2Node", "Target", "Target"),PinType.PinCategory="object",PinType.PinSubCategoryObject=BlueprintGeneratedClass'/Game/MixamoAnimPack/Mixamo_SWAT/Mixamo_Swat.Mixamo_SWAT_C',PinType.PinSubCategoryMemberReference=(),PinType.bIsArray=False,PinType.bIsReference=False,PinType.bIsConst=False,PinType.bIsWeakPointer=False,PersistentGuid=00000000000000000000000000000000,bHidden=True,bNotConnectable=False,bDefaultValueIsReadOnly=False,bDefaultValueIsIgnored=False,bAdvancedView=False,)
+			End Object
+
+[/PARAM]
+[/OBJECT]
+[/EXCERPT]
+-->
\ No newline at end of file
diff --git a/Engine/Documentation/Source/Engine/Audio/Dialogue/4/Dialogue_4.INT.udn b/Engine/Documentation/Source/Engine/Audio/Dialogue/4/Dialogue_4.INT.udn
new file mode 100644
--- /dev/null
+++ b/Engine/Documentation/Source/Engine/Audio/Dialogue/4/Dialogue_4.INT.udn
@@ -0,0 +1,34 @@
+Availability: Public
+Title:4. Displaying Subtitles
+Crumbs: %ROOT%, Engine, Engine/Audio
+Description:Enabling subtitles to display for your project
+version: 4.13
+parent:Engine/Audio/Dialogue
+order:4
+type:quick start
+checkpoint:dialogueqs
+prereq: Engine/Blueprints/QuickStart
+related:Engine/Audio
+related:Engine/Audio/Overview
+related:Engine/Audio/SoundCues/Editor
+related:Engine/Audio/SoundCues/NodeReference
+
+[NAV]
+
+Subtitles should be enabled by default, but if you need to check that they are on, there is a checkbox in **Project Settings**.
+
+1. Open **Project Settings**.
+
+1. Under **General Settings**, locate the **Subtitles** section, then make sure the **Subtitles Enabled** box is checked. 
+
+	![](SubtitlesOn.png)
+
+Now, when you trigger the Quest NPC's dialogue, the accompanying subtitle text shows at the bottom of the screen.
+
+![](SubtitlesPlaying.png)
+
+[REGION:note]
+Subtitles will not display unless there is a **Sound Wave** set in your **Dialogue Wave** for the context that is currently happening.
+[/REGION]
+
+[NAV]
\ No newline at end of file
diff --git a/Engine/Documentation/Source/Engine/Audio/Dialogue/5/Dialogue_5.INT.udn b/Engine/Documentation/Source/Engine/Audio/Dialogue/5/Dialogue_5.INT.udn
new file mode 100644
--- /dev/null
+++ b/Engine/Documentation/Source/Engine/Audio/Dialogue/5/Dialogue_5.INT.udn
@@ -0,0 +1,38 @@
+Availability: Public
+Title:5. Next Steps: Sound Cues and Dialogue
+Crumbs: %ROOT%, Engine, Engine/Audio
+Description:A jumping off point to build a more complex dialogue manager system
+version: 4.13
+parent:Engine/Audio/Dialogue
+order:5
+type:quick start
+checkpoint:dialogueqs
+prereq: Engine/Blueprints/QuickStart
+related:Engine/Audio
+related:Engine/Audio/Overview
+related:Engine/Audio/SoundCues/Editor
+related:Engine/Audio/SoundCues/NodeReference
+
+[NAV]
+
+For a more complex conversation, you will likely want to use Sound Cues and Audio Components. Audio Components enable you to bind some functionality to when audio is finished playing, using the 
+**On Audio Finished** delegate.  
+
+[REGION:note]
+One thing to consider, however, is that you cannot dynamically change the **Dialogue Context** when using Sound Cues, so your Blueprints or C++ logic will look 
+different than the above example. 
+[/REGION]
+
+To play **Dialogue Waves** using **Sound Cues**, there is a **Dialogue Player** node in the **Sound Cue Editor**.
+
+![](DialoguePlayerNode.png)
+
+Select the **Dialogue Player** node and look at the **Details** panel. Here is where you can set the **Dialogue Wave** to use, as well as the **Speaker** and **Directed At** entries for the **Context**.
+
+![](DialogueNodeDetails.png)
+
+Here is some example Blueprint logic that shows how you could bind events to play the next piece of dialogue after the previous audio finishes playing.
+
+![](AudioNextSteps.png)
+
+[NAV]
diff --git a/Engine/Documentation/Source/Engine/Audio/Dialogue/Dialogue.INT.udn b/Engine/Documentation/Source/Engine/Audio/Dialogue/Dialogue.INT.udn
new file mode 100644
--- /dev/null
+++ b/Engine/Documentation/Source/Engine/Audio/Dialogue/Dialogue.INT.udn
@@ -0,0 +1,88 @@
+Availability: Public
+Title:Dialogue Voices and Waves
+Crumbs: %ROOT%, Engine, Engine/Audio
+Description:Example of a dialogue setup with different listeners and dialogue contexts
+version: 4.13
+checkpoint:dialogueqs
+prereq: Engine/Blueprints/QuickStart
+type:quick start
+parent:Engine/Audio
+related:Engine/Audio/Overview
+related:Engine/Audio/SoundCues/Editor
+related:Engine/Audio/SoundCues/NodeReference
+
+[VAR:Topic]
+[OBJECT:Topic]
+	[PARAM:image]
+		![%Engine/Audio/Overview:title%](Engine/Audio/audio_topic.png)
+	[/PARAM]
+	[PARAM:icon]
+		![](%ROOT%/audio_icon.png)(convert:false)
+	[/PARAM]
+	[PARAM:title]
+		%Engine/Audio/Overview:title%
+	[/PARAM]
+	[PARAM:description]
+		%Engine/Audio/Overview:description%
+	[/PARAM]
+	[PARAM:path]
+		[RELATIVE:Engine/Audio/Overview]
+	[/PARAM]
+[/OBJECT]
+[/VAR]
+[VAR:TopicCompact]
+[OBJECT:TopicCompact]
+	[PARAM:image]
+		![%Engine/Audio/Overview:title%](Engine/Audio/audio_topic.png)
+	[/PARAM]
+	[PARAM:icon]
+		![](%ROOT%/audio_icon.png)(convert:false)
+	[/PARAM]
+	[PARAM:title]
+		%Engine/Audio/Overview:title%
+	[/PARAM]
+	[PARAM:description]
+		%Engine/Audio/Overview:description%
+	[/PARAM]
+	[PARAM:path]
+		[RELATIVE:Engine/Audio/Overview]
+	[/PARAM]
+[/OBJECT]
+[/VAR]
+
+
+This example goes over a very basic dialogue setup. For longer conversations, you will likely want to create a conversation manager system. 
+In this scenario, we will set up a quest NPC and two possible playable characters, and show how using Dialogue Voices and Dialogue Waves enables
+you to adjust the quest giver's tone depending on who he is speaking to.
+
+[REGION:imagetable]
+| ![](QuestGiver.png)(h:300)| ![](Adam.png)(h:300)|![](Zoe.png)(h:300) |
+|---|----|----|
+|Quest NPC | Adam | Zoe | 
+[/REGION]
+
+Also, here's a snippet of our example design document:
+
+		- Quest NPC is a female soldier who wants to hire someone from Adam's team as backup on a supply escort mission.
+		- Adam is a mercenary who has a stealth-first, attacking only when necessary policy.
+		- Zoe is an ex-soldier who joined Adam's team after the last major Galactic Skirmish.
+
+
+
+[OBJECT:TopicList]
+	[PARAM:icon]
+		![](%ROOT%/build_icon.png)(convert:false)
+	[/PARAM]
+	[PARAM:title]
+		Steps
+	[/PARAM]
+	[PARAM:description]
+	[/PARAM]
+	[PARAM:links]
+	[DIR(output:"checkpoint" parent:"Engine/Audio/Dialogue" org:"hierarchy")]
+	[/PARAM]
+[/OBJECT]
+
+[REGION:call_to_action]
+	[Click to Start](Engine/Audio/Dialogue/1 "%Engine/Audio/Dialogue/1:title%")
+[/REGION]
\ No newline at end of file
diff --git a/Engine/Documentation/Source/Engine/Audio/Overview/Overview.INT.udn b/Engine/Documentation/Source/Engine/Audio/Overview/Overview.INT.udn
--- a/Engine/Documentation/Source/Engine/Audio/Overview/Overview.INT.udn
+++ b/Engine/Documentation/Source/Engine/Audio/Overview/Overview.INT.udn
@@ -149,7 +149,7 @@ Inside the Sound Classes section, you set which Sound Classes are to be affected
 The Sound Mix section allows you to specify how the Sound Mix properties are applied or removed.  **Delay** indicates how long before the mix properties should begin being applied.  **Fade in Time** and **Fade out Time** specify how quickly to transition from no effect to the specified properties.  **Duration** allows a pushed Sound Mix to automatically pop itself after the specified duration.  A value of -1 indicates to never automatically pop and passively applied Sound Mixes will not automatically pop.
 
 
-### Dialogue Voice & Dialogue Wave
+### Dialogue Voice and Dialogue Wave
 The **Dialogue Voice** and **Dialogue Wave** assets are used for generating in-game dialogue events, crafting subtitles, and for supplementing localization efforts. 
 
 When editing a newly created Dialogue Voice asset, you can define the **Gender** and **Plurality** of a voice actor and although you do not specify any audio assets inside the Dialogue Voice, the information provided here can be referenced inside the Dialogue Wave. 
@@ -170,6 +170,8 @@ Dialogue Wave assets can also be applied to a Sound Cue by using the Dialogue Pl
 In addition to Dialogue Contexts, you can apply a **Mature** filter which flags the dialogue as containing mature/adult content. Under **Script**, you can enter the text that is spoken in the attached audio inside the **Spoken Text** section. 
 You can also enter contextual information for the audio for translation purposes or for notes to a voice actor in the **Voice Actor Direction** section. 
 
+For more information, see the [](Engine/Audio/Dialogue) example.
+
 
 ## Strategies
 
diff --git a/Engine/Documentation/Source/Engine/Basics/DerivedDataCache/DerivedDataCache.INT.udn b/Engine/Documentation/Source/Engine/Basics/DerivedDataCache/DerivedDataCache.INT.udn
--- a/Engine/Documentation/Source/Engine/Basics/DerivedDataCache/DerivedDataCache.INT.udn
+++ b/Engine/Documentation/Source/Engine/Basics/DerivedDataCache/DerivedDataCache.INT.udn
@@ -29,7 +29,7 @@ section in your game's **DefaultEngine.ini**. This section is originally declare
 where the paths are set to Epic's internal DDC share. Re-declare the section and change the paths to point 
 to a share on your network (e.g. `Path=\\mystudio.net\DDC`):
 
-	[DerivedDataBackendGraph]
+	[InstalledDerivedDataBackendGraph]
 	MinimumDaysToKeepFile=7
 	Root=(Type=KeyLength, Length=120, Inner=AsyncPut)
 	AsyncPut=(Type=AsyncPut, Inner=Hierarchy)
diff --git a/Engine/Documentation/Source/Engine/BlueprintSplines/HowTo/AddSplineComponent/AddSplineComponent.INT.udn b/Engine/Documentation/Source/Engine/BlueprintSplines/HowTo/AddSplineComponent/AddSplineComponent.INT.udn
--- a/Engine/Documentation/Source/Engine/BlueprintSplines/HowTo/AddSplineComponent/AddSplineComponent.INT.udn
+++ b/Engine/Documentation/Source/Engine/BlueprintSplines/HowTo/AddSplineComponent/AddSplineComponent.INT.udn
@@ -1,4 +1,4 @@
-Availability:Docs
+Availability:Public
 Title:How-To Add a Spline Component to a Blueprint
 Crumbs: %ROOT%, Engine
 Description:This how-to will show you how to add a Spline or Spline Mesh Component to a Blueprint.
@@ -6,8 +6,11 @@ Navigation:topic
 parent:Engine/BlueprintSplines/HowTo
 order:1
 social-image:Splines_social.png
+version:4.13
+skilllevel:Beginner
 related:Engine/BlueprintSplines/HowTo/EditSplineComponentInEditor
 related:Engine/BlueprintSplines/HowTo/EditSplineInConstScript
+tags:Splines
 
 [VAR:Topic]
 [OBJECT:Topic]
@@ -49,7 +52,7 @@ related:Engine/BlueprintSplines/HowTo/EditSplineInConstScript
 [/VAR]
 
 
-This page will take you through the steps necessary to add a **Blueprint Spline Component** or **Blueprint Spline Mesh Component** to any Blueprint you have available. Both of these components require that your Blueprint inherits from any class that can trace its linage back to **Actor** (Pawn, Character, etc...).
+This page will take you through the steps necessary to add a **Blueprint Spline Component** or **Blueprint Spline Mesh Component** to any Blueprint you have available. Both of these components require that your Blueprint inherits from any class that can trace its linage back to **Actor** (Pawn, Character, etc.).
 
 ## Steps
 
@@ -63,11 +66,11 @@ This page will take you through the steps necessary to add a **Blueprint Spline
 	
 1. Adjust any settings for the added component. See [](Engine/BlueprintSplines/Reference/ComponentProperties) or [](Engine/BlueprintSplines/Reference/MeshComponentProperties) for more information on the properties available.
 
-	You won't be able to directly edit the Spline's or Spline Mesh's points in the **Blueprint Editor Viewport**. You'll either need define the curve in the **Construction Script** of your Blueprint or place your Blueprint in your level and edit the points there.
+	You can directly edit a **Blueprint Spline Component's** points in the **Blueprint Editor Viewport**.
 	
 ## Result
 
-You will have added a **Spline Component** or **Spline Mesh Component** to your Blueprint, see the next how-to's for more information on how to manipulate them.
+You will have added a **Blueprint Spline Component** or **Blueprint Spline Mesh Component** to your Blueprint, see the next how-to's for more information on how to manipulate them.
 
 
 
diff --git a/Engine/Documentation/Source/Engine/BlueprintSplines/HowTo/EditSplineComponentInEditor/EditSplineCompInEditor.INT.udn b/Engine/Documentation/Source/Engine/BlueprintSplines/HowTo/EditSplineComponentInEditor/EditSplineCompInEditor.INT.udn
--- a/Engine/Documentation/Source/Engine/BlueprintSplines/HowTo/EditSplineComponentInEditor/EditSplineCompInEditor.INT.udn
+++ b/Engine/Documentation/Source/Engine/BlueprintSplines/HowTo/EditSplineComponentInEditor/EditSplineCompInEditor.INT.udn
@@ -1,5 +1,5 @@
-Availability:Docs
-Title:How-To Edit a Spline Component in the Level Editor
+Availability:Public
+Title:How-To Edit a Spline Component in the Viewport
 Crumbs: %ROOT%, Engine
 Description:Null
 Navigation:topic
@@ -7,6 +7,10 @@ parent:Engine/BlueprintSplines/HowTo
 order:2
 social-image:Splines_social.png
 related:Engine/BlueprintSplines/HowTo/EditSplineInConstScript
+version:4.13
+skilllevel:Beginner
+reference-image:FullEdit.gif
+tags:Splines
 
 [VAR:Topic]
 [OBJECT:Topic]
@@ -47,7 +51,12 @@ related:Engine/BlueprintSplines/HowTo/EditSplineInConstScript
 [/OBJECT]
 [/VAR]
 
-This how-to will step through editing **Blueprint Spline Component** curves while they are in the Level Editor. **Spline Mesh Components** can also be set so they can be edited in the Level Editor, but only have two points to use, and you cannot add any more points.
+This how-to will step through editing **Blueprint Spline Component** curves while they are in the Level Editor. **Blueprint Spline Mesh Components** can also be set so they can be edited in the Level Editor, but only have two points to use, and you cannot add any more points.
+
+[REGION:note]
+While this how-to covers editing splines while in the Level Editor, you can actually fully edit a **Blueprint Spline Component** or **Blueprint Spline Mesh Component** in the viewport of the Blueprint Editor using the exact same methods and tools.
+[/REGION]
+
 
 ## Steps
 
@@ -55,17 +64,19 @@ This how-to will step through editing **Blueprint Spline Component** curves whil
 
 	![](PlaceActor.png)
 
-1. Select Actor with Spline Component in your level. If you just placed it into the level, it will already be selected.
+1. The actor you just placed should be selected. If it isn't, select any Actor in your level that has a Spline Component.
 
-	If your Blueprint Actor has no other visible Components (Meshes, billboards, etc...) it may be difficult to select, however, you can always select it from the **World Outliner**.
+	If your Blueprint Actor has no other visible Components (Meshes, billboards, etc.), it may be difficult to select. However, you can always select it from the **World Outliner**.
 
 	[REGION:note]
-	If you are using a **Spline Mesh Component** you'll need to enable **Allow Spline Editing Per Instance** on the Blueprint in the Blueprint Editor to see its **Curve Points**.
+	If you are using a **Spline Mesh Component**, you'll need to enable **Allow Spline Editing Per Instance** on the Blueprint in the Blueprint Editor to see its **Curve Points**.
 	[/REGION]
 	
 1. Select a **Curve Point** on the spline, and using the Move, Rotate, and Scale tools, manipulate the Curve Points and Tangent Points into the desired shape.
 
-	![](FullEdit.gif)(convert:false)
+	[REGION:asyncgif]
+	![](FullEdit.gif)
+	[/REGION]
 
 	You can also use the following features to edit the curve:
 	
@@ -73,7 +84,7 @@ This how-to will step through editing **Blueprint Spline Component** curves whil
 	* To Duplicate a point, right-click the **Curve Point** you want to duplicate and select **Duplicate Spline Point**
 	* To Remove a point, right-click on a **Curve Point** and select **Delete Spline Point**
 	
-	For more information on all the tools available to edit Spline Components see [](Engine/BlueprintSplines/Reference/SplineEditorTool).
+	For more information on all the tools available to edit Spline Components, see [](Engine/BlueprintSplines/Reference/SplineEditorTool).
 	
 	[REGION:note]
 	Once you start editing the Spline Component, the **Override Construction Script** property will be set to **True** to prevent it from overwriting your work. 
@@ -81,7 +92,7 @@ This how-to will step through editing **Blueprint Spline Component** curves whil
 	
 ## Result
 
-While above, we are just editing a spline that doesn't affect anything, there are a number of uses for this data. See [](Resources/ContentExamples/Blueprint_Splines) documentation and the **Content Examples** project for additional information.
+While the spline we have been editing doesn't affect anything, there are a number of uses for this data. See [](Resources/ContentExamples/Blueprint_Splines) documentation and the **Content Examples** project for additional information.
 
 
 
diff --git a/Engine/Documentation/Source/Engine/BlueprintSplines/HowTo/EditSplineInConstScript/EditSplineInConstScript.INT.udn b/Engine/Documentation/Source/Engine/BlueprintSplines/HowTo/EditSplineInConstScript/EditSplineInConstScript.INT.udn
--- a/Engine/Documentation/Source/Engine/BlueprintSplines/HowTo/EditSplineInConstScript/EditSplineInConstScript.INT.udn
+++ b/Engine/Documentation/Source/Engine/BlueprintSplines/HowTo/EditSplineInConstScript/EditSplineInConstScript.INT.udn
@@ -6,6 +6,8 @@ Navigation:topic
 parent:Engine/BlueprintSplines/HowTo
 order:3
 social-image:Splines_social.png
+version:4.13
+skilllevel:Intermediate
 
 [VAR:Topic]
 [OBJECT:Topic]
diff --git a/Engine/Documentation/Source/Engine/BlueprintSplines/HowTo/HowToLanding.INT.udn b/Engine/Documentation/Source/Engine/BlueprintSplines/HowTo/HowToLanding.INT.udn
--- a/Engine/Documentation/Source/Engine/BlueprintSplines/HowTo/HowToLanding.INT.udn
+++ b/Engine/Documentation/Source/Engine/BlueprintSplines/HowTo/HowToLanding.INT.udn
@@ -1,4 +1,4 @@
-Availability:Docs
+Availability:Public
 Title:Blueprint Splines How-To's
 Crumbs: %ROOT%, Engine
 Description:Various How-To's for Blueprint Spline Components and Blueprint Spline Mesh Components.
@@ -6,6 +6,9 @@ Navigation:topic
 parent:Engine/BlueprintSplines
 order:3
 social-image:Splines_social.png
+version:4.13
+type:landing
+tags:Splines
 
 [VAR:Topic]
 [OBJECT:Topic]
@@ -50,7 +53,7 @@ Below are the how-to's available for Blueprint Spline Components and Blueprint S
 
 ## How-To's
 
-[DIR(output:"listbutton" parent:"Engine/BlueprintSplines/HowTo" org:"hierarchy" end:"1")]
+[DIR(output:"fancy" parent:"Engine/BlueprintSplines/HowTo" org:"hierarchy" end:"1")]
 
 
 
diff --git a/Engine/Documentation/Source/Engine/BlueprintSplines/Overview/BPSplineOverview.INT.udn b/Engine/Documentation/Source/Engine/BlueprintSplines/Overview/BPSplineOverview.INT.udn
--- a/Engine/Documentation/Source/Engine/BlueprintSplines/Overview/BPSplineOverview.INT.udn
+++ b/Engine/Documentation/Source/Engine/BlueprintSplines/Overview/BPSplineOverview.INT.udn
@@ -1,4 +1,4 @@
-Availability:Docs
+Availability:Public
 Title:Blueprint Spline Components Overview
 Crumbs: %ROOT%, Engine
 Description:A brief overview of Blueprint Spline Components and Blueprint Spline Mesh Components.
@@ -11,6 +11,9 @@ related:Engine/BlueprintSplines/HowTo
 related:Engine/BlueprintSplines/Reference/ComponentProperties
 related:Engine/BlueprintSplines/Reference/MeshComponentProperties
 related:Engine/BlueprintSplines/Reference/SplineEditorTool
+version:4.13
+type:overview
+tags:Splines
 
 [VAR:Topic]
 [OBJECT:Topic]
@@ -52,7 +55,7 @@ related:Engine/BlueprintSplines/Reference/SplineEditorTool
 [/VAR]
 
 
-At their core, **Blueprint Spline Components** are just a path for you to use for positional data. You can use it to move **Actors** (or other **Components**) around the world or place a series of **Actors** (or other **Components**) along the spline. They are fully editable in Blueprints with the ability to add/remove/duplicate Spline Points, change their tangent types, even animate them on tick. Further, they are also editable in the Level Editor with the capability to pass back those edits to their **Blueprint Construction Script**.
+At their core, a **Blueprint Spline Component** is just a path for you to define and use positional data. You can use it to move **Actors** (or other **Components**) around the world, or place a series of **Actors** (or other **Components**) along the spline. They are fully editable in the Blueprint Viewport and in the Level Editor, with the ability to add/remove/duplicate Spline Points, change their tangent types, and even animate them on tick. Further, they are also editable using the **Blueprint Construction Script**, taking in edits made in the Blueprint Viewport or Level Editor, and mofifying them further.
 
 ![](BPSC_1.png)
 
@@ -60,6 +63,6 @@ At their core, **Blueprint Spline Components** are just a path for you to use fo
 
 ![](BPSMC_1.png)
 
-While their use cases are very different, they share the same procedures for adding them to a Blueprint, and editing tools.
+While their use cases are very different, they share the same procedures for adding them to a Blueprint, and use the same editing tools.
 
 
diff --git a/Engine/Documentation/Source/Engine/BlueprintSplines/Reference/ComponentProperties/SplineComponentReference.INT.udn b/Engine/Documentation/Source/Engine/BlueprintSplines/Reference/ComponentProperties/SplineComponentReference.INT.udn
--- a/Engine/Documentation/Source/Engine/BlueprintSplines/Reference/ComponentProperties/SplineComponentReference.INT.udn
+++ b/Engine/Documentation/Source/Engine/BlueprintSplines/Reference/ComponentProperties/SplineComponentReference.INT.udn
@@ -1,4 +1,4 @@
-Availability:Docs
+Availability:Public
 Title:Blueprint Spline Component Property Reference
 Crumbs: %ROOT%, Engine
 Description:Contains a reference of all properties available in the editor for Blueprint Spline Components.
@@ -6,6 +6,9 @@ Navigation:topic
 parent:Engine/BlueprintSplines
 order:1
 social-image:Splines_social.png
+version:4.13
+type:reference
+tags:Splines
 
 [VAR:Topic]
 [OBJECT:Topic]
@@ -50,12 +53,12 @@ social-image:Splines_social.png
 
 [VAR:spacer2]![](%ROOT%/spacer.png)(w:40 convert:false)[/VAR]
 
-This page contains a reference listing of properties available on **Spline Components**. The properties displayed will be slightly different if the **Blueprint Spline Component** is selected in the **Blueprint Editor**, is a **Root Component**, or if the component is selected in the **Level Editor**.
+This page contains a reference listing of properties available on **Blueprint Spline Components**. The properties displayed will be slightly different if the **Blueprint Spline Component** is selected in the **Blueprint Editor**, is a **Root Component**, or if the component is selected in the **Level Editor**.
 
 ## Properties
 
 [EXCERPT:PropTransformCat]
-[OBJECT:modulesection]
+[OBJECT:modulesectionexpanded]
 	[PARAMLITERAL:id]
 		TransformCat
 	[/PARAMLITERAL]
@@ -85,7 +88,7 @@ This page contains a reference listing of properties available on **Spline Compo
 
 
 [EXCERPT:PropSocketsCat]
-[OBJECT:modulesection]
+[OBJECT:modulesectionexpanded]
 	[PARAMLITERAL:id]
 		SocketsCat
 	[/PARAMLITERAL]
@@ -104,7 +107,7 @@ This page contains a reference listing of properties available on **Spline Compo
 [REGION:raw]
 | Property | %spacer% | Description |
 | --- | :---: | :--- |
-| **Parent Socket** || When a this component is the **Child** of a **Skeletal Mesh Component** (or **Static Mesh Component** with a **Socket**) you can specify a **Socket** or **Joint** to attach this component to. |
+| **Parent Socket** || When this component is the **Child** of a **Skeletal Mesh Component** (or **Static Mesh Component** with a **Socket**), you can specify a **Socket** or **Joint** to attach this component to. |
 [/REGION]
 
 [/VAR]
@@ -187,168 +190,44 @@ This page contains a reference listing of properties available on **Spline Compo
 
 
 
-[OBJECT:modulesection]
-	[PARAMLITERAL:id]
-		PhysicsCat
-	[/PARAMLITERAL]
-	[PARAM:heading]
-		Physics
-	[/PARAM]
-	[PARAM:content]
-		%Engine/Physics/PhysicsBodies/Reference:c3%
-	[/PARAM]
-[/OBJECT]
-
-
-
-
-[OBJECT:modulesection]
-	[PARAMLITERAL:id]
-		CollisionCat
-	[/PARAMLITERAL]
-	[PARAM:heading]
-		Collision
-	[/PARAM]
-	[PARAM:content]
-		%Engine/Physics/PhysicsBodies/Reference:p2%
-	[/PARAM]
-[/OBJECT]
-
-
-
-
 
-[EXCERPT:PropLightingCat]
-[OBJECT:modulesection]
-	[PARAMLITERAL:id]
-		LightingCat
-	[/PARAMLITERAL]
-	[PARAM:heading]
-		Lighting
-	[/PARAM]
-	[PARAM:content]
-		%c5%
-	[/PARAM]
-[/OBJECT]
-
-[VAR:c5]
-
-![](LightingCatRef.png)
-
-[REGION:raw]
-| Property | %spacer% | Description |
-| --- | :---: | :--- |
-| **Cast Shadow** || %doxygen:UPrimitiveComponent::CastShadow% |
-[/REGION]
-
-###Advanced
-
-![](LightingCatRefAdv.png)
-
-[REGION:raw]
-| Property | %spacer% | Description |
-| --- | :---: | :--- |
-| **Affect Dynamic Indirect Lighting** || %doxygen:UPrimitiveComponent::bAffectDynamicIndirectLighting% |
-| **Affect Distance Field Lighting** || %doxygen:UPrimitiveComponent::bAffectDistanceFieldLighting% |
-| **Cast Dynamic Shadow** || %doxygen:UPrimitiveComponent::bCastDynamicShadow% |
-| **Cast Static Shadow** || %doxygen:UPrimitiveComponent::bCastStaticShadow% |
-| **Cast Volumetric Translucent Shadow** || %doxygen:UPrimitiveComponent::bCastVolumetricTranslucentShadow% |
-| **Self Shadow Only** || %doxygen:UPrimitiveComponent::bSelfShadowOnly% |
-| **Cast Far Shadow** || %doxygen:UPrimitiveComponent::bCastFarShadow% |
-| **Cast Inset Shadow** || %doxygen:UPrimitiveComponent::bCastInsetShadow% |
-| **Cast Cinematic Shadow** || %doxygen:UPrimitiveComponent::bCastCinematicShadow% |
-| **Cast Hidden Shadow** || %doxygen:UPrimitiveComponent::bCastHiddenShadow% |
-| **Cast Shadow as Two Sided** || %doxygen:UPrimitiveComponent::bCastShadowAsTwoSided% |
-| **Light as if Static** || %doxygen:UPrimitiveComponent::bLightAsIfStatic% |
-| **Light Attachments as Group** || %doxygen:UPrimitiveComponent::bLightAttachmentsAsGroup% |
-| **Indirect Lighting Cache Quality** || %doxygen:UPrimitiveComponent::IndirectLightingCacheQuality% |
-| **Single Sample Shadow From Stationary Lights** || %doxygen:UPrimitiveComponent::bSingleSampleShadowFromStationaryLights% |
-| **Lighting Channles** || %doxygen:UPrimitiveComponent::LightingChannels% |
-[/REGION]
-
-[/VAR]
-[/EXCERPT]
-
-
-
-[EXCERPT:PropRenderingCat]
-[OBJECT:modulesection]
+[EXCERPT:PropSelectedPointsCat]
+[OBJECT:modulesectionexpanded]
 	[PARAMLITERAL:id]
-		RenderingCat
+		SelectedPointsCat
 	[/PARAMLITERAL]
 	[PARAM:heading]
-		Rendering
+		Selected Points
 	[/PARAM]
 	[PARAM:content]
-		%c6%
+		%sp3%
 	[/PARAM]
 [/OBJECT]
 
-[VAR:c6]
-
-![](RenderingCatRef.png)
-
-[REGION:raw]
-| Property | %spacer% | Description |
-| --- | :---: | :--- |
-| **Visible** || %doxygen:USceneComponent::bVisible% |
-| **Hidden in Game** ||  %doxygen:USceneComponent::bHiddenInGame%|
-[/REGION]
-
-###Advanced
+[VAR:sp3]
 
-![](RenderingCatRefAdv.png)
+![](BPS_SelectedPointsRef.png)
 
 [REGION:raw]
 | Property | %spacer% | Description |
 | --- | :---: | :--- |
-| **Texture Streaming** || [INCLUDE:#TextureStreamingProp] |
-| **LOD** || [INCLUDE:#LODProp] |
-| **Render In Main Pass** || %doxygen:UPrimitiveComponent::bRenderInMainPass% |
-| **Receives Decals** || %doxygen:UPrimitiveComponent::bReceivesDecals% |
-| **Owner No See** || %doxygen:UPrimitiveComponent::bOwnerNoSee% |
-| **Only Owner See** || %doxygen:UPrimitiveComponent::bOnlyOwnerSee% |
-| **Treat As Background for Occlusion** || %doxygen:UPrimitiveComponent::bTreatAsBackgroundForOcclusion% |
-| **Use As Occluder** || %doxygen:UPrimitiveComponent::bUseAsOccluder% |
-| **Render CustomDepth Pass** || %doxygen:UPrimitiveComponent::bRenderCustomDepth% |
-| **CustomDepth Stencil Value** || %doxygen:UPrimitiveComponent::CustomDepthStencilValue% |
-| **Translucency Sort Priority** || Translucent objects with a lower sort priority draw behind objects with a higher priority. Translucent objects with the same priority are rendered from back-to-front based on their bounds origin. |
-| **Lpv Bias Multiplier** || %doxygen:UPrimitiveComponent::LpvBiasMultiplier% |
-| **Bounds Scale** || %doxygen:UPrimitiveComponent::BoundsScale% |
-| **Use Attach Parent Bound** || %doxygen:USceneComponent::bUseAttachParentBound% |
+| **Input Key** || This is the index of the selected **Curve Point** on the curve. |
+| **Position** || This is the location of the selected **Curve Point** in local space. |
+| **Arrive Tangent** || This is the vector that defines the tangent of the curve as it approaches the selected **Curve Point**. |
+| **Leave Tangent** || This is the vector that defines the tangent of the curve as it departs the selected **Curve Point**. |
+| **Rotation** || Rotation can be applied to a **Curve Point** to modify its tangents. This value is applied separately from the tangents, so both can be modified to create the desired results. |
+| **Scale** || Scale can be applied to a **Curve Point** to modify its tangents. This value is applied separately from the tangents, so both can be modified to create the desired results. |
+| **Type** || [INCLUDE:#PointsType] |
 [/REGION]
 
 [/VAR]
 [/EXCERPT]
 
-<!--
-[EXCERPT:TextureStreamingProp]
-
-[REGION:raw]
-| Property | %spacer2% | Description |
-| --- | :---: | :--- |
-| **Force Mip Streaming** || %doxygen:UPrimitiveComponent::bForceMipStreaming% |
-[/REGION]
 
-[/EXCERPT]
-
-[EXCERPT:LODProp]
-
-[REGION:raw]
-| Property | %spacer2% | Description |
-| --- | :---: | :--- |
-| **Min Draw Distance** || %doxygen:UPrimitiveComponent::MinDrawDistance% |
-| **Desired Max Draw Distance** || %doxygen:UPrimitiveComponent::LDMaxDrawDistance% |
-| **Current Max Draw Distance** || %doxygen:UPrimitiveComponent::CachedMaxDrawDistance% |
-| **Allow Cull Distance Volume** || %doxygen:UPrimitiveComponent::bAllowCullDistanceVolume% |
-| **Detail Mode** || %doxygen:USceneComponent::DetailMode% |
-[/REGION]
 
-[/EXCERPT]
--->
 
 [EXCERPT:PropTagsCat]
-[OBJECT:modulesection]
+[OBJECT:modulesectionexpanded]
 	[PARAMLITERAL:id]
 		TagsCat
 	[/PARAMLITERAL]
@@ -376,7 +255,7 @@ This page contains a reference listing of properties available on **Spline Compo
 
 
 [EXCERPT:PropActivationCat]
-[OBJECT:modulesection]
+[OBJECT:modulesectionexpanded]
 	[PARAMLITERAL:id]
 		ActivationCat
 	[/PARAMLITERAL]
@@ -404,7 +283,7 @@ This page contains a reference listing of properties available on **Spline Compo
 
 
 [EXCERPT:PropEventsCat]
-[OBJECT:modulesection]
+[OBJECT:modulesectionexpanded]
 	[PARAMLITERAL:id]
 		EventsCat
 	[/PARAMLITERAL]
@@ -443,31 +322,17 @@ This page contains a reference listing of properties available on **Spline Compo
 [/EXCERPT]
 
 
+<!--
+[EXCERPT:PointsType]
+Defines the type of curve at the selected **Curve Point**.
 
-[EXCERPT:PropMobileCat]
-[OBJECT:modulesection]
-	[PARAMLITERAL:id]
-		MobileCat
-	[/PARAMLITERAL]
-	[PARAM:heading]
-		Mobile
-	[/PARAM]
-	[PARAM:content]
-		%c10%
-	[/PARAM]
-[/OBJECT]
-
-[VAR:c10]
-
-![](MobileCatRef.png)
-
-[REGION:raw]
-| Property | %spacer% | Description |
+| Type | %spacer2% | Description |
 | --- | :---: | :--- |
-| **Receive Combined Static and CSM Shadows from Stationary Lights** || %doxygen:UPrimitiveComponent::bReceiveCombinedCSMAndStaticShadowsFromStationaryLights%  |
-[/REGION]
-
-[/VAR]
+| **Linear** || From the selected **Curve Point** to the next, the curve will be straight. This affects the tangent coming into the selected **Curve Point** as well as the tanget departing the next. |
+| **Curve** || Default value. The locations of the **Curve Points** (before and after the selected one) define the tangents of the selected **Curve Point**. |
+| **Constant** || Like **Linear** but **doesn't** affect the tangent coming into the selected **Curve Point** as well as the tanget departing the next. |
+| **CurveClamped** || Like **Curve** but it clamps the tangents of the curve. |
+| **CurveCustomTangent** || If you adjust any tangent on a **Curve Point**, it is set to this value. |
 [/EXCERPT]
-
+-->
 
diff --git a/Engine/Documentation/Source/Engine/BlueprintSplines/Reference/MeshComponentProperties/SplineMeshComponentReference.INT.udn b/Engine/Documentation/Source/Engine/BlueprintSplines/Reference/MeshComponentProperties/SplineMeshComponentReference.INT.udn
--- a/Engine/Documentation/Source/Engine/BlueprintSplines/Reference/MeshComponentProperties/SplineMeshComponentReference.INT.udn
+++ b/Engine/Documentation/Source/Engine/BlueprintSplines/Reference/MeshComponentProperties/SplineMeshComponentReference.INT.udn
@@ -1,4 +1,4 @@
-Availability:Docs
+Availability:Public
 Title:Blueprint Spline Mesh Component Property Reference
 Crumbs: %ROOT%, Engine
 Description:Contains a reference of all properties available in the editor for Blueprint Spline Mesh Components.
@@ -6,6 +6,9 @@ Navigation:topic
 parent:Engine/BlueprintSplines
 order:1
 social-image:Splines_social.png
+version:4.13
+type:reference
+tags:Splines
 
 [VAR:Topic]
 [OBJECT:Topic]
@@ -50,7 +53,7 @@ social-image:Splines_social.png
 
 [VAR:spacer2]![](%ROOT%/spacer.png)(w:40 convert:false)[/VAR]
 
-This page contains a reference listing of properties available on **Spline Components**. The properties displayed will be slightly different if the **Blueprint Spline Component** is selected in the **Blueprint Editor**, is a **Root Component**, or if the component is selected in the **Level Editor**.
+This page contains a reference listing of properties available on **Blueprint Spline Mesh Components**. The properties displayed will be slightly different if the **Blueprint Spline Mesh Component** is selected in the **Blueprint Editor**, is a **Root Component**, or if the component is selected in the **Level Editor**.
 
 ## Properties
 
@@ -178,7 +181,7 @@ This page contains a reference listing of properties available on **Spline Compo
 
 
 
-[OBJECT:modulesection]
+[OBJECT:modulesectionexpanded]
 	[PARAMLITERAL:id]
 		PhysicsCat
 	[/PARAMLITERAL]
@@ -190,7 +193,7 @@ This page contains a reference listing of properties available on **Spline Compo
 	[/PARAM]
 [/OBJECT]
 
-[OBJECT:modulesection]
+[OBJECT:modulesectionexpanded]
 	[PARAMLITERAL:id]
 		CollisionCat
 	[/PARAMLITERAL]
@@ -202,9 +205,142 @@ This page contains a reference listing of properties available on **Spline Compo
 	[/PARAM]
 [/OBJECT]
 
-[INCLUDE:Engine/BlueprintSplines/Reference/ComponentProperties#PropLightingCat]
 
-[INCLUDE:Engine/BlueprintSplines/Reference/ComponentProperties#PropRenderingCat]
+
+
+[EXCERPT:PropLightingCat]
+[OBJECT:modulesectionexpanded]
+	[PARAMLITERAL:id]
+		LightingCat
+	[/PARAMLITERAL]
+	[PARAM:heading]
+		Lighting
+	[/PARAM]
+	[PARAM:content]
+		%c5%
+	[/PARAM]
+[/OBJECT]
+
+[VAR:c5]
+
+![](LightingCatRef.png)
+
+[REGION:raw]
+| Property | %spacer% | Description |
+| --- | :---: | :--- |
+| **Cast Shadow** || %doxygen:UPrimitiveComponent::CastShadow% |
+[/REGION]
+
+###Advanced
+
+![](LightingCatRefAdv.png)
+
+[REGION:raw]
+| Property | %spacer% | Description |
+| --- | :---: | :--- |
+| **Affect Dynamic Indirect Lighting** || %doxygen:UPrimitiveComponent::bAffectDynamicIndirectLighting% |
+| **Affect Distance Field Lighting** || %doxygen:UPrimitiveComponent::bAffectDistanceFieldLighting% |
+| **Cast Dynamic Shadow** || %doxygen:UPrimitiveComponent::bCastDynamicShadow% |
+| **Cast Static Shadow** || %doxygen:UPrimitiveComponent::bCastStaticShadow% |
+| **Cast Volumetric Translucent Shadow** || %doxygen:UPrimitiveComponent::bCastVolumetricTranslucentShadow% |
+| **Self Shadow Only** || %doxygen:UPrimitiveComponent::bSelfShadowOnly% |
+| **Cast Far Shadow** || %doxygen:UPrimitiveComponent::bCastFarShadow% |
+| **Cast Inset Shadow** || %doxygen:UPrimitiveComponent::bCastInsetShadow% |
+| **Cast Cinematic Shadow** || %doxygen:UPrimitiveComponent::bCastCinematicShadow% |
+| **Cast Hidden Shadow** || %doxygen:UPrimitiveComponent::bCastHiddenShadow% |
+| **Cast Shadow as Two Sided** || %doxygen:UPrimitiveComponent::bCastShadowAsTwoSided% |
+| **Light as if Static** || %doxygen:UPrimitiveComponent::bLightAsIfStatic% |
+| **Light Attachments as Group** || %doxygen:UPrimitiveComponent::bLightAttachmentsAsGroup% |
+| **Indirect Lighting Cache Quality** || %doxygen:UPrimitiveComponent::IndirectLightingCacheQuality% |
+| **Single Sample Shadow From Stationary Lights** || %doxygen:UPrimitiveComponent::bSingleSampleShadowFromStationaryLights% |
+| **Lighting Channels** || %doxygen:UPrimitiveComponent::LightingChannels% |
+[/REGION]
+
+[/VAR]
+[/EXCERPT]
+
+
+
+[EXCERPT:PropRenderingCat]
+[OBJECT:modulesectionexpanded]
+	[PARAMLITERAL:id]
+		RenderingCat
+	[/PARAMLITERAL]
+	[PARAM:heading]
+		Rendering
+	[/PARAM]
+	[PARAM:content]
+		%c6%
+	[/PARAM]
+[/OBJECT]
+
+[VAR:c6]
+
+![](RenderingCatRef.png)
+
+[REGION:raw]
+| Property | %spacer% | Description |
+| --- | :---: | :--- |
+| **Visible** || %doxygen:USceneComponent::bVisible% |
+| **Hidden in Game** ||  %doxygen:USceneComponent::bHiddenInGame%|
+[/REGION]
+
+###Advanced
+
+![](RenderingCatRefAdv.png)
+
+[REGION:raw]
+| Property | %spacer% | Description |
+| --- | :---: | :--- |
+| **Texture Streaming** || [INCLUDE:#TextureStreamingProp] |
+| **LOD** || [INCLUDE:#LODProp] |
+| **Render In Main Pass** || %doxygen:UPrimitiveComponent::bRenderInMainPass% |
+| **Receives Decals** || %doxygen:UPrimitiveComponent::bReceivesDecals% |
+| **Owner No See** || %doxygen:UPrimitiveComponent::bOwnerNoSee% |
+| **Only Owner See** || %doxygen:UPrimitiveComponent::bOnlyOwnerSee% |
+| **Treat As Background for Occlusion** || %doxygen:UPrimitiveComponent::bTreatAsBackgroundForOcclusion% |
+| **Use As Occluder** || %doxygen:UPrimitiveComponent::bUseAsOccluder% |
+| **Render CustomDepth Pass** || %doxygen:UPrimitiveComponent::bRenderCustomDepth% |
+| **CustomDepth Stencil Value** || %doxygen:UPrimitiveComponent::CustomDepthStencilValue% |
+| **Translucency Sort Priority** || Translucent objects with a lower sort priority draw behind objects with a higher priority. Translucent objects with the same priority are rendered from back-to-front based on their bounds origin. |
+| **Lpv Bias Multiplier** || %doxygen:UPrimitiveComponent::LpvBiasMultiplier% |
+| **Bounds Scale** || %doxygen:UPrimitiveComponent::BoundsScale% |
+| **Use Attach Parent Bound** || %doxygen:USceneComponent::bUseAttachParentBound% |
+[/REGION]
+
+[/VAR]
+[/EXCERPT]
+
+
+
+<!--
+[EXCERPT:TextureStreamingProp]
+
+[REGION:raw]
+| Property | %spacer2% | Description |
+| --- | :---: | :--- |
+| **Force Mip Streaming** || %doxygen:UPrimitiveComponent::bForceMipStreaming% |
+[/REGION]
+
+[/EXCERPT]
+
+[EXCERPT:LODProp]
+
+[REGION:raw]
+| Property | %spacer2% | Description |
+| --- | :---: | :--- |
+| **Min Draw Distance** || %doxygen:UPrimitiveComponent::MinDrawDistance% |
+| **Desired Max Draw Distance** || %doxygen:UPrimitiveComponent::LDMaxDrawDistance% |
+| **Current Max Draw Distance** || %doxygen:UPrimitiveComponent::CachedMaxDrawDistance% |
+| **Allow Cull Distance Volume** || %doxygen:UPrimitiveComponent::bAllowCullDistanceVolume% |
+| **Detail Mode** || %doxygen:USceneComponent::DetailMode% |
+[/REGION]
+
+[/EXCERPT]
+-->
+
+
+
 
 [INCLUDE:Engine/BlueprintSplines/Reference/ComponentProperties#PropTagsCat]
 
@@ -212,4 +348,28 @@ This page contains a reference listing of properties available on **Spline Compo
 
 [INCLUDE:Engine/BlueprintSplines/Reference/ComponentProperties#PropEventsCat]
 
-[INCLUDE:Engine/BlueprintSplines/Reference/ComponentProperties#PropMobileCat]
+[EXCERPT:PropMobileCat]
+[OBJECT:modulesectionexpanded]
+	[PARAMLITERAL:id]
+		MobileCat
+	[/PARAMLITERAL]
+	[PARAM:heading]
+		Mobile
+	[/PARAM]
+	[PARAM:content]
+		%c10%
+	[/PARAM]
+[/OBJECT]
+
+[VAR:c10]
+
+![](MobileCatRef.png)
+
+[REGION:raw]
+| Property | %spacer% | Description |
+| --- | :---: | :--- |
+| **Receive Combined Static and CSM Shadows from Stationary Lights** || %doxygen:UPrimitiveComponent::bReceiveCombinedCSMAndStaticShadowsFromStationaryLights%  |
+[/REGION]
+
+[/VAR]
+[/EXCERPT]
\ No newline at end of file
diff --git a/Engine/Documentation/Source/Engine/BlueprintSplines/Reference/SplineEditorTool/SplineEditorToolReference.INT.udn b/Engine/Documentation/Source/Engine/BlueprintSplines/Reference/SplineEditorTool/SplineEditorToolReference.INT.udn
--- a/Engine/Documentation/Source/Engine/BlueprintSplines/Reference/SplineEditorTool/SplineEditorToolReference.INT.udn
+++ b/Engine/Documentation/Source/Engine/BlueprintSplines/Reference/SplineEditorTool/SplineEditorToolReference.INT.udn
@@ -1,4 +1,4 @@
-Availability:Docs
+Availability:Public
 Title:Blueprint Spline Editor Tool Reference
 Crumbs: %ROOT%, Engine
 Description:A quick reference of the tools and procedures for editing Blueprint Spline Components while in the Level Editor.
@@ -6,6 +6,13 @@ Navigation:topic
 parent:Engine/BlueprintSplines
 order:1
 social-image:Splines_social.png
+version:4.13
+type:reference
+reference-image:MovePoint.gif
+reference-image:MoveTangent.gif
+reference-image:RotatePoint.gif
+reference-image:ScalePoint.gif
+tags:Splines
 
 [VAR:Topic]
 [OBJECT:Topic]
@@ -50,18 +57,18 @@ social-image:Splines_social.png
 
 [VAR:spacer2]![](%ROOT%/spacer.png)(w:40 convert:false)[/VAR]
 
-This page contains information on tools and procedures for editing **Blueprint Spline Components** and **Blueprint Spline Mesh Components** in the Level Editor.
+This page contains information on tools and procedures for editing **Blueprint Spline Components** and **Blueprint Spline Mesh Components** in the Level Editor. As of 4.13, **Blueprint Spline Components** can also be edited in the **Blueprint Editor Viewport** using all the tools and procedures below.
 
 ## Tools
 
-With a Spline Point (selecting a Spline Curve selects the last point on it) selected in the editor, you can use the standard manipulation tools (move, rotate, and scale) to adjust positions and tangents of the points.
+With a Spline Point (selecting a Spline Curve selects the last point on it) selected in the editor, you can use the standard manipulation tools (move, rotate, and scale) to select and adjust the positions and tangents of the spline curve's points.
 
 | Tool | Result |
 | --- | :--- |
-| **Move Selected Point** | ![](MovePoint.gif)(convert:false) |
-| **Move Selected Tangent** | ![](MoveTangent.gif)(convert:false) |
-| **Rotate** | ![](RotatePoint.gif)(convert:false) |
-| **Scale** | ![](ScalePoint.gif)(convert:false) |
+| **Move Selected Point** | [REGION:asyncgif]![](MovePoint.gif)[/REGION] |
+| **Move Selected Tangent** | [REGION:asyncgif]![](MoveTangent.gif)[/REGION] |
+| **Rotate** | [REGION:asyncgif]![](RotatePoint.gif)[/REGION] |
+| **Scale** | [REGION:asyncgif]![](ScalePoint.gif)[/REGION] |
 
 
 ## Procedures
@@ -70,50 +77,54 @@ There are a few menus accessable by right-clicking on the spline either on a poi
 
 ### Right-Click on a Spline Point
 
-Right-clicking on a **point** will give you the following menu and options.
+Right-clicking on a point will give you the following menu and options:
 
 ![](Spline_RC_Point.png)
 
 [REGION:raw]
-| Option | %spacer% | Description |
-| --- | :---: | :--- |
-| **Delete Spline Point** || Removes the selected Spline Point. |
-| **Duplicate Spline Point** || Duplicates the currently selected Spline Point. |
-| **Spline Point Type** || [INCLUDE:#SplinePointType] |
-| **Reset to Automatic Tangent** || [INCLUDE:#ResetToAutomaticTanget] |
-| **Reset to Default** || Resets the Spline to its archetype default. |
-| **Visualize Roll and Scale** || Whether or not the spline should visualize roll and scale in the editor. |
+| Option | Description |
+| --- | :--- |
+| **Delete Spline Point** | Removes the selected Spline Point. |
+| **Duplicate Spline Point** | Duplicates the currently selected Spline Point. |
+| **Spline Point Type** | [INCLUDE:#SplinePointType] |
+| **Reset to Automatic Tangent** | [INCLUDE:#ResetToAutomaticTanget] |
+| **Reset to Default** | Resets the Spline to its archetype default. |
+| **Visualize Roll and Scale** | Whether or not the spline should visualize roll and scale in the editor. |
+| **Allow Discontinous Splines** | Toggles whether the visualization allows **Arrive** and **Leave** tangents to be set separately. |
 [/REGION]
 
 <!--
 [EXCERPT:SplinePointType]
 
-Changes the interpolation style for the selected Spline Point. The Spline Curve associated with the Spline Curve point is displayed in red (by default). It will be the curve affected by changing this. The default **Point Type** is **Curve**. Below are examples of what 
+Changes the interpolation style for the selected Spline Point. The Spline Curve associated with the Spline Curve point is displayed in red (by default). It will be the curve affected by changing this. The default **Point Type** is **Curve**. Below are examples of what the different types look like:
 
 ![](Spline_RC_Point_Type.png)
 
-| Option | %spacer2% | Description |
-| --- | :---: | :--- |
-| **Curve** || ![](PointTypeCurve.png)(convert:false) |
-| **Linear** || ![](PointTypeLinear.png)(convert:false) |
-| **Constant** || ![](PointTypeConstant.png)(convert:false) |
+[REGION:transparenttable]
+| Option | Description |
+| --- | --- |
+| **Curve** | ![](PointTypeCurve.PNG)(convert:false) |
+| **Linear** | ![](PointTypeLinear.PNG)(convert:false) |
+| **Constant** | ![](PointTypeConstant.PNG)(convert:false) |
+[/REGION]
 [/EXCERPT]
 
 [EXCERPT:ResetToAutomaticTanget]
 
 ![](Spline_RC_Point_Reset.png)
 
-| Option | %spacer2% | Description |
-| --- | :---: | :--- |
-| **Unclamped Tangent** || Resets the tangents of the selected **Curve Point** to their **Unclamped Tangent** values. |
-| **Clamped Tangent** || Resets the tangents of the selected **Curve Point** to their **Clamped Tangent** values. |
-
+[REGION:transparenttable]
+| Option | Description |
+| --- | :--- |
+| **Unclamped Tangent** | Resets the tangents of the selected **Curve Point** to their **Unclamped Tangent** values. |
+| **Clamped Tangent** | Resets the tangents of the selected **Curve Point** to their **Clamped Tangent** values. |
+[/REGION]
 [/EXCERPT]
 -->
 
 ### Right-Click on a Spline Curve
 
-Right-clicking on a **curve** will give you the following menu and options.
+Right-clicking on a curve will give you the following menu and options:
 
 ![](Spline_RC_Curve.png)
 
@@ -122,7 +133,8 @@ Right-clicking on a **curve** will give you the following menu and options.
 | --- | :---: | :--- |
 | **Add Spline Point Here** || Adds a new Spline Point at the cursor's location. |
 | **Reset to Default** || Resets the Spline to its archetype default. |
-| **Visualize Roll and Scale** || Whether or not the spline should visualize roll and scale in the editor. |
+| **Visualize Roll and Scale** || Toggles whether or not the spline will visualize roll and scale in the editor. |
+| **Allow Discontinous Splines** || Toggles whether the visualization allows **Arrive** and **Leave** tangents to be set separately. |
 [/REGION]
 
 
diff --git a/Engine/Documentation/Source/Engine/BlueprintSplines/SplinesLanding.INT.udn b/Engine/Documentation/Source/Engine/BlueprintSplines/SplinesLanding.INT.udn
--- a/Engine/Documentation/Source/Engine/BlueprintSplines/SplinesLanding.INT.udn
+++ b/Engine/Documentation/Source/Engine/BlueprintSplines/SplinesLanding.INT.udn
@@ -1,48 +1,50 @@
-Availability:Docs
+Availability:Public
 Title:Blueprint Splines
 Crumbs: %ROOT%, Engine
 Description:Landing page for Blueprint Spline and Spline Mesh Components.
 Navigation:topic
 parent:Engine
-order:3
-social-image:Splines_social.png
+order:14
+social-image:Engine/BlueprintSplines/Overview/BPSC_1.png
 related: Engine/Landscape/Editing/Splines/
+version:4.13
+tags:Splines
 
 [VAR:Topic]
 [OBJECT:Topic]
 	[PARAM:image]
-		![%Engine/Splines:title%](Engine/Splines/Splines_topic.png)
+		![%Engine/BlueprintSplines:title%](Engine/BlueprintSplines/Splines_topic.png)
 	[/PARAM]
 	[PARAM:icon]
 		![](%ROOT%/editor_icon.png)(convert:false)
 	[/PARAM]
 	[PARAM:title]
-		%Engine/Splines:title%
+		%Engine/BlueprintSplines:title%
 	[/PARAM]
 	[PARAM:description]
-		%Engine/Splines:description%
+		%Engine/BlueprintSplines:description%
 	[/PARAM]
 	[PARAM:path]
-		[RELATIVE:Engine/Splines]
+		[RELATIVE:Engine/BlueprintSplines]
 	[/PARAM]
 [/OBJECT]
 [/VAR]
 [VAR:TopicCompact]
 [OBJECT:TopicCompact]
 	[PARAM:image]
-		![%Engine/Splines:title%](Engine/Splines/Splines_topic.png)
+		![%Engine/BlueprintSplines:title%](Engine/BlueprintSplines/Splines_topic.png)
 	[/PARAM]
 	[PARAM:icon]
 		![](%ROOT%/editor_icon.png)(convert:false)
 	[/PARAM]
 	[PARAM:title]
-		%Engine/Splines:title%
+		%Engine/BlueprintSplines:title%
 	[/PARAM]
 	[PARAM:description]
-		%Engine/Splines:description%
+		%Engine/BlueprintSplines:description%
 	[/PARAM]
 	[PARAM:path]
-		[RELATIVE:Engine/Splines]
+		[RELATIVE:Engine/BlueprintSplines]
 	[/PARAM]
 [/OBJECT]
 [/VAR]
@@ -51,9 +53,14 @@ related: Engine/Landscape/Editing/Splines/
 
 Below are links to the documentation pages that cover **Blueprint Spline Components** and **Blueprint Spline Mesh Components**.
 
-## Essentials
+## Starting Out
 
-[DIR(output:"listbutton" parent:"Engine/BlueprintSplines" org:"hierarchy" end:"1")]
+[DIR(output:"listbutton" parent:"Engine/BlueprintSplines" org:"hierarchy" end:"1" type:"overview,landing")]
+
+
+## Reference
+
+[DIR(output:"listbutton" parent:"Engine/BlueprintSplines" org:"hierarchy" end:"1" type:"reference")]
 
 
 ## Samples
diff --git a/Engine/Documentation/Source/Engine/Blueprints/BP_HowTo/ActorReference/ActorReference.INT.udn b/Engine/Documentation/Source/Engine/Blueprints/BP_HowTo/ActorReference/ActorReference.INT.udn
--- a/Engine/Documentation/Source/Engine/Blueprints/BP_HowTo/ActorReference/ActorReference.INT.udn
+++ b/Engine/Documentation/Source/Engine/Blueprints/BP_HowTo/ActorReference/ActorReference.INT.udn
@@ -5,7 +5,7 @@ Description:Using Actor Reference to move different Actors around in a level.
 Related: Gameplay/HowTo/ReferenceAssets/Blueprints
 version: 4.9
 skilllevel:Intermediate
-type:howto
+type:how-to
 parent:Engine/Blueprints/BP_HowTo
 order:1
 tags:Data Types
diff --git a/Engine/Documentation/Source/Engine/Blueprints/BP_HowTo/BlueprintHowTo.INT.udn b/Engine/Documentation/Source/Engine/Blueprints/BP_HowTo/BlueprintHowTo.INT.udn
--- a/Engine/Documentation/Source/Engine/Blueprints/BP_HowTo/BlueprintHowTo.INT.udn
+++ b/Engine/Documentation/Source/Engine/Blueprints/BP_HowTo/BlueprintHowTo.INT.udn
@@ -8,7 +8,7 @@ Related: Gameplay
 version: 4.9
 parent:Engine/Blueprints
 order:2
-type:howto
+type:landing
 tags:Blueprints
 
 [VAR:Topic]
@@ -56,7 +56,7 @@ When you are ready and have an understanding of the core aspects behind Blueprin
 
 
 ## Beginner
-[DIR(output:"fancy" type:"howto" skill:"Beginner" parent:"Engine/Blueprints/BP_HowTo" org:"hierarchy")]
+[DIR(output:"fancy" type:"how-to" skill:"Beginner" parent:"Engine/Blueprints/BP_HowTo" org:"hierarchy")]
 
 ## Intermediate
 [REGION:wrapper transparent]
diff --git a/Engine/Documentation/Source/Engine/Blueprints/Blueprints.INT.udn b/Engine/Documentation/Source/Engine/Blueprints/Blueprints.INT.udn
--- a/Engine/Documentation/Source/Engine/Blueprints/Blueprints.INT.udn
+++ b/Engine/Documentation/Source/Engine/Blueprints/Blueprints.INT.udn
@@ -6,6 +6,7 @@ version: 4.9
 parent:%ROOT%
 order:5
 social-image:bp_social_image.png
+type:landing
 
 [VAR:Topic]
 [OBJECT:Topic]
@@ -191,7 +192,7 @@ Blueprint-specific markup available in Unreal Engine's C++ implementation enable
 [/OBJECT]
 [/REGION]
 
-[DIR(output:"listbutton" parent:"Engine/Blueprints" type:"guide")]
+[DIR(output:"listbutton" parent:"Engine/Blueprints" type:"overview" org:"hierarchy")]
 
 
 <!--
diff --git a/Engine/Documentation/Source/Engine/Blueprints/TechnicalGuide/BlueprintsTechnicalGuide.INT.udn b/Engine/Documentation/Source/Engine/Blueprints/TechnicalGuide/BlueprintsTechnicalGuide.INT.udn
--- a/Engine/Documentation/Source/Engine/Blueprints/TechnicalGuide/BlueprintsTechnicalGuide.INT.udn
+++ b/Engine/Documentation/Source/Engine/Blueprints/TechnicalGuide/BlueprintsTechnicalGuide.INT.udn
@@ -4,7 +4,7 @@ Crumbs: %ROOT%, Programming
 Description:Technical guide for programmers working with Blueprints.
 version: 4.9
 Parent:Engine/Blueprints
-Type:guide
+type:overview
 
 Blueprints are a powerful new feature introduced in Unreal Engine 4.  
 
diff --git a/Engine/Documentation/Source/Engine/Blueprints/TechnicalGuide/NativizingBlueprints/Advanced/AdvancedNativeBlueprints.INT.udn b/Engine/Documentation/Source/Engine/Blueprints/TechnicalGuide/NativizingBlueprints/Advanced/AdvancedNativeBlueprints.INT.udn
--- a/Engine/Documentation/Source/Engine/Blueprints/TechnicalGuide/NativizingBlueprints/Advanced/AdvancedNativeBlueprints.INT.udn
+++ b/Engine/Documentation/Source/Engine/Blueprints/TechnicalGuide/NativizingBlueprints/Advanced/AdvancedNativeBlueprints.INT.udn
@@ -18,6 +18,9 @@ To nativize Blueprints while building with Unreal Automation too, add the `-nati
 Nativized blueprints are stored in a new, separate module, which means that the headers from gameplay module are now used outside their module. 
 You should check that all headers have declared all classes and all structs that are used, either as a included header, or as a forward declaration.
 
+Many projects have an "aggregate" header `[GameName]/[GameName].h` or `[GameName]/Public/[GameName].h`. If you add the path to this header to 
+`FilesToIncludeInModuleHeader` under `[BlueprintNativizationSettings]` in `DefaultEditor.ini`, it can fix many issues related to missing headers.
+
 There are some rules/restrictions that help make cooperation between manually written code and the nativized code more reliable and efficient. 
 
 The following implementations are not supported: 
diff --git a/Engine/Documentation/Source/Engine/Blueprints/UserGuide/MathNode/MathNode.INT.udn b/Engine/Documentation/Source/Engine/Blueprints/UserGuide/MathNode/MathNode.INT.udn
--- a/Engine/Documentation/Source/Engine/Blueprints/UserGuide/MathNode/MathNode.INT.udn
+++ b/Engine/Documentation/Source/Engine/Blueprints/UserGuide/MathNode/MathNode.INT.udn
@@ -4,7 +4,7 @@ Crumbs: %ROOT%, Engine/Blueprints, Engine/Blueprints/UserGuide
 Description:The math expression node allows you to type in a math expression and builds the appropriate sub-graph to create that expression.
 version: 4.9
 Parent:Engine/Blueprints
-Type:guide
+type:overview
 
 [VAR:Topic]
 [OBJECT:Topic]
diff --git a/Engine/Documentation/Source/Engine/Blueprints/UserGuide/PatchingNodes/MobilePatchUtilities.INT.udn b/Engine/Documentation/Source/Engine/Blueprints/UserGuide/PatchingNodes/MobilePatchUtilities.INT.udn
new file mode 100644
--- /dev/null
+++ b/Engine/Documentation/Source/Engine/Blueprints/UserGuide/PatchingNodes/MobilePatchUtilities.INT.udn
@@ -0,0 +1,106 @@
+Availability:Docs
+Title:Mobile Patch Utility Nodes
+Crumbs: %ROOT%, Engine, Engine/Blueprints, Engine/Blueprints/UserGuide
+Description:Describes the different kinds of Mobile Patch Utility Blueprint nodes.
+Related: Engine/Blueprints
+Related: Platforms/Mobile
+Related: Platforms/Mobile/MobilePackagingWizard
+version: 4.14
+Parent:Engine/Blueprints
+type:overview
+
+[TOC (start:2 end:2)]
+
+After using the [](Platforms/Mobile/MobilePackagingWizard) to create a small initial download for your mobile project, you also need to be able to provide the rest of your project in downloadable chunks to the user.
+The new Mobile Patch Utilities Blueprint library contains all the functionality required to allow a mobile game to download and install game content and patches from a cloud website instead of being 
+distributed as part of the initial download from the App Store. There is functionality to determine if updated game content is available, initiate the download, track progress, handle any errors, and 
+finally install the content paks that are downloaded successfully. Functionality to check for sufficient storage space and WiFi connectivity is also available, so you can warn the user in such cases. 
+Both Android and iOS are supported.
+
+[REGION:tip]
+To show these functions in context with each other, example setups are presented here.
+If a function or event in the Blueprint graph examples isn't listed in a table below, then it is a custom node you would need to create in your own project as well. 
+[/REGION]
+
+## Working with Pending Content
+
+[REGION:lightbox]
+[![](Patching_1.png)(convert:false)](Patching_1.png)
+[/REGION]
+
+This is the first stage of the patching process. When patching is initiated, either from a level load or a user action, content needs to be requested from the remote server using the **Request Content** node.
+In this example, the install directory is saved off in a variable, but it could also be passed through as function inputs. The **Content To Install** is also saved into a variable for reuse later. 
+When the content request is complete, different events will be executed depending on whether the request succeeded or failed. Custom events are used as inputs for the **Request Content** function. 
+The **Check Download Space** function shown here is calling a custom event that starts off the next set of patching logic.
+
+| Node | Description | 
+| --- | --- |
+| Request Content |  Attempts to download a manifest file using the specified manifest URL. On success, it will return an object that represents the remote content. This object can be queried for additional information, like total content size, download size, etc. The user can choose to download and install remote content. |
+
+[REGION:lightbox]
+[![](Patching_2.png)(convert:false)](Patching_2.png)
+[/REGION]
+
+After requesting the remote content succeeds, the next step is to verify that there is enough space for the download using the **Get Required Disk Space** node and comparing the result to the result from **Get Disk Free Space**.
+If the required disk space is less than the disk free space, then a custom event is called to start the download.
+
+| Node | Description | 
+| --- | --- |
+| Get Disk Free Space | Gets the disk free space in megabytes where content is installed (or will be installed). |
+| Get Required Disk Space | Gets the required disk space in megabytes for this content installation. |
+
+[REGION:lightbox]
+[![](Patching_3.png)(convert:false)](Patching_3.png)
+[/REGION]
+
+Although there is a lot going on in this graph, the only Mobile Patch Utility function is **Start Install**, which actually attempts to download and install the remote content. Just like **Request Content**, **Start Install** has
+event parameters for success and failure. If the install succeeds, the **Mount Content** custom event is used to trigger the final stage of the patching process. Another custom event, **Display Download State** is fired off here
+with a timer. Using the timer and the **Update Download State** function, the download state can be shown to the user at regular intervals. More on that logic is in the next Blueprint graph.
+
+| Node | Description | 
+| --- | --- |
+| Start Install |  Attempts to download and install remote content. The user can choose to mount the installed content into the game. |
+
+[REGION:lightbox]
+![](Patching_4.png)
+[/REGION]
+
+While the **Display Download State** custom event is connected to a **Format Text** node here, this could also be the stage of the patching process where you fill a progress bar on the UI, or other methods of
+displaying progress to the user. Here, the outputs of **Get Download Size**, **Get Download Speed**, **Get Download Status Text**, **Get Install Progress**, and **Get Total Downloaded Size** nodes are 
+put together in a simple text block to present to the user. This script is called at regular intervals with the timer, but does not call back to any of the other example graphs.
+
+| Node | Description | 
+| --- | --- |
+| Get Download Size | Gets the total download size for this content installation. |
+| Get Download Speed | Gets the current download speed in megabytes per second. Valid during installation. |
+| Get Download Status Text | Gets the current installation status text. Valid during installation. |
+| Get Install Progress | Gets the current installation progress. The result is between 0 and 1 for known progress, or less than 0 for unknown progress.  |
+| Get Total Downloaded Size | Gets the total downloaded size in megabytes. Valid during installation. |
+
+## Working with Installed Content
+
+[REGION:lightbox]
+![](Patching_5.png)
+[/REGION]
+
+The last stage of the patching process is to mount the content. We are using the stored **Install Directory** variable from the first Blueprint graph, but this could again be passed through as a function input.
+The **Get Installed Content** node is used to look up the content from that directory, and then the **Mount** function mounts it into the game.  Unlike **Request Content** and **Start Install**, **Mount**
+does not have success and failure event callbacks.
+
+| Node | Description | 
+| --- | --- |
+| Get Installed Content | Gets the installed content. This will return a non-null object if there is installed content in the specified directory. The user can choose to mount installed content into the game. |
+| Mount | Mounts installed content. |
+
+## Other Mobile Patching Functions
+
+The above graphs are a simple example of how to set up patching on a mobile device, but there are other, more complex, behaviors you could set up as well. For example, you could have your game warn users when they don't have
+an active WIFI connection, or provide different content for iOS and Android devices. For these use cases and others, the below functions would be useful additions to your Blueprints.
+
+| Node | Description | 
+| --- | --- |
+| Get Active Device Profile Name | Gets the name of currently selected device profile name. |
+| Get Installed Content Size | Gets the installed content size in megabytes  |
+| Get Supported Platform Names | Gets the list of supported platform names on this device. Example: Android_ETC2, Android_ASTC |
+| Has Active WiFi Connection | Returns whether or not a WiFi connection is currently available. |
+	 
\ No newline at end of file
diff --git a/Engine/Documentation/Source/Engine/Blueprints/UserGuide/Variables/Bitmask/Bitmask.INT.udn b/Engine/Documentation/Source/Engine/Blueprints/UserGuide/Variables/Bitmask/Bitmask.INT.udn
--- a/Engine/Documentation/Source/Engine/Blueprints/UserGuide/Variables/Bitmask/Bitmask.INT.udn
+++ b/Engine/Documentation/Source/Engine/Blueprints/UserGuide/Variables/Bitmask/Bitmask.INT.udn
@@ -5,7 +5,7 @@ Description: Blueprint integer variables can be declared as bitmasks to store bi
 version: 4.12
 skilllevel:Intermediate
 Parent:Engine/Blueprints
-type:guide
+type:overview
 
 [TOC (start:2)]
 
@@ -23,7 +23,9 @@ Enumeration Editor_. You could also use an existing bitflag enumeration type tha
 
 1. In your _Blueprint_, create an integer variable and then enable the **Bitmask** attribute in the **Details** panel. 
 
-	![image alt text](image_2.png)	This equates to the `Bitmask` metadata in C++ code. 
+	![image alt text](image_2.png)
+
+	This equates to the `Bitmask` metadata in C++ code. 
 
 1. Select your _Blueprint Enumeration_ (or existing bitflag enumeration) in the **Bitmask Enum** dropdown menu in the **Details** panel.
 
diff --git a/Engine/Documentation/Source/Engine/Blueprints/UserGuide/Variables/Structs/Structs.INT.udn b/Engine/Documentation/Source/Engine/Blueprints/UserGuide/Variables/Structs/Structs.INT.udn
--- a/Engine/Documentation/Source/Engine/Blueprints/UserGuide/Variables/Structs/Structs.INT.udn
+++ b/Engine/Documentation/Source/Engine/Blueprints/UserGuide/Variables/Structs/Structs.INT.udn
@@ -6,7 +6,7 @@ version: 4.12
 skilllevel:Intermediate
 Parent:Engine/Blueprints/Scripting
 tag:Arrays and Structures
-type:guide
+type:overview
 
 [TOC]
 
diff --git a/Engine/Documentation/Source/Engine/Deployment/Deployment.INT.udn b/Engine/Documentation/Source/Engine/Deployment/Deployment.INT.udn
--- a/Engine/Documentation/Source/Engine/Deployment/Deployment.INT.udn
+++ b/Engine/Documentation/Source/Engine/Deployment/Deployment.INT.udn
@@ -4,7 +4,7 @@ Crumbs:%ROOT%, Engine
 Description:Generating platform-specific content for assets used in the game.
 Version: 4.9
 parent:Engine
-order:14
+order:18
 
 [VAR:Topic]
 [OBJECT:Topic]
diff --git a/Engine/Documentation/Source/Engine/Editor/EditorManual.INT.udn b/Engine/Documentation/Source/Engine/Editor/EditorManual.INT.udn
--- a/Engine/Documentation/Source/Engine/Editor/EditorManual.INT.udn
+++ b/Engine/Documentation/Source/Engine/Editor/EditorManual.INT.udn
@@ -5,6 +5,7 @@ Description:Guide to using and configuring the Unreal Editor toolset.
 Navigation:topic
 parent:%ROOT%
 order:2
+type:landing
 
 [VAR:Params]
 [PARAM:image]
diff --git a/Engine/Documentation/Source/Engine/Editor/VR/ActivateVRMode/ActivateVRMode.INT.udn b/Engine/Documentation/Source/Engine/Editor/VR/ActivateVRMode/ActivateVRMode.INT.udn
--- a/Engine/Documentation/Source/Engine/Editor/VR/ActivateVRMode/ActivateVRMode.INT.udn
+++ b/Engine/Documentation/Source/Engine/Editor/VR/ActivateVRMode/ActivateVRMode.INT.udn
@@ -3,7 +3,7 @@ Description: Learn how to enable VR mode in Unreal Editor
 Crumbs:
 Availability: public
 parent: Engine/Editor/VR
-type: howto
+type:how-to
 order: 2
 tags: VR
 tags: Getting Started
diff --git a/Engine/Documentation/Source/Engine/Editor/VR/ActorInteraction/ActorInteraction.INT.udn b/Engine/Documentation/Source/Engine/Editor/VR/ActorInteraction/ActorInteraction.INT.udn
--- a/Engine/Documentation/Source/Engine/Editor/VR/ActorInteraction/ActorInteraction.INT.udn
+++ b/Engine/Documentation/Source/Engine/Editor/VR/ActorInteraction/ActorInteraction.INT.udn
@@ -3,7 +3,7 @@ Description: Learn how to place, transform, and modify Actors when working in VR
 Crumbs:
 Availability: public
 parent: Engine/Editor/VR
-type: howto
+type:how-to
 order: 4
 tags: VR
 tags: Level Editing
diff --git a/Engine/Documentation/Source/Engine/Editor/VR/VREditor.INT.udn b/Engine/Documentation/Source/Engine/Editor/VR/VREditor.INT.udn
--- a/Engine/Documentation/Source/Engine/Editor/VR/VREditor.INT.udn
+++ b/Engine/Documentation/Source/Engine/Editor/VR/VREditor.INT.udn
@@ -81,11 +81,11 @@ In order to use Unreal Editor in VR mode, your development environment must meet
 
 ## Starting Out
 
-[DIR(end:"1" output:"topic" type:"howto" tags:"Getting Started" parent:"Engine/Editor/VR")]
+[DIR(end:"1" output:"topic" type:"how-to" tags:"Getting Started" parent:"Engine/Editor/VR")]
 
 ## Guides
 
-[DIR(end:"1" output:"topic" type:"howto" tags:"Level Editing, GitHub" parent:"Engine/Editor/VR")]
+[DIR(end:"1" output:"topic" type:"how-to" tags:"Level Editing, GitHub" parent:"Engine/Editor/VR")]
 
 ## Reference
 
diff --git a/Engine/Documentation/Source/Engine/Engine.INT.udn b/Engine/Documentation/Source/Engine/Engine.INT.udn
--- a/Engine/Documentation/Source/Engine/Engine.INT.udn
+++ b/Engine/Documentation/Source/Engine/Engine.INT.udn
@@ -5,74 +5,11 @@ Description:Complete guide to using the tools and systems that make up Unreal En
 Navigation:topic
 parent:%ROOT%
 order:3
+topic-image:Engine/Rendering/rendering_topic.png
+topic-icon:%ROOT%/docs_icon.png
+type:landing
 
-[VAR:Params]
-[PARAM:image]
-	![%Engine:title%](Engine/Rendering/rendering_topic.png)
-[/PARAM]
-[PARAM:icon]
-	![](%ROOT%/docs_icon.png)(convert:false)
-[/PARAM]
-[PARAM:title]
-	%Engine:title%
-[/PARAM]
-[PARAM:description]
-	%Engine:description%
-[/PARAM]
-[PARAM:path]
-	[RELATIVE:Engine]
-[/PARAM]
-[/VAR]
-
-[VAR:Topic]
-[OBJECT:Topic]
-	[PARAM:image]
-		![%Engine:title%](Engine/Rendering/rendering_topic.png)
-	[/PARAM]
-	[PARAM:icon]
-		![](%ROOT%/docs_icon.png)(convert:false)
-	[/PARAM]
-	[PARAM:title]
-		%Engine:title%
-	[/PARAM]
-	[PARAM:description]
-		%Engine:description%
-	[/PARAM]
-	[PARAM:path]
-		[RELATIVE:Engine]
-	[/PARAM]
-[/OBJECT]
-[/VAR]
-
-[VAR:TopicCompact]
-[OBJECT:TopicCompact]
-	[PARAM:image]
-		![%Engine:title%](Engine/Rendering/rendering_topic.png)
-	[/PARAM]
-	[PARAM:icon]
-		![](%ROOT%/docs_icon.png)(convert:false)
-	[/PARAM]
-	[PARAM:title]
-		%Engine:title%
-	[/PARAM]
-	[PARAM:description]
-		%Engine:description%
-	[/PARAM]
-	[PARAM:path]
-		[RELATIVE:Engine]
-	[/PARAM]
-[/OBJECT]
-[/VAR]
-
-[VAR:Role]
-[OBJECT:Role]
-%params%
-[PARAMLITERAL:id]
-engine
-[/PARAMLITERAL]
-[/OBJECT]
-[/VAR]
 
 ## Browse Topics
 
-[DIR(end:"1" output:"topiccompact" parent:"Engine" org:"hierarchy")]
+[DIR(end:"1" output:"topiccompact" parent:"Engine")]
diff --git a/Engine/Documentation/Source/Engine/Foliage/Foliage.INT.udn b/Engine/Documentation/Source/Engine/Foliage/Foliage.INT.udn
--- a/Engine/Documentation/Source/Engine/Foliage/Foliage.INT.udn
+++ b/Engine/Documentation/Source/Engine/Foliage/Foliage.INT.udn
@@ -5,7 +5,7 @@ Description:System for rendering instanced meshes on the surfaces of other geome
 Version: 4.12
 SkillLevel:Intermediate
 parent:Engine
-order:9
+order:11
 
 [VAR:Topic]
 [OBJECT:Topic]
diff --git a/Engine/Documentation/Source/Engine/HLOD/HLODLanding.INT.udn b/Engine/Documentation/Source/Engine/HLOD/HLODLanding.INT.udn
--- a/Engine/Documentation/Source/Engine/HLOD/HLODLanding.INT.udn
+++ b/Engine/Documentation/Source/Engine/HLOD/HLODLanding.INT.udn
@@ -1,24 +1,27 @@
 title:Hierarchical Level of Detail
 description:Landing page for the Hierarchical Level of Detail system in Unreal Engine 4
 crumbs:
-availability:docs
+availability:Public
 version:4.13
+parent:Engine
+order:13
+type:landing
+tags:HLOD
+topic-image:defaultHLOD_topic.png
+social-image:HLOD_social.png
+
+![](HLOD_header.png)
 
 **Hierarchical Level of Detail** meshes (HLODs) have been developed as a way to increase performance and keep the quality bar high for projects using Unreal Engine 4.11 or later. HLODs can replace multiple **Static Mesh Actors** with single, combined **Static Mesh Actor** at long view distances. This helps reduce the number of **Actors** that need to be rendered for the scene, increasing performance by lowering the number of draw calls per frame. 
 
-[OBJECT:TopicButtonList]
-	[PARAM:icon]
-		![](%ROOT%/tech_icon.png)(convert:false)
-	[/PARAM]
-	[PARAM:title]
-		HLOD Topics
-	[/PARAM]
-	[PARAM:description]
-	[/PARAM]
-	[PARAM:links]
-	* [](Engine/HLOD/Overview "%Engine/HLOD/Overview:description%")
-	* [](Engine/HLOD/Reference "%Engine/HLOD/Reference:description%")
-	* [](Engine/HLOD/HowTo "%Engine/HLOD/HowTo:description%")
-	[/PARAM]
-[/OBJECT]
+## Starting Out
+
+[DIR(output:"topic" type:"overview" parent:"Engine/HLOD")]
+
+## Guides
+
+[DIR(output:"topic" type:"how-to" parent:"Engine/HLOD")]
+
+## Reference
 
+[DIR(output:"topic" type:"reference" parent:"Engine/HLOD")]
