Copying //UE4/Dev-Documentation to Samples-Main (//UE4/Samples-Main)

diff --git a/Engine/Documentation/Source/Engine/HLOD/HowTo/BuildingHLODs/BuildingHLODs.INT.udn b/Engine/Documentation/Source/Engine/HLOD/HowTo/BuildingHLODs/BuildingHLODs.INT.udn
--- a/Engine/Documentation/Source/Engine/HLOD/HowTo/BuildingHLODs/BuildingHLODs.INT.udn
+++ b/Engine/Documentation/Source/Engine/HLOD/HowTo/BuildingHLODs/BuildingHLODs.INT.udn
@@ -1,10 +1,16 @@
-title:How-To Build Hierarchical Level of Detail Meshes
+title:Building Hierarchical Level of Detail Meshes
 description:This how to covers how to generate the HLOD meshes for your HLOD enabled Unreal Engine 4 project.
 crumbs:
-availability:docs
+availability:Public
 version:4.13
 related:Engine/HLOD/Overview
 related:Engine/HLOD/Reference
+skilllevel:intermediate
+order:2
+parent:Engine/HLOD/HowTo
+type:how-to
+tags:HLOD
+topic-image:Engine/HLOD/HowTo/HLOD_howto_Topic_02.png
 
 There are two major steps in building your **Hierarchical Level of Detail** (HLOD) Meshes. Firstly, you must **Generate Clusters**, which groups **Actors** in your level together based on the setting you specify in the **Cluster Generation Settings**. Then you have to **Generate Proxy Meshes**, which can take a very long time based on scene complexity and settings specified in the **Mesh Generation Settings**.
 
diff --git a/Engine/Documentation/Source/Engine/HLOD/HowTo/Enabling/EnablingHLODs.INT.udn b/Engine/Documentation/Source/Engine/HLOD/HowTo/Enabling/EnablingHLODs.INT.udn
--- a/Engine/Documentation/Source/Engine/HLOD/HowTo/Enabling/EnablingHLODs.INT.udn
+++ b/Engine/Documentation/Source/Engine/HLOD/HowTo/Enabling/EnablingHLODs.INT.udn
@@ -1,10 +1,16 @@
-title:How-To Enable the HLOD System
+title:Enabling the HLOD System
 description:This how-to covers how to enable the HLOD system in your Unreal Engine 4 project.
 crumbs:
-availability:docs
+availability:Public
 version:4.13
 related:Engine/HLOD/Overview
 related:Engine/HLOD/Reference
+skilllevel:beginner
+order:1
+parent:Engine/HLOD/HowTo
+type:how-to
+tags:HLOD
+topic-image:Engine/HLOD/HowTo/HLOD_howto_Topic_01.png
 
 The **Hierarchical Level of Detail System** (HLOD) must be enabled on each Level that wants to make use of it. Follow the steps below to enable the Hierarchical LOD System.
 
diff --git a/Engine/Documentation/Source/Engine/HLOD/HowTo/HLODHowTos.INT.udn b/Engine/Documentation/Source/Engine/HLOD/HowTo/HLODHowTos.INT.udn
--- a/Engine/Documentation/Source/Engine/HLOD/HowTo/HLODHowTos.INT.udn
+++ b/Engine/Documentation/Source/Engine/HLOD/HowTo/HLODHowTos.INT.udn
@@ -1,23 +1,28 @@
 title:Hierarchical Level of Detail How-To's
-description:Words
+description:Guides that walk you through enabling the Hierarchical Level of Detail (HLOD) system and setting up HLOD in your game.
 crumbs:
-availability:docs
+availability:Public
 version:4.13
+parent:Engine/HLOD
+tags:HLOD
+type:how-to
+order:3
+topic-image:Engine/HLOD/HowTo/HLOD_howto_Topic_02.png
+ 
+![](Engine/HLOD/HowTo/HLOD_howto_Topic_02.png)
 
-This page contains links to the **Hierarchical Level of Detail** (HLOD) mesh's How-To's. They will walk you through enabling the HLOD system and setting up HLOD in your game.
+These guides walk you through enabling the **Hierarchical Level of Detail** (HLOD) system and setting up HLOD in your game.
 
-[OBJECT:TopicButtonList]
-	[PARAM:icon]
-		![](%ROOT%/tech_icon.png)(convert:false)
-	[/PARAM]
-	[PARAM:title]
-		HLOD Topics
-	[/PARAM]
-	[PARAM:description]
-	[/PARAM]
-	[PARAM:links]
-		* [](Engine/HLOD/HowTo/Enabling "%Engine/HLOD/HowTo/Enabling:description%")
-		* [](Engine/HLOD/HowTo/BuildingHLODs "%Engine/HLOD/HowTo/BuildingHLODs:description%")
-	[/PARAM]
-[/OBJECT]
-[]()
\ No newline at end of file
+## Beginner
+
+[DIR(output:"fancy" parent:"Engine/HLOD/HowTo" skill_level:"beginner")]
+
+## Intermediate
+
+[DIR(output:"fancy" parent:"Engine/HLOD/HowTo" skill_level:"intermediate")]
+
+<!--
+## Advanced
+
+[DIR(output:"fancy" parent:"Engine/HLOD/HowTo" skill_level:"advanced")]
+-->
\ No newline at end of file
diff --git a/Engine/Documentation/Source/Engine/HLOD/Overview/HLODOverview.INT.udn b/Engine/Documentation/Source/Engine/HLOD/Overview/HLODOverview.INT.udn
--- a/Engine/Documentation/Source/Engine/HLOD/Overview/HLODOverview.INT.udn
+++ b/Engine/Documentation/Source/Engine/HLOD/Overview/HLODOverview.INT.udn
@@ -1,8 +1,15 @@
 title:Hierarchical Level of Detail Overview
 description:Overview of the Hierarchical Level of Detail system in Unreal Engine 4
 crumbs:%ROOT%
-availability:docs
+availability:Public
 version:4.13
+parent:Engine/HLOD
+type:overview
+order:1
+tags:HLOD
+topic-image:Engine/HLOD/HowTo/HLOD_howto_Topic_01.png
+
+![](Engine\HLOD\HowTo\HLOD_howto_header.png)
 
 HLODs, in their simplest form, combine pre-existing **Static Mesh Actors** into a single **Static Mesh Actor** and material, with atlased textures, to reduce draw calls which in-turn helps with your performance of your game. There is a robust system that enables you to set what needs to be generated when you create these **HLOD Actors**.
 
@@ -21,7 +28,7 @@ The proxy mesh above contains several different elements that originally all had
 
 [REGION:note]
 
-This can be problematic for anything using **Mask** or **Opacity**. These alpha channels are not passed to these assets for this merged texture. If opacity or masked items are need you would want to consider unchecking the option for **Merge Textures**. When doing this the Proxy Mesh will assign the original material as a **Material Element ID** instead of combining.
+This can be problematic for anything using **Mask** or **Opacity**. These alpha channels are not passed to these assets for this merged texture. If opacity or masked items are needed you would want to consider unchecking the option for **Merge Textures**. When doing this the Proxy Mesh will assign the original material as a **Material Element ID** instead of combining.
 
 [/REGION]
 
@@ -33,7 +40,7 @@ This can be problematic for anything using **Mask** or **Opacity**. These alpha
 
 ![image alt text](image_18.png)(h:266 w:229 convert:false)
 
-Place this volume around **Actors** you want to be grouped into the same cluster. It can be helpful to make this slightly larger than the **Actors** and should not be tightly bounding them. Once you are have placed the volume around the **Actors** you can use the option **Generate Clusters** from the HLOD Outliner to create a new cluster with these grouped **Actors**.
+Place this volume around **Actors** you want to be grouped into the same cluster. It can be helpful to make this slightly larger than the **Actors** and should not be tightly bounding them. Once you have placed the volume around the **Actors** you can use the option **Generate Clusters** from the HLOD Outliner to create a new cluster with these grouped **Actors**.
 
 ### Example
 
@@ -57,7 +64,7 @@ Right-click on the **LODActor** in the HLOD Outliner and click on **Select Conta
 
 ![image alt text](image_23.png)(h:268 w:357 convert:false)
 
-## Details Panel for **LODActor**
+## Details Panel for LODActor
 
 By selecting a **LODActor** in the level you have access to some settings in the Details panel to override some aspects of your generated HLOD along with details to which static meshes have been included with this cluster. This is not unlike any other **Actors** in your scene where you have the option to override or change some values to work within the confines of your design. 
 
diff --git a/Engine/Documentation/Source/Engine/HLOD/Reference/HLODReference.INT.udn b/Engine/Documentation/Source/Engine/HLOD/Reference/HLODReference.INT.udn
--- a/Engine/Documentation/Source/Engine/HLOD/Reference/HLODReference.INT.udn
+++ b/Engine/Documentation/Source/Engine/HLOD/Reference/HLODReference.INT.udn
@@ -1,8 +1,13 @@
 title:Hierarchical LOD Reference
 description:A reference of the interface elements and properties found in the Hierarchical LOD Outliner.
 crumbs:
-availability:docs
+availability:Public
 version:4.13
+parent:Engine/HLOD
+order:2
+type:reference
+tags:HLOD
+topic-image:Engine/HLOD/defaultHLOD_topic.png
 
 This page details the buttons and properties that exist on the **Hierarchical LOD Outliner**.
 
@@ -10,8 +15,10 @@ This page details the buttons and properties that exist on the **Hierarchical LO
 
 ![HLOD Outliner Fully Expanded](image_2.png)
 
-[REGION:raw]
-### Interface Elements
+
+### HLOD Actions
+
+![](InterfaceElements.png)
 
 Below are the descriptions for the 3 buttons and slider that exist in the **Hierarchical LOD Outliner**.
 
@@ -24,39 +31,39 @@ Below are the descriptions for the 3 buttons and slider that exist in the **Hier
 | | %spacer%  | |
 [/REGION]
 
-### Context Menus
+### Hierarchical LOD Outliner Context Menus
 
 The Hierarchical LOD Outliner has a few right-click context menus, below are how to access them and the options they present.
 
-* The first is accessed by right-clicking on a **Cluster Entry** in the HLOD Outliner:
+#### LOD Actor Context Menu
 
-	[REGION:raw]
-	![](ContextMenu1.png)
-	[/REGION]
+[REGION:raw]
+![](ContextMenu1.png)
+[/REGION]
 
-	[REGION:raw]
-	| Option | %spacer% | Description |
-	| --- | :---: | :--- |
-	| **Select LOD Actor** || Selects the LOD Actor in the viewport (same effect as left-clicking the cluster). |
-	| **Select Contained Actors** || Selects all the **Actors** contained in the **LOD Cluster**. |
-	| **ForceView** || Toggles the view of the **HLOD Actor** in the viewport. |
-	| **Rebuild Proxy Mesh** || Forces the right-clicked cluster to rebuild its proxy mesh. |
-	| **Create Containing Hierarchical Volume** || Creates a **Hierarchical Volume** around the **Actors** contained in the **Cluster** that can then be modified to encompass greater or fewer **Actors**. |
-	| **Delete Cluster** || Deletes the right-clicked Cluster. |
-	[/REGION]
+[REGION:raw]
+| Option | %spacer% | Description |
+| --- | :---: | :--- |
+| **Select LOD Actor** || Selects the LOD Actor in the viewport (same effect as left-clicking the cluster). |
+| **Select Contained Actors** || Selects all the **Actors** contained in the **LOD Cluster**. |
+| **ForceView** || Toggles the view of the **HLOD Actor** in the viewport. |
+| **Rebuild Proxy Mesh** || Forces the right-clicked cluster to rebuild its proxy mesh. |
+| **Create Containing Hierarchical Volume** || Creates a **Hierarchical Volume** around the **Actors** contained in the **Cluster** that can then be modified to encompass greater or fewer **Actors**. |
+| **Delete Cluster** || Deletes the right-clicked Cluster. |
+[/REGION]
 
-* The next is accessed by right-clicking on an **Actor** entry contained within an **Cluster Entry** in the HLOD Outliner:
+#### Actor Context Menu
 
-	[REGION:raw]
-	![](ContextMenu2.png)
-	[/REGION]
+[REGION:raw]
+![](ContextMenu2.png)
+[/REGION]
 
-	[REGION:raw]
-	| Option | %spacer% | Description |
-	| --- | :---: | :--- |
-	| **Remove From Cluster** || Removes the right-clicked **Actor** from the **Cluster**. |
-	| **Exclude From Cluster Generation** || Excludes the right-clicked **Actor** from **Cluster** generation. If you wish to re-add the excluded **Actor** click and drag it from the **World Outliner** into the **Cluster** in the **Hierarchical LOD Outliner**. |
-	[/REGION]
+[REGION:raw]
+| Option | %spacer% | Description |
+| --- | :---: | :--- |
+| **Remove From Cluster** || Removes the right-clicked **Actor** from the **Cluster**. |
+| **Exclude From Cluster Generation** || Excludes the right-clicked **Actor** from **Cluster** generation. If you wish to re-add the excluded **Actor** click and drag it from the **World Outliner** into the **Cluster** in the **Hierarchical LOD Outliner**. |
+[/REGION]
 
 ### Properties
 
@@ -64,14 +71,21 @@ Below, broken out by major section, are the properties that can be found in a **
 
 ### Cluster Generation Settings
 
+![](cluster.png)
+
+[REGION:raw]
 | Property | %spacer% | Description |
 | --- | :---: | :--- |
 | **Desired Bound Radius** || This is not guaranteed but used to calculate filling factor for auto clustering. |
 | **Desired Filling Percentage** || This is not guaranteed but used to calculate filling factor for auto clustering. |
 | **Min Number of Actors to Build** || Min number of **Actors** to build LODActor. |
+[/REGION]
 
 ### Mesh Generation Settings
 
+![](meshSettings.png)
+
+[REGION:raw]
 | Property | %spacer% | Description |
 | --- | :---: | :--- |
 | **Transition Screen Size** || The screen radius a mesh object should reach before swapping to the LOD **Actor**. When the parent LOD displays it won’t display any of the children. |
@@ -80,15 +94,19 @@ Below, broken out by major section, are the properties that can be found in a **
 | **Target Light Map UVChannel** || The target UV channel to assign the lightmap UV. |
 | **Target Light Map Resolution** || The lightmap resolution set for the merged mesh. |
 | **Pivot Point at Zero** || Whether a merged mesh should have a pivot at the World Origin, or at the first merged **Actor** from a generated cluster. |
-| **Merge Physics Data:** || Whether to merge the collision primitives when combining the **Actors**. |
+| **Merge Physics Data** || Whether to merge the collision primitives when combining the **Actors**. |
 | **Merge Materials** || Whether to merge source materials into one flat material. |
 | [REGION:tablesection]**[Material Settings See Below](#materialsettings)**[/REGION] |||
 | **Bake Vertex Data ** || Incorporates painted vertex colours into model/baked out material. |
 | **Use Vertex Data for Baking Material** || Whether or not vertex data such as vertex colours should be used when baking out materials. |
-| **Use Landscape Culling** || Whether or note to use available landscape geometry to cull away invisible triangles. |
+| **Use Landscape Culling** || Whether or not to use available landscape geometry to cull away invisible triangles. |
+[/REGION]
 
 ### Material Settings
 
+![](material.png)
+
+[REGION:raw]
 | Property | %spacer% | Description |
 | --- | :---: | :--- |
 | **Texture Size** || Size of the generated BaseColor texture map. |
@@ -100,7 +118,7 @@ Below, broken out by major section, are the properties that can be found in a **
 | **Roughness Constant** || Allows to you set a roughness constant value in place of a texture map. |
 | **Specular Map** || Whether to generate a Specular map. |
 | **Specular Constant** || Allows to you set a roughness constant value in place of a texture map. |
-| **Emissive Map** || Whether to generate aa Emissive map. |
+| **Emissive Map** || Whether to generate a Emissive map. |
 | **Diffuse Texture Size** || Override for the Diffuse Texture Size. Enabled when **Texture Sizing Type** is set to **Use per property manually overriden texture sizes**. |
 | **Normal Texture Size** || Override for the Normal Texture Size. Enabled when **Texture Sizing Type** is set to **Use per property manually overriden texture sizes**. |
 | **Metallic Texture Size** || Override for the Metallic Texture Size. Enabled when **Texture Sizing Type** is set to **Use per property manually overriden texture sizes**. |
@@ -109,6 +127,7 @@ Below, broken out by major section, are the properties that can be found in a **
 | **Emissive Texture Size** || Override for the Emissive Texture Size. Enabled when **Texture Sizing Type** is set to **Use per property manually overriden texture sizes**. |
 [/REGION]
 
+
 [VAR:spacer]![](%ROOT%/spacer.png)(w:150 convert:false)[/VAR]
 
 [VAR:spacer2]![](%ROOT%/spacer.png)(w:40 convert:false)[/VAR]
diff --git a/Engine/Documentation/Source/Engine/Landscape/Custom/CreatingCustomHeightmapsAndLayers.INT.udn b/Engine/Documentation/Source/Engine/Landscape/Custom/CreatingCustomHeightmapsAndLayers.INT.udn
--- a/Engine/Documentation/Source/Engine/Landscape/Custom/CreatingCustomHeightmapsAndLayers.INT.udn
+++ b/Engine/Documentation/Source/Engine/Landscape/Custom/CreatingCustomHeightmapsAndLayers.INT.udn
@@ -2,7 +2,7 @@ Availability: Public
 Title:Creating and Using Custom Heightmaps and Layers
 Crumbs: %ROOT%, Engine, Engine/Landscape
 Description:Creating and using custom Heightmaps and layers for your landscapes.
-version: 4.9
+version: 4.14
 parent:Engine/Landscape/Creation
 
 
@@ -18,15 +18,29 @@ If this is your first time using the Landscape tools, you might want to check ou
 ## Layers
 
 Landscape Layers are special textures that help to break up the look and feel of the Landscape. 
-<!--Each Landscape Layer can be thought of 
--->
-### Layer Format
-To get Landscape layers that where created in external programs to import correctly, the images need to meet the following requirements. 
-* Grayscale, 8-bit per pixel, single channel file.
-* In either .PNG or .RAW file format. 
 
-If you are creating your layers inside of Photoshop you should use the following settings when creating a new document. 
-![](Photoshop_Layer_Example.png)
+### Layer Formats
+* Landscape layers can be imported from external programs via implementation of the `ILandscapeHeightmapFileFormat` and `ILandscapeWeightmapFileFormat` interfaces. The editor's existing support for image-based importing has been converted to use this interface and is still fully supported. Images for the built-in formats are still required to be grayscale, 8 bits per pixel, single channel files in either .PNG or .RAW format. If you are creating layers in Photoshop, you should use the following settings when creating a new document:
+
+	![](Photoshop_Layer_Example.png)
+
+* In order to create new importers, your plugin should create instances of objects implementing `ILandscapeHeightmapFileFormat` and `ILandscapeWeightmapFileFormat`, and register those objects with `ILandscapeEditorModulemodule::RegisterHeightmapFileFormat` and `ILandscapeEditorModulemodule::RegisterWeightmapFileFormat`, respectively.
+
+* Implementing the `ILandscapeHeightmapFileFormat` interface in a plugin requires overriding the following functions:
+1. **const FLandscapeFileTypeInfo& GetInfo() const**: Returns type information indicating which file types this class handles, and whether or not exporting is supported.
+1. **FLandscapeHeightmapInfo Validate(const TCHAR* HeightmapFilename) const** - Validates the named file, or rejects it and returns an error code and message.
+1. **FLandscapeHeightmapImportData Import(const TCHAR* HeightmapFilename, FLandscapeFileResolution ExpectedResolution) const** - Actually imports the file.
+1. **void Export(const TCHAR* HeightmapFilename, TArrayView<const uint16> Data, FLandscapeFileResolution DataResolution, FVector Scale) const** - Exports the file, if this format supports exporting (see the return value from `GetInfo`). This is the only function that doesn't need to be overridden in order to compile. However, if it is called without being overridden, it will call `check`.
+1. **(Destructor)** - Classes that implement this interface should use a virtual destructor, as they will be deleted via a pointer to the interface class.
+
+* Implementing the `ILandscapeHeightmapFileFormat` interface is nearly identical, with only minor differences in some return types:
+1. **const FLandscapeFileTypeInfo& GetInfo() const** - Returns type information indicating which file types this class handles, and whether or not exporting is supported.
+1. **FLandscapeWeightmapInfo Validate(const TCHAR* WeightmapFilename) const** - Validates the named file, or rejects it and returns an error code and message.
+1. **FLandscapeWeightmapImportData Import(const TCHAR* WeightmapFilename, FLandscapeFileResolution ExpectedResolution) const** - Actually imports the file.
+1. **void Export(const TCHAR* WeightmapFilename, TArrayView<const uint16> Data, FLandscapeFileResolution DataResolution, FVector Scale) const** - Exports the file, if this format supports exporting (see the return value from `GetInfo`). This is the only function that doesn't need to be overridden in order to compile. However, if it is called without being overridden, it will call `check`.
+1. **(Destructor)** - Classes that implement this interface should use a virtual destructor, as they will be deleted via a pointer to the interface class.
+
+* For further information and examples, you can see the interfaces in `LandscapeFileFormatInterfaces.h`, the .PNG implementations in `LandscapeFileFormatPng.cpp1` and `LandscapeFileFormatPng.h` , and the .RAW implementations in `LandscapeFileFormatRaw.cpp` and `LandscapeFileFormatRaw.h`. All of this code is in the LandscapeEditor module in the engine.
 
 <!--
 ### Landscape Size vs. Layer Size
diff --git a/Engine/Documentation/Source/Engine/Landscape/Landscape.INT.udn b/Engine/Documentation/Source/Engine/Landscape/Landscape.INT.udn
--- a/Engine/Documentation/Source/Engine/Landscape/Landscape.INT.udn
+++ b/Engine/Documentation/Source/Engine/Landscape/Landscape.INT.udn
@@ -4,7 +4,7 @@ Crumbs:%ROOT%, Engine
 Description:Landscape system for creating terrain for large, open, outdoor environments.
 version: 4.9
 parent:Engine
-order:8
+order:10
 
 [VAR:Topic]
 [OBJECT:Topic]
diff --git a/Engine/Documentation/Source/Engine/LevelStreaming/HowTo/StreamWithC++/StreamC++.INT.udn b/Engine/Documentation/Source/Engine/LevelStreaming/HowTo/StreamWithC++/StreamC++.INT.udn
--- a/Engine/Documentation/Source/Engine/LevelStreaming/HowTo/StreamWithC++/StreamC++.INT.udn
+++ b/Engine/Documentation/Source/Engine/LevelStreaming/HowTo/StreamWithC++/StreamC++.INT.udn
@@ -3,7 +3,7 @@ Title:Load and Unload Levels with C++
 Crumbs:%ROOT%, Engine, Engine/UI/LevelEditor
 Description:How to stream levels with a custom streaming Actor created in C++
 parent:Engine/LevelStreaming/HowTo
-type:howto
+type:how-to
 order:3
 tags:Level Streaming
 
@@ -46,15 +46,11 @@ OverlapBegins must be tagged with a UFUNCTION macro, and must have the same sign
 
 		#include "Kismet/GameplayStatics.h"
 		
-1. Now you are ready to create your OverlapBegins functionality. In `LevelStreamerActor.cpp`, begin defining the function. Add an if statement to check if we have a world, which 
-will allow us to get and check for our Character. You can use the GameplayStatics function `GetPlayerCharacter` to get the Character at index 0. 
+1. Now you are ready to create your OverlapBegins functionality. In `LevelStreamerActor.cpp`, begin defining the function. You can use the GameplayStatics function `GetPlayerCharacter` to get the Character at index 0. 
 
 		void ALevelStreamerActor::OverlapBegins(UPrimitiveComponent* OverlappedComponent, AActor* OtherActor, UPrimitiveComponent* OtherComp, int32 OtherBodyIndex, bool bFromSweep, const FHitResult & SweepResult)
 		{
-			if(GWorld)
-			{
-				ACharacter* MyCharacter = UGameplayStatics::GetPlayerCharacter(this, 0);			
-			}
+				ACharacter* MyCharacter = UGameplayStatics::GetPlayerCharacter(this, 0);		
 		}
 
 1. After getting MyCharacter, check it against the OtherActor overlapping your BoxComponent. Also, verify that LevelToLoad is not empty, then call LoadStreamLevel. 
@@ -139,15 +135,12 @@ will allow us to get and check for our Character. You can use the GameplayStatic
 
 		void ALevelStreamerActor::OverlapBegins(UPrimitiveComponent* OverlappedComponent, AActor* OtherActor, UPrimitiveComponent* OtherComp, int32 OtherBodyIndex, bool bFromSweep, const FHitResult & SweepResult)
 		{
-			if(GWorld)
-			{
 				ACharacter* MyCharacter = UGameplayStatics::GetPlayerCharacter(this, 0);
 				if (OtherActor == MyCharacter && LevelToLoad != "")
 				{
 					FLatentActionInfo LatentInfo;
 					UGameplayStatics::LoadStreamLevel(this, LevelToLoad, true, true, LatentInfo);
 				}
-			}
 		}
 1. Compile your code, then switch back to the editor.
 
@@ -174,15 +167,12 @@ In LevelStreamerActor.cpp:
 
 	void ALevelStreamerActor::OverlapEnds(UPrimitiveComponent* OverlappedComponent, AActor* OtherActor, UPrimitiveComponent* OtherComp, int32 OtherBodyIndex)
 	{
-		if (GWorld)
-		{
 			ACharacter* MyCharacter = UGameplayStatics::GetPlayerCharacter(this, 0);
 			if (OtherActor == MyCharacter && LevelToLoad != "")
 			{
 				FLatentActionInfo LatentInfo;
 				UGameplayStatics::UnloadStreamLevel(this, LevelToLoad, LatentInfo);
 			}
-		}
 	}
 
 In the constructor:
diff --git a/Engine/Documentation/Source/Engine/LevelStreaming/HowTo/StreamWithVolumes/StreamLSV.INT.udn b/Engine/Documentation/Source/Engine/LevelStreaming/HowTo/StreamWithVolumes/StreamLSV.INT.udn
--- a/Engine/Documentation/Source/Engine/LevelStreaming/HowTo/StreamWithVolumes/StreamLSV.INT.udn
+++ b/Engine/Documentation/Source/Engine/LevelStreaming/HowTo/StreamWithVolumes/StreamLSV.INT.udn
@@ -4,7 +4,7 @@ Crumbs:%ROOT%, Engine, Engine/UI/LevelEditor
 Description:How to use volumes to control the streaming of Levels based on the player's viewpoint.
 Version: 4.12
 parent:Engine/LevelStreaming/HowTo
-type:howto
+type:how-to
 order:1
 tags:Level Streaming
 
@@ -33,6 +33,8 @@ tags:Level Streaming
 	
 1. Select **SunTemple_Streaming**, then click on the **Level Details** icon.
 
+	![](LevelDetails.png)
+
 1. Next to **Streaming Volumes**, click on the plus sign (**+**) to add a new volume. Then, use the dropdown menu or the eyedropper to select the **Level Streaming Volume** you just placed in the level.
 	![](AddLSV.png)
 
diff --git a/Engine/Documentation/Source/Engine/LevelStreaming/LevelStreaming.INT.udn b/Engine/Documentation/Source/Engine/LevelStreaming/LevelStreaming.INT.udn
--- a/Engine/Documentation/Source/Engine/LevelStreaming/LevelStreaming.INT.udn
+++ b/Engine/Documentation/Source/Engine/LevelStreaming/LevelStreaming.INT.udn
@@ -4,7 +4,7 @@ Crumbs:%ROOT%, Engine, Engine/UI/LevelEditor
 Description:Asynchronously loading and unloading levels during play to decrease memory usage and create seamless worlds.
 Version: 4.12
 parent:Engine
-order:10
+order:12
 tags:Level Streaming
 
 [VAR:Topic]
diff --git a/Engine/Documentation/Source/Engine/Matinee/MatineeAndCinematics.INT.udn b/Engine/Documentation/Source/Engine/Matinee/MatineeAndCinematics.INT.udn
--- a/Engine/Documentation/Source/Engine/Matinee/MatineeAndCinematics.INT.udn
+++ b/Engine/Documentation/Source/Engine/Matinee/MatineeAndCinematics.INT.udn
@@ -4,7 +4,7 @@ Crumbs: %ROOT%, Engine
 Description:Animation tool that provides the ability to animate the properties of Actors over time, either to create dynamic gameplay or cinematic in-game sequences.
 Version: 4.9
 parent:Engine
-order:11
+order:15
 
 [VAR:Topic]
 [OBJECT:Topic]
diff --git a/Engine/Documentation/Source/Engine/MediaFramework/MediaFramework.INT.udn b/Engine/Documentation/Source/Engine/MediaFramework/MediaFramework.INT.udn
--- a/Engine/Documentation/Source/Engine/MediaFramework/MediaFramework.INT.udn
+++ b/Engine/Documentation/Source/Engine/MediaFramework/MediaFramework.INT.udn
@@ -4,7 +4,7 @@ Crumbs:%ROOT%, Engine
 Description:Working with Media Playback functionality in Unreal Engine 4. 
 Navigation:topic
 parent:Engine
-order:12
+order:16
 Version: 4.13
 
 [VAR:Topic]
diff --git a/Engine/Documentation/Source/Engine/MediaFramework/TechReference/MF_TechRef.INT.udn b/Engine/Documentation/Source/Engine/MediaFramework/TechReference/MF_TechRef.INT.udn
--- a/Engine/Documentation/Source/Engine/MediaFramework/TechReference/MF_TechRef.INT.udn
+++ b/Engine/Documentation/Source/Engine/MediaFramework/TechReference/MF_TechRef.INT.udn
@@ -5,7 +5,7 @@ Description:Technical information regarding Media Framework in Unreal Engine 4.
 Navigation:topic
 parent:Engine/MediaFramework
 order:2
-Version: 4.13
+Version: 4.14
 Related: Engine\MediaFramework\Overview
 Related: Engine\MediaFramework\HowTo
 
@@ -22,43 +22,50 @@ If you are upgrading Media Framework content from a pre-4.13 project, see the ti
 
 The table below indicates the supported extensions based on each media player platform. 
 
-| | Android Media Player | AVF Media Player | PS4 Media Player | WMF Media Player | Xbox One Media Player |
-| --- | --- | --- |--- | --- | --- |
-| **.3g2** | No | No | No | Yes | No |
-| **.3gp** | No | No | No | Yes | No |
-| **.3gp2** | No | No | No | Yes | No |
-| **.3gpp** | Yes | No | No | Yes | No |
-| **.aac** | Yes | No | No | Yes | No |
-| **.adts** | No | No | No | Yes | No |
-| **.asf** | No | No | No | Yes | No |
-| **.avi** | No | Yes | No | Yes | No |
-| **.m4a** | No | No | No | Yes | No |
-| **.m4v** | No | Yes | No | Yes | No |
-| **.mov** | No | Yes | No | Yes | No |
-| **.mp3** | No | No | No | Yes | No |
-| **.mp4** | Yes | Yes | Yes | Yes | No |
-| **.sami** | No | No | No | Yes | No |
-| **.smi** | No | No | No | Yes | No |
-| **.wav** | No | No | No | Yes | No |
-| **.wma** | No | No | No | Yes | No |
-| **.wmv** | No | No | No | Yes | No |
+_Click the image below to download a .xlsx version of the Supported Files & URLs spreadsheets._
+
+[REGION:raw]
+[![](MediaFrameworkExtentions.png)](SupportedExtentions.xlsx)
+[/REGION]
+
+<!-- Editor Note - you can update the table inside the Attachments folder if need be. -->
+
+The list of supported file formats above is ongoing and support for additional file extensions will come online in future engine updates. For the best compatibility and performance, it is recommended that the **.mp4** file extension in H.264 be used. 
+Also denoted in the table above are extensions that are expected to be supported as of the 4.15 engine release. 
 
 [REGION:warning]
-The MfMedia player plug-in for XboxOne is currently in development and does not work in 4.13. The AvfMedia player plug-in for macOS and iOS is in an experimental state and not all features may work as expected. We are working to resolve both issues and hope to have them addressed in a future engine release.  
+The MfMedia player plug-in for XboxOne is currently in development is partially supported if UE4 is built against Windows 7 or later. The AvfMedia player plug-in for macOS and iOS is in an experimental state and not all features may work as expected. We are working to resolve both issues and hope to have them addressed in a future engine release.  
 [/REGION]
 
 This table shows the supported URLs based on media player platform. 
 
-| | Android Media Player | AVF Media Player | PS4 Media Player | WMF Media Player | Xbox One Media Player |
-| --- | --- | --- |--- | --- | --- |
-| **file://** | Yes | Yes | Yes | Yes | No |
-| **http://** | Yes | No | Yes | Yes | No |
-| **httpd://** | Yes | No | No | Yes | No |
-| **https://** | Yes | No | Yes | Yes | No |
-| **mms://** | Yes | No | No | Yes | No |
-| **rtsp://** | Yes | No | No | Yes | No |
-| **rtspt://** | Yes | No | No | Yes | No |
-| **rtspu://** | Yes | Yes | No | Yes | No |
+[REGION:raw]
+![](URLSheet.png)
+[/REGION]
+
+## WMF Media Configuration Settings
+
+Windows Media Foundation (WMF) handles standard audio/video playback and recording on the Windows Platform, however, WMF is also extensible. 
+By default, WMF supports a number of different formats, encoding, and file containers, but it can also be extended with the use of optional codec packs that can be downloaded and installed from the Internet. 
+What this means, however, is that certain users may not have a particular codec needed to decode and playback a media file. 
+
+In the past, Unreal Engine only supported default codecs that are supplied with WMF to ensure that everyone can handle and process media the same way. 
+There may be times when you want to use a different method of encoding your media or your own proprietary media encoder that you want to ship your game with (for example, you ship your game with an installer that provides the needed encoder). 
+
+You can now enable and allow non-standard codecs for **WMF Media** from the **Plugins** section of **Project Settings**. 
+
+![](ProjectSettings.png)
+
+By default, the player plugin will attempt to detect audio and video codecs that are not supported by the operating system, requiring the user to install additional codec packs. 
+If the player plugin detects an unsupported format when accessing the media, a warning message is provided in the **Output Log**. 
+
+![](UnsupportedLog.png)
+
+Additionally, you can see media asset information inside a **Media Player** asset in the **Info** panel. 
+
+![](MediaInfo.png)
+
+When enabling **Allow non-supported codecs**, the plugin player will skip the check and allow the usage of non-standard codecs. 
 
 
 ## Troubleshooting
@@ -81,6 +88,7 @@ If a media source fails to open or play, check the **Output Log** under the [Dev
 [/REGION]
 
 You can also check information about a Media Source inside the Media Player asset and **Info** tab. 
+This will show what options are available for the media source (Forward/Reverse Rates nad Scrubbing/Seeking), as well as the different Audio and Video streams and information about each stream. 
 
 ![](InfoWindow.png)
 
@@ -119,3 +127,6 @@ See Microsoft's [MPEG-4 File Source](https://msdn.microsoft.com/en-us/library/wi
 	* QuickTime movie (.mov) support on Windows 7 and later is currently unreliable. 
 		*	Workaround: This format is currently not recommended for use. 
 
+
+
+
diff --git a/Engine/Documentation/Source/Engine/OpenWorldTools/OpenWorldTools.INT.udn b/Engine/Documentation/Source/Engine/OpenWorldTools/OpenWorldTools.INT.udn
--- a/Engine/Documentation/Source/Engine/OpenWorldTools/OpenWorldTools.INT.udn
+++ b/Engine/Documentation/Source/Engine/OpenWorldTools/OpenWorldTools.INT.udn
@@ -6,7 +6,7 @@ Navigation:topic
 SkillLevel:Intermediate
 Version: 4.8
 parent:Engine
-order:7
+order:9
 
 ![](T_OWT_Header_00.png)
 
diff --git a/Engine/Documentation/Source/Engine/Performance/ForwardRenderer/ForwardRenderer.INT.udn b/Engine/Documentation/Source/Engine/Performance/ForwardRenderer/ForwardRenderer.INT.udn
new file mode 100644
--- /dev/null
+++ b/Engine/Documentation/Source/Engine/Performance/ForwardRenderer/ForwardRenderer.INT.udn
@@ -0,0 +1,134 @@
+Availability:Docs
+Title: Forward Shading Renderer for VR
+Crumbs: %ROOT%, Engine, Engine/Performance
+Description:Describes the advantages of using the Forward Shading Renderer for increased VR performance.
+type: overview
+version: 4.14
+parent:Engine/Performance
+tags:Performance and Profiling
+tags: Rendering
+tags: VR
+topic-image: ForwardRenderingTopic.png
+social-image: ForwardRenderingSocial.png
+related: Engine/Performance
+related: Engine/Performance/GPU
+related: Engine/Rendering
+
+[REGION:banner]
+![](ForwardRenderingBanner.png)
+[/REGION]
+
+[REGION:warning]
+This feature is currently considered experimental and some aspects may not work as expected or may be altered in future updates. 
+[/REGION]
+
+[EXCERPT:intro]
+By default, Unreal Engine 4 uses a Deferred Renderer as it provides the most versatility and grants access to more rendering features. 
+However, there are some trade-offs in using the Deferred Render that might not be right for all VR experiences. 
+**Forward Rendering** provides a faster baseline with faster rendering passes which may lead to better performance on VR platforms. 
+Not only is Forward Rendering faster, it also provides better anti-aliasing options than the Deferred Render which may lead to more crisp visuals. 
+[/EXCERPT]
+
+[TOC(start:2 end:2)]
+
+## Enabling Forward Shading
+
+To enable the Forward Shading Renderer, open your **Project Settings** then under **Rendering**, enable **Forward Shading**. 
+
+![](EnableForwardShading.png)
+
+You will be prompted to restart the editor where after the restart, you can begin using the Forward Renderer options and features. 
+
+## Enabling Multisample Anti-Aliasing
+
+To use Multisample Anti-Aliasing (MSAA), open your **Project Settings** then under **Rendering** and **Default Settings**, change the **Anti-Aliasing Method** to **MSAA**.
+
+![](EnableMSAA.png)
+
+The forward renderer supports both MSAA and Temporal Anti-Aliasing (TAA) and in most cases, TAA is preferable because it removes both geometric aliasing and specular aliasing. 
+In VR however, the constant sub-pixel movement introduced by head tracking generates unwanted blurriness, making MSAA a better choice. 
+Projects that choose to use MSAA will want to build content to mitigate specular aliasing. 
+Additionally, the **Normal to Roughness** feature can help reduce specular aliasing from detailed normal maps and automatic LOD generation for static meshes can flatten features on distant meshes to help reduce aliasing from small triangles. 
+
+In our tests, using MSAA instead of TAA increases GPU frame time by about 25% (actual cost will depend on your content). 
+
+Below an example is given with Temporal AA enabled versus 4X MSAA enabled. 
+
+[OBJECT:ComparisonSlider]
+	[PARAM:before]
+	![With Temporal AA](TemporalAA.png)(w:1020)
+	[/PARAM]
+	[PARAM:after]
+	![With 4X MSAA](MSAA.png)(w:1020)
+	[/PARAM]
+[/OBJECT]
+
+[REGION:note]
+You can use the console variable `r.MSAACount` to control how many MSAA samples are computed for every pixel. 
+Assigning `r.MSAACount 1` will fall back to using Temporal AA, which allows for convenient toggling between Anti-Aliasing methods. 
+[/REGION]
+
+## Performance and Features
+
+Switching from the Deferred Renderer to the Forward Renderer may provide you with an increase in performance for your project. 
+
+The forward renderer can be faster than the deferred renderer for some content. Most of the performance improvement comes from features that can be disabled per material. 
+By default, only the nearest reflection capture will be applied without parallax correction unless the material opts-in to High Quality Reflections, height fog is computed per-vertex, and planar reflections are only applied to materials that enable it. 
+
+![](ForwardShading.png)
+
+Leveraging these options in Epic's new VR game, Robo Recall, the forward renderer is about 22% faster than the deferred renderer on an NVIDIA 970 GTX. 
+
+![](ExampleImage.png)
+
+The Forward Renderer works by culling lights and Reflection Captures to a frustum-space grid. Each pixel in the forward pass then iterates over the lights and Reflection Captures affecting it, sharing the material with them. 
+Dynamic Shadows for Stationary Lights are computed beforehand and packed into channels of a screen-space shadow mask, leveraging the existing limit of 4 overlapping Stationary Lights. 
+
+Be sure to use Unreal Engine's [](Engine/Performance) Tools to gauge the performance of your game when switching from the Deferred Renderer to the Forward Renderer. 
+You will want to pay special attention to [](Engine/Performance/GPU) to see where GPU cost is going, as well as using the [](Engine/Performance/StatCommands) to aid with your profiling efforts.
+
+### Currently Supported Features
+
+This list comprises the currently supported features of the Forward Renderer: 
+
+* Full support for Stationary Lights, including Dynamic Shadows from Movable objects which blend together with precomputed environment shadows
+* Multiple Reflection Captures blended together with parallax correction
+* Planar Reflections of a partial scene composited into Reflection Captures
+* DBuffer Decals
+* Precomputed Lighting and Sky Lights
+* Unshadowed Movable Lights
+* Capsule Shadows
+* Alpha-to-coverage for Masked Materials
+* Instanced Stereo Compatible
+
+## Known Issues & Common Questions
+
+The following features are **not yet supported** while using the Forward Renderer: 
+
+* Screen Space Techniques (SSR, SSAO, Contact Shadows)
+* Shadow Casting Movable Lights
+* Dynamically Shadowed Translucency
+* Translucency receiving environment shadows from a Stationary Light
+* Light Functions and IES Profiles
+* Alpha-to-coverage
+* MSAA on D-Buffer Decals, Motion Blur, Dynamic Shadows and Capsule Shadows
+
+Below are some common questions/issues that may arise when using the experimental version of the Forward Renderer. 
+
+* **My Material is broken when switching to the Forward Renderer, something about GBuffer scene textures?** 
+	* Access to the GBuffer is not available for texture sampling in the Forward Renderer and is only available with the Deferred Renderer. 
+* **Does the Forward Renderer allow fewer texture samples than the Deferred Render?** 
+	* The Forward Renderer packs all features into one shader so there are fewer texture samplers available for your material. 
+	* You can usually solve this by using Shared Samplers. 
+* **Does Forward Rendering MSAA support alpha to coverage for alpha masked materials?** 
+	* This is on our radar, however, in this experimental release, MSAA does not support alpha to coverage in 4.14.
+* **How come I have issues with MSAA not being able to find samples against Atmospheric Fog?** 
+	* Atmospheric Fog does not handle MSAA correctly yet. 
+	* For Height Fog, you can get around with Vertex Fogging (which is enabled by default when enabling Forward Shading). 
+	* This means it is computed in the Forward BasePass and therefore anti-aliased correctly. 
+	* Atmosphere is still computed in a Deferred Pass off of resolved scene depths so it does not anti-alias correctly (which we hope to fix in future updates). 
+
+
+
+	
+
diff --git a/Engine/Documentation/Source/Engine/Performance/GPU/GPU.INT.udn b/Engine/Documentation/Source/Engine/Performance/GPU/GPU.INT.udn
--- a/Engine/Documentation/Source/Engine/Performance/GPU/GPU.INT.udn
+++ b/Engine/Documentation/Source/Engine/Performance/GPU/GPU.INT.udn
@@ -3,7 +3,7 @@ Title: GPU Profiling
 Crumbs: %ROOT%, Engine
 Description:How to optimize the GPU demands of your game.
 Navigation:topic
-version: 4.9
+version: 4.14
 
 The GPU has many units working in parallel and it is common to be bound by different units for different parts of the frame.
 Because of this, it makes sense to look at finding where the GPU cost is going when looking for the bottleneck and what a GPU bottleneck is.
@@ -61,7 +61,7 @@ Often, the performance cost scales with the amount of pixels. To test that, you
 Using **_r.ScreenPercentage_** is even more convenient but keep in mind there is some extra upsampling cost added once the feature gets used.
 
 If you see a measurable performance change, you are bound by something pixel related. Usually, it is either memory bandwidth (reading and writing) or math bound (ALU), but in rare cases, some specific 
-units are saturated (e.g. MRT export). If you can lower memory (or math) on the relevant passes and see a performance difference, you know it was bound by the memory bandwith (or the ALU units). 
+units are saturated (e.g. MRT export). If you can lower memory (or math) on the relevant passes and see a performance difference, you know it was bound by the memory bandwidth (or the ALU units). 
 The change does not have to look the same - this is just a test. Now you know you have to reduce the cost there to improve the performance.
 
 The shadow map resolution does not scale with the screen resolution (use **_r.Shadow.MaxResolution_**), but unless you have very large areas of shadow casting masked or translucent materials, 
@@ -82,22 +82,64 @@ You can adjust _r.EarlyZPass_ to see if your scene would benefit from a full ear
 If changing the resolution does not matter much, you are likely bound by vertex processing (vertex shader, or tessellation) cost.
 Often, you will have to change content to verify that. Typical causes include:
 
-* Too many vertices (Use Level of Detail meshes)
-* Complex WorldPositionOffset / Displacement material using textures with poor mip mapping (adjust the material)
+* Too many vertices. (Use Level of Detail meshes)
+* Complex World Position Offset / Displacement Material using Textures with poor MIP mapping. (adjust the Material)
 * Tessellation (Avoid if possible, adjust tessellation factor - fastest way: show Tessellation, some hardware scales badly with larger tessellation levels)
-* Many UV or normal seams result in more vertices (look at unwrapped UV - many island are bad, flat shaded meshes have a 3 vertices per triangle)
-* Too many vertex attributes (extra UV channels)
-* Verify the vertex count is reasonable, some importer code might not have welded the vertices (combine vertices that have same position, UV and normal)
+* Many UV or Normal seams result in more vertices. (look at unwrapped UV - many islands are bad, flat shaded meshes have 3 vertices per triangle)
+* Too many vertex attributes. (extra UV channels)
+* Verify the vertex count is reasonable, some importer code might not have welded the vertices. (combine vertices that have same position, UV and normal)
 
 
 Less often, you are bound by something else. That could be:
 * Object cost (more likely a CPU e.g. cost but there might be some GPU cost)
 * Triangle setup cost (very high poly meshes with a cheap vertex shader e.g. shadow map static meshes, rarely the issue)
-  * use level of detail (LOD) meshes
+* Use level of detail (LOD) meshes
 * View cost (e.g. HZB occlusion culling)
 * Scene cost (e.g. GPU particle simulation)
 
 
+## Live GPU Profiler
+
+The live GPU profiler provides real-time  per-frame stats for the major rendering categories. To use the live GPU profiler, press the **Backtick** key to open the console and then input **stat GPU** and press **Enter**. You can also bring the live GPU profiler up via the **Stat** submenu in the **Viewport Options** dropdown.
+
+![](GPU_Stats.png)
+
+The stats are cumulative and non-hierarchical, so you can see the major categories without having to dig down through a tree of events. For example, **Shadow Projection** is the sum of all the shadow projections for all lights across all the views. The on-screen GPU stats provide a simple visual breakdown of the GPU load when your title is running. They are also useful to measure the impact of changes instantaneously; for example when changing console variables, modifying materials in the editor or modifying and re-compiling shaders on the fly (with recompileshaders changed). The GPU stats can be recorded to a file when the title is running for analysis later. 
+
+As with existing stats, you can use the console commands **stat startfile** and **stat stopfile** to record the stats to a **ue4stats** file, and then visualize them by opening the file in the [Unreal Frontend tool](Engine/Deployment/UnrealFrontend/).
+
+![](Saved_Profile.png)
+
+[REGION:caption]
+Profiling the GPU with UnrealFrontend. Total, postprocessing and basepass times are shown.
+[/REGION]
+
+
+
+Stats are declared in code as float counters, e.g:
+
+    DECLARE_FLOAT_COUNTER_STAT(TEXT("Postprocessing"), Stat_GPU_Postprocessing, STATGROUP_GPU);
+
+
+Code blocks on the rendering thread can then be instrumented with SCOPED_GPU_STAT macros which reference those stat names. These work similarly to SCOPED_DRAW_EVENT. For example:
+
+    SCOPED_GPU_STAT(RHICmdList, Stat_GPU_Postprocessing);
+
+
+GPU work that isn’t explicitly instrumented will be included in a catch-all [unaccounted] stat. If that gets too high, it indicates that some additional **SCOPED_GPU_STAT** events are needed to account for the missing work. It’s worth noting that unlike the draw events, GPU stats are cumulative. You can add multiple entries for the same stat, and these are aggregated across the frame. Certain CPU-bound cases the GPU timings can be affected by CPU bottlenecks (bubbles) where the GPU is waiting for the CPU to catch up, so please consider that if you see unexpected results in cases where draw thread time is high. On PlayStation 4 we correct those bubbles by excluding the time between command list submissions from the timings. In future releases, we will be extending that functionality to other modern rendering APIs.
+
+
+
+
+
+
+
+
+
+
+
+
+
 
 
 
diff --git a/Engine/Documentation/Source/Engine/Performance/PerformanceProfiling.INT.udn b/Engine/Documentation/Source/Engine/Performance/PerformanceProfiling.INT.udn
--- a/Engine/Documentation/Source/Engine/Performance/PerformanceProfiling.INT.udn
+++ b/Engine/Documentation/Source/Engine/Performance/PerformanceProfiling.INT.udn
@@ -5,7 +5,8 @@ Description:How to identify performance problems and fix them.
 Navigation:topic
 version: 4.9
 parent:Engine
-order:13
+order:17
+tags:Performance and Profiling
 
 [VAR:Topic]
 [OBJECT:Topic]
diff --git a/Engine/Documentation/Source/Engine/Performance/Scalability/Scalability.INT.udn b/Engine/Documentation/Source/Engine/Performance/Scalability/Scalability.INT.udn
--- a/Engine/Documentation/Source/Engine/Performance/Scalability/Scalability.INT.udn
+++ b/Engine/Documentation/Source/Engine/Performance/Scalability/Scalability.INT.udn
@@ -3,6 +3,9 @@ Title: Scalability
 Crumbs: %ROOT%, Engine, Engine/Performance
 Description:Options for adjusting quality and performance.
 version: 4.9
+parent:Engine/Performance
+tags:Performance and Profiling
+tags:scalability
 
 [VAR:Topic]
 [OBJECT:Topic]
diff --git a/Engine/Documentation/Source/Engine/Performance/Visualizers/RenderingVisualizers.INT.udn b/Engine/Documentation/Source/Engine/Performance/Visualizers/RenderingVisualizers.INT.udn
--- a/Engine/Documentation/Source/Engine/Performance/Visualizers/RenderingVisualizers.INT.udn
+++ b/Engine/Documentation/Source/Engine/Performance/Visualizers/RenderingVisualizers.INT.udn
@@ -2,6 +2,8 @@ Availability:Public
 Title: Rendering Visualizers
 Crumbs: %ROOT%, Engine, Engine/Performance
 Description:Display options and view modes to check on the results of various rendering features.
+parent:Engine/Performance
+tags:Performance and Profiling
 
 [VAR:Topic]
 [OBJECT:Topic]
diff --git a/Engine/Documentation/Source/Engine/Physics/Constraints/ConstraintsReference/Constraints.INT.udn b/Engine/Documentation/Source/Engine/Physics/Constraints/ConstraintsReference/Constraints.INT.udn
--- a/Engine/Documentation/Source/Engine/Physics/Constraints/ConstraintsReference/Constraints.INT.udn
+++ b/Engine/Documentation/Source/Engine/Physics/Constraints/ConstraintsReference/Constraints.INT.udn
@@ -2,48 +2,16 @@ Availability: Public
 Title:Physics Constraint Reference
 Crumbs: %ROOT%, Engine, Engine/Physics
 Description:A reference for all the properties available on Physics Constraints in Unreal Editor.
-Version: 4.9
-
-[VAR:Topic]
-[OBJECT:Topic]
-	[PARAM:image]
-		![%Engine/Physics/Constraints/ConstraintsReference:title%](Engine/Physics/physics_topic.png)
-	[/PARAM]
-	[PARAM:icon]
-		![](%ROOT%/physics_icon.png)(convert:false)
-	[/PARAM]
-	[PARAM:title]
-		%Engine/Physics/Constraints/ConstraintsReference:title%
-	[/PARAM]
-	[PARAM:description]
-		%Engine/Physics/Constraints/ConstraintsReference:description%
-	[/PARAM]
-	[PARAM:path]
-		[RELATIVE:Engine/Physics/Constraints/ConstraintsReference]
-	[/PARAM]
-[/OBJECT]
-[/VAR]
-[VAR:TopicCompact]
-[OBJECT:TopicCompact]
-	[PARAM:image]
-		![%Engine/Physics/Constraints/ConstraintsReference:title%](Engine/Physics/physics_topic.png)
-	[/PARAM]
-	[PARAM:icon]
-		![](%ROOT%/physics_icon.png)(convert:false)
-	[/PARAM]
-	[PARAM:title]
-		%Engine/Physics/Constraints/ConstraintsReference:title%
-	[/PARAM]
-	[PARAM:description]
-		%Engine/Physics/Constraints/ConstraintsReference:description%
-	[/PARAM]
-	[PARAM:path]
-		[RELATIVE:Engine/Physics/Constraints/ConstraintsReference]
-	[/PARAM]
-[/OBJECT]
-[/VAR]
+parent:Engine/Physics/Constraints
+order:1
+Version: 4.13
+topic-image:Engine/Physics/physics_topic.png
+tags:Physics
+SkillLevel: Beginner
 
+[VAR:spacer]![](%ROOT%/spacer.png)(w:150 convert:false)[/VAR]
 
+[VAR:spacer2]![](%ROOT%/spacer.png)(w:40 convert:false)[/VAR]
 
 This page contains a reference listing of properties in the Physics Constraints. 
 
@@ -51,10 +19,6 @@ This page contains a reference listing of properties in the Physics Constraints.
 
 Below are the properties for Physics Constraints separated by major category.
 
-[VAR:spacer]![](%ROOT%/spacer.png)(w:150 convert:false)[/VAR]
-
-[VAR:spacer2]![](%ROOT%/spacer.png)(w:40 convert:false)[/VAR]
-
 [OBJECT:modulesectionexpanded]
 	[PARAMLITERAL:id]
 		cons
@@ -74,14 +38,13 @@ Below are the properties for Physics Constraints separated by major category.
 [REGION:raw]
 | Property | %spacer% | Description |
 | --- | :---: | :--- |
-| **Constraint Actor 1** || %doxygen:UPhysicsConstraintComponent::ConstraintActor1% |
-| **Component Name 1** || %doxygen:UPhysicsConstraintComponent::ComponentName1% |
-| **Constraint Actor 2** || %doxygen:UPhysicsConstraintComponent::ConstraintActor2% |
-| **Component Name 2** || %doxygen:UPhysicsConstraintComponent::ComponentName2%  |
-| **Joint Name** || %doxygen:FConstraintInstance::JointName% |
-| **Constraint Bone 1** || %doxygen:FConstraintInstance::ConstraintBone1% |
-| **Constraint Bone 2** || %doxygen:FConstraintInstance::ConstraintBone2% |
-| **Disable Collision** || %doxygen:FConstraintInstance::bDisableCollision% |
+| **Constraint Actor 1** || When a **Physics Constraint** is placed into the level editor you must specify **Actors** to constrain. This is the first of 2 **Actors**. |
+| **Component Name 1** || The first target component that will be constrained. This can be specified when using **Actors** to constrain only a specific component and not the **Actor** root. |
+| **Constraint Actor 2** || When a **Physics Constraint** is placed into the level editor you must specify **Actors** to constrain. This is the second of 2 **Actors**. |
+| **Component Name 2** || The second target component that will be constrained. This can be specified when using **Actors** to constrain only a specific component and not the **Actor** root.  |
+| **Joint Name** || When constraining in the **Physics Asset Tool**, this is the name of the Bone that is initially constrained. |
+| **Constraint Bone 1** || When constraining in the **Physics Asset Tool**, this is the name of the first joint to constrain. |
+| **Constraint Bone 2** || When constraining in the **Physics Asset Tool**, this is the name of the second joint to constrain. |
 [/REGION]
 
 [/VAR]
@@ -93,7 +56,7 @@ Below are the properties for Physics Constraints separated by major category.
 		proj
 	[/PARAMLITERAL]
 	[PARAM:heading]
-		Projection
+		Constraint Behavior
 	[/PARAM]
 	[PARAM:content]
 		%p2%
@@ -107,16 +70,17 @@ Below are the properties for Physics Constraints separated by major category.
 [REGION:raw]
 | Property | %spacer% | Description |
 | --- | :---: | :--- |
-| **Enable Projection** || %doxygen:FConstraintInstance::bEnableProjection% |
-| **Projection Linear Tolerance** || %doxygen:FConstraintInstance::ProjectionLinearTolerance% |
-| **Projection Angular Tolerance** || %doxygen:FConstraintInstance::ProjectionAngularTolerance% |
+| **Disable Collision** || This disables collision between the constrained components. |
+| **Enable Projection** || If the distance error between bodies exceeds **Projection Linear Tolerance** value, or rotational error exceeds **Projection Angular Tolerance**, the Physics Body will be projected to fix this. This fixes cases where the visuals look disconnected, such as individual chain links breaking away from each other. |
+| **Projection Linear Tolerance** || If the distance error exceeds this value, in world units, the body will be projected. |
+| **Projection Angular Tolerance** || If the rotational error exceeds this value, in degrees, the body will be projected. |
 [/REGION]
 
 [/VAR]
 
 
 
-
+<!--
 [OBJECT:modulesectionexpanded]
 	[PARAMLITERAL:id]
 		join
@@ -142,6 +106,7 @@ Below are the properties for Physics Constraints separated by major category.
 [/REGION]
 
 [/VAR]
+-->
 
 
 
@@ -164,10 +129,11 @@ Below are the properties for Physics Constraints separated by major category.
 [REGION:raw]
 | Property | %spacer% | Description |
 | --- | :---: | :--- |
-| **Linear XMotion** || %doxygen:FConstraintInstance::LinearXMotion% |
-| **Linear YMotion** || %doxygen:FConstraintInstance::LinearYMotion% |
-| **Linear ZMotion** || %doxygen:FConstraintInstance::LinearZMotion% |
-| **Linear Limit Size** || %doxygen:FConstraintInstance::LinearLimitSize% |
+| **XMotion** || Indicates the linear constraint applied along the X-axis. [INCLUDE:#linearMotionEnum] |
+| **YMotion** || Indicates the linear constraint applied along the Y-axis. [INCLUDE:#linearMotionEnum] |
+| **ZMotion** || Indicates the linear constraint applied along the Z-axis. [INCLUDE:#linearMotionEnum] |
+| **Limit** || The distance allowed between the two joint reference frames. |
+| **Scale Linear Limits** || If true, linear limits scale using the absolute min of the 3D scale of the owning component. |
 [/REGION]
 
 ### Advanced
@@ -177,11 +143,13 @@ Below are the properties for Physics Constraints separated by major category.
 [REGION:raw]
 | Property | %spacer% | Description |
 | --- | :---: | :--- |
-| **Linear Limit Soft** || %doxygen:FConstraintInstance::bLinearLimitSoft% |
-| **Linear Limit Stiffness** || %doxygen:FConstraintInstance::LinearLimitStiffness% |
-| **Linear Limit Damping** || %doxygen:FConstraintInstance::LinearLimitDamping% |
-| **Linear Breakable** || %doxygen:FConstraintInstance::bLinearBreakable% |
-| **Linear Break Threshold** || %doxygen:FConstraintInstance::LinearBreakThreshold% |
+| **Soft Constraint** || Whether we want to use a soft constraint (spring). |
+| **Stiffness** || Stiffness of the soft constraint. Only used when Soft Constraint is on. |
+| **Damping** || Damping of the soft constraint. |
+| **Restitution** || Controls the amount of bounce when the constraint is violated. |
+| **Contact Distance** || Determines how close to the limit we have to get before turning the joint on. |
+| **Linear Breakable** || Whether it is possible to break the joint with linear force. |
+| **Linear Break Threshold** || Force needed to break the distance constraint.  |
 [/REGION]
 
 [/VAR]
@@ -210,13 +178,13 @@ Below are the properties for Physics Constraints separated by major category.
 [REGION:raw]
 | Property | %spacer% | Description |
 | --- | :---: | :--- |
-| **Angular Swing 1Motion** || %doxygen:FConstraintInstance::AngularSwing1Motion% |
-| **Angular Twist Motion** || %doxygen:FConstraintInstance::AngularTwistMotion% |
-| **Angular Swing 2Motion** || %doxygen:FConstraintInstance::AngularSwing2Motion% |
-| **Swing 1Limit Angle** || %doxygen:FConstraintInstance::Swing1LimitAngle% |
-| **Twist Limit Angle** || %doxygen:FConstraintInstance::TwistLimitAngle% |
-| **Swing 2Limit Angle** || %doxygen:FConstraintInstance::Swing2LimitAngle% |
-| **Angular Rotation Offset** || %doxygen:FConstraintInstance::AngularRotationOffset% |
+| **Swing 1Motion** || Indicates whether the Swing1 limit is used. [INCLUDE:#angularMotionEnum]  |
+| **Swing 2Motion** || Indicates whether the Swing2 limit is used. [INCLUDE:#angularMotionEnum] |
+| **Twist Motion** || Indicates whether the Twist limit is used. [INCLUDE:#angularMotionEnum] |
+| **Swing 1Limit Angle** || Angle of movement along the **XY** Plane. |
+| **Twist Limit Angle** || Angle of movement along the **XZ** Plane. |
+| **Swing 2Limit Angle** || Symmetric angle of roll along the X-axis. |
+| **Angular Rotation Offset** || Specifies the angular offset between the two frames of reference. By default this limit goes from (-Angle, +Angle). |
 [/REGION]
 
 ### Advanced
@@ -226,14 +194,21 @@ Below are the properties for Physics Constraints separated by major category.
 [REGION:raw]
 | Property | %spacer% | Description |
 | --- | :---: | :--- |
-| **Swing Limit Soft** || %doxygen:FConstraintInstance::bSwingLimitSoft% |
-| **Swing Limit Stiffness** || %doxygen:FConstraintInstance::SwingLimitStiffness% |
-| **Swing Limit Damping** || %doxygen:FConstraintInstance::SwingLimitDamping% |
-| **Twist Limit Soft** || %doxygen:FConstraintInstance::bTwistLimitSoft% |
-| **Twist Limit Stiffness** || %doxygen:FConstraintInstance::TwistLimitStiffness% |
-| **Twist Limit Damping** || %doxygen:FConstraintInstance::TwistLimitDamping% |
-| **Angular Breakable** || %doxygen:FConstraintInstance::bAngularBreakable% |
-| **Angular Break Threshold** || %doxygen:FConstraintInstance::AngularBreakThreshold% |
+| [REGION:tablesection]**Swing Limits**[/REGION] |||
+| **Soft Constraint** || Whether we want to use a soft constraint (spring). |
+| **Stiffness** || Stiffness of the soft constraint. Only used when Soft Constraint is on. |
+| **Damping** || Damping of the soft constraint. |
+| **Restitution** || Controls the amount of bounce when the constraint is violated. |
+| **Contact Distance** || Determines how close to the limit we have to get before turning the joint on. |
+| [REGION:tablesection]**Twist Limits**[/REGION] |||
+| **Soft Constraint** || Whether we want to use a soft constraint (spring). |
+| **Stiffness** || Stiffness of the soft constraint. Only used when Soft Constraint is on. |
+| **Damping** || Damping of the soft constraint. |
+| **Restitution** || Controls the amount of bounce when the constraint is violated. |
+| **Contact Distance** || Determines how close to the limit we have to get before turning the joint on. |
+| [REGION:tablesection]**Breakable**[/REGION] |||
+| **Angular Breakable** || Whether it is possible to break the joint with angular force. |
+| **Angular Break Threshold** || Torque needed to break the joint. |
 [/REGION]
 
 [/VAR]
@@ -260,13 +235,21 @@ Below are the properties for Physics Constraints separated by major category.
 [REGION:raw]
 | Property | %spacer% | Description |
 | --- | :---: | :--- |
-| **Linear Position Drive** || %doxygen:FConstraintInstance::bLinearPositionDrive% |
-| **Linear Position Target** || %doxygen:FConstraintInstance::LinearPositionTarget% |
-| **Linear Position Strength** || %doxygen:FConstraintInstance::LinearDriveSpring% |
-| **Linear Velocity Drive** || %doxygen:FConstraintInstance::bLinearVelocityDrive% |
-| **Linear Velocity Target** || %doxygen:FConstraintInstance::LinearVelocityTarget% |
-| **Linear Velocity Strength** || %doxygen:FConstraintInstance::LinearDriveDamping% |
-| **Max Linear Force** || %doxygen:FConstraintInstance::LinearDriveForceLimit% |
+| **Position Target** || Enables the position linear motor on one or more axes and sets the local position desired. |
+| **Strength** || How much force to apply to reach the desired position. |
+| **Velocity Target** || Enables the velocity linear motor on one or more axes and sets the desired velocity. |
+| **Strength** || How much force to apply to reach the desired velocity. |
+[/REGION]
+
+
+### Advanced
+
+![](PhysCon_MotorADV.png)
+
+[REGION:raw]
+| Property | %spacer% | Description |
+| --- | :---: | :--- |
+| **Max Force** || The force limit of the drive.|
 [/REGION]
 
 [/VAR]
@@ -293,14 +276,24 @@ Below are the properties for Physics Constraints separated by major category.
 [REGION:raw]
 | Property | %spacer% | Description |
 | --- | :---: | :--- |
-| **Angular Orientation Drive** || %doxygen:FConstraintInstance::bAngularOrientationDrive% |
-| **Angular Orientation Target** || %doxygen:FConstraintInstance::AngularOrientationTarget% |
-| **Angular Orientation Strength** || %doxygen:FConstraintInstance::AngularDriveSpring% |
-| **Angular Velocity Drive** || %doxygen:FConstraintInstance::bAngularVelocityDrive% |
-| **Angular Velocity Target** || %doxygen:FConstraintInstance::AngularVelocityTarget% |
-| **Angular Velocity Strength** || %doxygen:FConstraintInstance::AngularDriveDamping% |
-| **Max Angular Force** || %doxygen:FConstraintInstance::AngularDriveForceLimit% |
-| **Angular Drive Mode** || %doxygen:FConstraintInstance::AngularDriveMode% |
+| **Angular Drive Mode** || [INCLUDE:#AngMotorMode] |
+| **Target Orientation** || Target proemtatopm relative to the body reference frame. |
+| **Drives** || Depending on the Angular Drive Mode setting, you can enable or disable the different motors for that mode. |
+| **Strength** || The force applied to reach the target orientation. |
+| **Target Velocity** || Target angular velocity relative to the body reference frame |
+| **Drives** || Depending on the Angular Drive Mode setting, you can enable or disable the different motors for that mode. |
+| **Stength** || The force applied to reach the target velocity. |
+[/REGION]
+
+
+### Advanced
+
+![](PhysCon_MotorADV.png)
+
+[REGION:raw]
+| Property | %spacer% | Description |
+| --- | :---: | :--- |
+| **Max Force** || The force limit of the drive.|
 [/REGION]
 
 [/VAR]
@@ -308,6 +301,18 @@ Below are the properties for Physics Constraints separated by major category.
 
 
 <!--
+
+[EXCERPT:AngMotorMode]
+
+Whether this angular motor uses SLERP (Sphereical Linear Interpolation) or decompose into a Swing Motor and a Twist Motor (cone and roll constraints). SLERP will not work if any of the angular constraints are locked.
+
+| Property | %spacer2% | Description |
+| --- | :---: | :--- |
+| **SLERP** || Sets the motor to SLERP (Sphereical Linear Interpolation) mode. SLERP mode will not work if any of axis of the angular constraint is locked. |
+| **Twist and Swing** || Sets the motor to **Twist and Swing** (cone and roll constraints) mode.|
+
+[/EXCERPT]
+
 ![](PhysCon_JointPresets.png)(convert:false) 
 
 [INCLUDE:#presetTable]
@@ -328,12 +333,12 @@ Below are the properties for Physics Constraints separated by major category.
 | Property | Description |
 | -------- | ----------- |
 |[REGION:tablesection]Constraint[/REGION]||
-| **Component Name 1** | %doxygen:UPhysicsConstraintComponent::ComponentName1% |
-| **Component Name 2** | %doxygen:UPhysicsConstraintComponent::ComponentName2%  |
-| **Joint Name** | %doxygen:FConstraintInstance::JointName% |
-| **Constraint Bone 1** | %doxygen:FConstraintInstance::ConstraintBone1% |
-| **Constraint Bone 2** | %doxygen:FConstraintInstance::ConstraintBone2% |
-| **Disable Collision** | %doxygen:FConstraintInstance::bDisableCollision% |
+| **Component Name 1** |  |
+| **Component Name 2** |   |
+| **Joint Name** |  |
+| **Constraint Bone 1** |  |
+| **Constraint Bone 2** |  |
+| **Disable Collision** |  |
 
 * For information on using **Physics Constraints** as a Blueprint Component, see: [**%Engine/Physics\Constraints\ConstraintsBlueprints:title%**](Engine/Physics\Constraints\ConstraintsBlueprints)
 
@@ -342,10 +347,10 @@ Below are the properties for Physics Constraints separated by major category.
 | Property | Description |
 | -------- | ----------- |
 |[REGION:tablesection]Constraint[/REGION]||
-| **Joint Name** | %doxygen:FConstraintInstance::JointName% |
-| **Constraint Bone 1** | %doxygen:FConstraintInstance::ConstraintBone1% |
-| **Constraint Bone 2** | %doxygen:FConstraintInstance::ConstraintBone2% |
-| **Disable Collision** | %doxygen:FConstraintInstance::bDisableCollision% |
+| **Joint Name** |  |
+| **Constraint Bone 1** |  |
+| **Constraint Bone 2** |  |
+| **Disable Collision** |  |
 
 * For further reference on the **Physics Asset Tool** (**PhAT**), see: [**%Engine/Physics\PhAT\Reference:title%**](Engine/Physics\PhAT\Reference)
 
@@ -358,53 +363,53 @@ Below are the properties for Physics Constraints separated by major category.
 | Property | Description |
 | -------- | ----------- |
 |[REGION:tablesection]Linear Limits[/REGION]||
-| **Linear XMotion** | %doxygen:FConstraintInstance::LinearXMotion% |
-| **Linear YMotion** | %doxygen:FConstraintInstance::LinearYMotion% |
-| **Linear ZMotion** | %doxygen:FConstraintInstance::LinearZMotion% |
+| **Linear XMotion** |  |
+| **Linear YMotion** |  |
+| **Linear ZMotion** |  |
 | **Linear Limit Types** | [INCLUDE:#linearMotionEnum] |
-| **Linear Limit Size** | %doxygen:FConstraintInstance::LinearLimitSize% |
-| **Linear Limit Soft** | %doxygen:FConstraintInstance::bLinearLimitSoft% |
-| **Linear Limit Stiffness** | %doxygen:FConstraintInstance::LinearLimitStiffness% |
-| **Linear Limit Damping** | %doxygen:FConstraintInstance::LinearLimitDamping% |
-| **Linear Breakable** | %doxygen:FConstraintInstance::bLinearBreakable% |
-| **Linear Break Threshold** | %doxygen:FConstraintInstance::LinearBreakThreshold% |
+| **Linear Limit Size** |  |
+| **Linear Limit Soft** |  |
+| **Linear Limit Stiffness** |  |
+| **Linear Limit Damping** |  |
+| **Linear Breakable** |  |
+| **Linear Break Threshold** |  |
 |[REGION:tablesection]Projection[/REGION]||
-| **Enable Projection** | %doxygen:FConstraintInstance::bEnableProjection% |
-| **Projection Linear Tolerance** | %doxygen:FConstraintInstance::ProjectionLinearTolerance% |
-| **Projection Angular Tolerance** | %doxygen:FConstraintInstance::ProjectionAngularTolerance% |
+| **Enable Projection** |  |
+| **Projection Linear Tolerance** |  |
+| **Projection Angular Tolerance** |  |
 |[REGION:tablesection]Angular Limits[/REGION]||
-| **Angular Swing 1Motion** | %doxygen:FConstraintInstance::AngularSwing1Motion% |
-| **Angular Swing 2Motion** | %doxygen:FConstraintInstance::AngularSwing2Motion% |
-| **Angular Twist Motion** | %doxygen:FConstraintInstance::AngularTwistMotion% |
+| **Angular Swing 1Motion** |  |
+| **Angular Swing 2Motion** |  |
+| **Angular Twist Motion** |  |
 | **Angular Limit Types** | [INCLUDE:#angularMotionEnum] |
-| **Swing 1Limit Angle** | %doxygen:FConstraintInstance::Swing1LimitAngle% |
-| **Swing 2Limit Angle** | %doxygen:FConstraintInstance::Swing2LimitAngle% |
-| **Twist Limit Angle** | %doxygen:FConstraintInstance::TwistLimitAngle% |
-| **Swing Limit Soft** | %doxygen:FConstraintInstance::bSwingLimitSoft% |
-| **Swing Limit Stiffness** | %doxygen:FConstraintInstance::SwingLimitStiffness% |
-| **Swing Limit Damping** | %doxygen:FConstraintInstance::SwingLimitDamping% |
-| **Twist Limit Soft** | %doxygen:FConstraintInstance::bTwistLimitSoft% |
-| **Twist Limit Stiffness** | %doxygen:FConstraintInstance::TwistLimitStiffness% |
-| **Twist Limit Damping** | %doxygen:FConstraintInstance::TwistLimitDamping% |
-| **Angular Breakable** | %doxygen:FConstraintInstance::bAngularBreakable% |
-| **Angular Break Threshold** | %doxygen:FConstraintInstance::AngularBreakThreshold% |
+| **Swing 1Limit Angle** |  |
+| **Swing 2Limit Angle** |  |
+| **Twist Limit Angle** |  |
+| **Swing Limit Soft** |  |
+| **Swing Limit Stiffness** |  |
+| **Swing Limit Damping** |  |
+| **Twist Limit Soft** |  |
+| **Twist Limit Stiffness** |  |
+| **Twist Limit Damping** |  |
+| **Angular Breakable** |  |
+| **Angular Break Threshold** |  |
 |[REGION:tablesection]Linear Motor[/REGION]||
-| **Linear Position Drive** | %doxygen:FConstraintInstance::bLinearPositionDrive% |
-| **Linear XPosition Drive** | %doxygen:FConstraintInstance::bLinearXPositionDrive% |
-| **Linear YPosition Drive** | %doxygen:FConstraintInstance::bLinearYPositionDrive% |
-| **Linear ZPosition Drive** | %doxygen:FConstraintInstance::bLinearZPositionDrive% |
-| **Linear Position Target** | %doxygen:FConstraintInstance::LinearPositionTarget% |
+| **Linear Position Drive** |  |
+| **Linear XPosition Drive** |  |
+| **Linear YPosition Drive** |  |
+| **Linear ZPosition Drive** |  |
+| **Linear Position Target** |  |
 | **Linear Position Strength** | Spring to apply to the linear drive.  |
-| **Linear Velocity Drive** | %doxygen:FConstraintInstance::bLinearVelocityDrive% |
-| **Linear Velocity Target** | %doxygen:FConstraintInstance::LinearVelocityTarget% |
+| **Linear Velocity Drive** |  |
+| **Linear Velocity Target** |  |
 | **Linear Velocity Strength** | Damping to apply to the linear drive. |
 | **Max Linear Force** | Limit to the force the linear drive can apply. |
 |[REGION:tablesection]Angular Motor[/REGION]||
-| **Angular Orientation Drive** | %doxygen:FConstraintInstance::bAngularOrientationDrive% |
-| **Angular Orientation Target** | %doxygen:FConstraintInstance::AngularOrientationTarget% |
+| **Angular Orientation Drive** |  |
+| **Angular Orientation Target** |  |
 | **Angular Position Strength** | Spring value to apply for the angular drive. |
-| **Angular Velocity Drive** | %doxygen:FConstraintInstance::bAngularVelocityDrive% |
-| **Angular Velocity Target** | %doxygen:FConstraintInstance::AngularVelocityTarget% |
+| **Angular Velocity Drive** |  |
+| **Angular Velocity Target** |  |
 | **Angular Velocity Strength** | Damping value to apply for the angular drive. |
 | **Max Angular Force ** | Limit to the force the angular drive can apply. |
 
@@ -414,17 +419,17 @@ Below are the properties for Physics Constraints separated by major category.
 [EXCERPT:linearMotionEnum]
 | **Type** | **Description** |
 | --- | ---- |
-| **LCM_Limited** | Limited freedom along this axis. The limit is defined for all axes by the **Linear Limit Size** property. |
-| **LCM_Locked** | Fully constraint along this axis. |
-| **LCM_Free** | No constraint along this axis. |
+| **Free** | No constraint along this axis. |
+| **Limited** | Limited freedom along this axis. The limit is defined for all axes by the **Linear Limit Size** property. |
+| **Locked** | Fully constraint along this axis. |
 [/EXCERPT:linearMotionEnum]
 
 [EXCERPT:angularMotionEnum]
 | **Type** | **Description** |
 | --- | ---- |
-| **ACM_Limited** | Limited freedom around this axis. The limit for each Angular Motion is controlled individually by a correspondingly named Limit property: **Swing 1Limit Angle**, **Swing 2Limit Angle**, and **Twist Limit Angle**. |
-| **ACM_Locked** | Fully constraint around this axis. |
-| **ACM_Free** | No constraint around this axis. |
+| **Free** | No constraint around this axis. |
+| **Limited** | Limited freedom around this axis. The limit for each Angular Motion is controlled individually by a correspondingly named Limit property: **Swing 1Limit Angle**, **Swing 2Limit Angle**, and **Twist Limit Angle**. |
+| **Locked** | Fully constraint around this axis. |
 [/EXCERPT:angularMotionEnum]
 
 -->
diff --git a/Engine/Documentation/Source/Engine/Physics/PhAT/HowTo/ApplyConstraintProfile/ApplyConstraintProfile.INT.udn b/Engine/Documentation/Source/Engine/Physics/PhAT/HowTo/ApplyConstraintProfile/ApplyConstraintProfile.INT.udn
new file mode 100644
--- /dev/null
+++ b/Engine/Documentation/Source/Engine/Physics/PhAT/HowTo/ApplyConstraintProfile/ApplyConstraintProfile.INT.udn
@@ -0,0 +1,181 @@
+﻿Availability:Docs
+Title:How-To Apply a Physics Constraint Profile
+Crumbs: %ROOT%, Engine
+Description:This how-to will show you how to apply a Physics Constraint Profile to a Skeletal Mesh Component using Blueprints.
+parent:Engine/Physics/PhAT/HowTo
+order:1
+Version: 4.13
+skilllevel:Intermediate
+prereq: Engine/Blueprints/QuickStart
+prereq: Engine/Physics/PhAT/HowTo/CreateConstraintProfile
+related:Videos/PLZlv_N0_O1ga0aV9jVqJgog0VWz1cLL5f/N1tDjbFXeOo
+related: Engine/Physics/PhAT/HowTo/ApplyPhysicalAnimationProfile
+topic-image:Engine/Physics/physics_topic.png
+
+
+Below, you will find the steps for creating a simple graph to enable a **Constraint Profile** on a **Skeletal Mesh Component** in a **Pawn**. 
+
+## Steps
+1. Open or create a Blueprint with a **Skeletal Mesh Component**.
+	* If your Blueprint doesn't contain a **Skeletal Mesh Component**, add one using the **Components Panel**.
+
+	![](addSkelMeshComp.png)
+
+1. Adjust the **Skeletal Mesh Component** Collision settings.
+	* The Collision Preset needs to be changed so that that **Skeletal Mesh Component** has **Collision Enabled**, and  if a **Pawn** capsule (or other geometry) exists make sure that the collision settings are compatible. For example with a **Pawn** capsule make sure to ignore the **Pawn** collision:
+
+		![](Engine\Physics\PhAT\HowTo\ApplyPhysicalAnimationProfile\collisionSettings.png)
+
+		[REGION:note]
+		You'll notice that the **Object Type** is set to **Pawn** and that we're ignoring **Pawns** in our collision channels. This fixes the issue of the **Skeletal Mesh** trying to eject itself from 
+		the collision capsule. However, if you want the **Skeletal Mesh** to collide with other **Pawns** you'll need to adjust what **Object Type** your **Skeletal Mesh** is and then change 
+		how the **Capsule** interacts with that **Object Type**. See: **[](Engine/Physics/Collision/HowTo/AddCustomCollisionType)**, for more information on creating **Custom Collision Channels**.
+		[/REGION]
+
+1. Add a **Physical Animation Component** to your Blueprint using the **Components Panel**.
+
+	![](addPhysAnimComp.png)
+
+1. Add a Name Variable and call it **Bone Name**.
+
+	[REGION:raw]
+	![](variableName.png)
+	[/REGION]
+
+1. Compile so we can set the value of the **Bone Name** variable.
+
+	![](compile.png)
+
+1. Change **Bone Name**'s Default Value to the desired target **Bone**, in this case `spine_01`.
+
+	![](boneName.png)
+
+1. Switch to the **Event Graph**.
+1. Locate or create the **Event BeginPlay** event node.
+
+	![](eventBP.png)
+
+1. Add a reference to your **Skeletal Mesh Component**
+
+	![](addSKMPref.png)
+
+1. Pull off the blue pin on your **Skeletal Mesh Component** and search for _Set Constraint_.
+
+	![](createSetAll.png)
+
+1. Add a **Set Constraint Profile for All** node.
+	* **Target** is your **Skeletal Mesh Component**.
+	* **Profile Name** is the name of the **Constraint Profile** that you created (or will create) in the **Physics Asset Tool**.
+	* (optional) **Default if Not Found** if a bone doesn't exist in the Profile, it will keep its current settings.
+
+	![](constraintProfNodeSettings.png)
+
+1. Set All Bodies Below Simulate Physics Node connected to the **Apply Physcial Animation Profile Below** Node.
+	* **Target** is your **Skeletal Mesh Component**.
+	* **In Bone Name** will take the **Bone Name** variable as input.
+	* **New Simulate** should be set to **True**.
+	* Again, since we are using `spine_01` as the target bone, **Include Self** should be checked.
+1. Final Graph:
+
+	[OBJECT:Snippet]
+		[PARAMLITERAL:id]
+			ApplyPhysicalAnimProfBP
+		[/PARAMLITERAL]
+		[PARAMLITERAL:type]
+			blueprint
+		[/PARAMLITERAL]
+		[PARAM:image]
+			![](fullGraph.png)
+		[/PARAM]
+		[PARAM:snippet]
+			%BPSnip_001%
+		[/PARAM]
+	[/OBJECT]
+
+	**Hidden Graph Information**
+
+	[REGION:raw]
+	| Variable | Value | Description |
+	| --- | :--- | :--- |
+	| **Bone Name** %spacer2% | `spine_01` | We use it in the **Set All Bodies Below Simulate Physics** to define what parts of the Skeletal Mesh are simulating. |
+	[/REGION]
+
+	| Component || Description |
+	| --- | :---: | :--- |
+	| **Skeletal Mesh Component** || The **Skeletal Mesh Component** that will be have the **Constraint Profile** applied to it. If you are inheriting from **Character** this will just be named `Mesh`.  |
+
+## Result
+
+Now when you play your game, the **Constraint Profile** will take over. Depending on your settings, this could mean that your character falls apart, or has motors engaged to make their arms flail around.
+
+![](result.png)
+
+This effect was achieved by setting the **Linear Limits** on the upper arm constraints to **Free**.
+
+![](achievedBy.png)
+
+<!-- Variables -->
+
+[VAR:spacer]![](%ROOT%/spacer.png)(w:150 convert:false)[/VAR]
+
+[VAR:spacer2]![](%ROOT%/spacer.png)(w:40 convert:false)[/VAR]
+
+[VAR:BPSnip_001]
+	Begin Object Class=K2Node_Event Name="K2Node_Event_8"
+	   EventReference=(MemberParent=Class'/Script/Engine.Actor',MemberName="ReceiveBeginPlay")
+	   bOverrideFunction=True
+	   NodePosX=32
+	   NodePosY=-16
+	   NodeGuid=48E7E2C9458A1028A689B3A249855163
+	   CustomProperties Pin (PinId=D1792C7A4CBF979351A63AB775A61034,PinName="OutputDelegate",Direction="EGPD_Output",PinType.PinCategory="delegate",PinType.PinSubCategoryObject=None,PinType.PinSubCategoryMemberReference=(MemberParent=Class'/Script/Engine.Actor',MemberName="ReceiveBeginPlay"),PinType.bIsArray=False,PinType.bIsReference=False,PinType.bIsConst=False,PinType.bIsWeakPointer=False,PersistentGuid=00000000000000000000000000000000,bHidden=False,bNotConnectable=False,bDefaultValueIsReadOnly=False,bDefaultValueIsIgnored=False,bAdvancedView=False,)
+	   CustomProperties Pin (PinId=3880C4144B19D9B0F602638392FAD6FA,PinName="then",Direction="EGPD_Output",PinType.PinCategory="exec",PinType.PinSubCategoryObject=None,PinType.PinSubCategoryMemberReference=(),PinType.bIsArray=False,PinType.bIsReference=False,PinType.bIsConst=False,PinType.bIsWeakPointer=False,LinkedTo=(K2Node_CallFunction_132 7ED8710C4A46FF83B175828EFBF77DAE,),PersistentGuid=00000000000000000000000000000000,bHidden=False,bNotConnectable=False,bDefaultValueIsReadOnly=False,bDefaultValueIsIgnored=False,bAdvancedView=False,)
+	End Object
+	Begin Object Class=K2Node_VariableGet Name="K2Node_VariableGet_155"
+	   VariableReference=(MemberName="SkeletalMeshComponent",bSelfContext=True)
+	   NodePosX=32
+	   NodePosY=112
+	   NodeGuid=42ADF0644E0964ECBC298DB6863ADA89
+	   CustomProperties Pin (PinId=788C7F5B4B4B5A7362C8AEA404935EEB,PinName="SkeletalMeshComponent",Direction="EGPD_Output",PinType.PinCategory="object",PinType.PinSubCategoryObject=Class'/Script/Engine.SkeletalMeshComponent',PinType.PinSubCategoryMemberReference=(),PinType.bIsArray=False,PinType.bIsReference=False,PinType.bIsConst=False,PinType.bIsWeakPointer=False,LinkedTo=(K2Node_CallFunction_132 E1392084422728C657D6C3A54A4116F2,),PersistentGuid=00000000000000000000000000000000,bHidden=False,bNotConnectable=False,bDefaultValueIsReadOnly=False,bDefaultValueIsIgnored=False,bAdvancedView=False,)
+	   CustomProperties Pin (PinId=D7B4F9F54069CB07496DEE843FDBBA6B,PinName="self",PinFriendlyName=NSLOCTEXT("K2Node", "Target", "Target"),PinType.PinCategory="object",PinType.PinSubCategoryObject=BlueprintGeneratedClass'/Game/PhysAnimPawn.PhysAnimPawn_C',PinType.PinSubCategoryMemberReference=(),PinType.bIsArray=False,PinType.bIsReference=False,PinType.bIsConst=False,PinType.bIsWeakPointer=False,PersistentGuid=00000000000000000000000000000000,bHidden=True,bNotConnectable=False,bDefaultValueIsReadOnly=False,bDefaultValueIsIgnored=False,bAdvancedView=False,)
+	End Object
+	Begin Object Class=K2Node_CallFunction Name="K2Node_CallFunction_132"
+	   FunctionReference=(MemberParent=Class'/Script/Engine.SkeletalMeshComponent',MemberName="SetConstraintProfileForAll")
+	   NodePosX=320
+	   NodeGuid=19366CE04ECF7D2D1B3F3FAB826015E3
+	   CustomProperties Pin (PinId=7ED8710C4A46FF83B175828EFBF77DAE,PinName="execute",PinToolTip="\nExec",PinType.PinCategory="exec",PinType.PinSubCategoryObject=None,PinType.PinSubCategoryMemberReference=(),PinType.bIsArray=False,PinType.bIsReference=False,PinType.bIsConst=False,PinType.bIsWeakPointer=False,LinkedTo=(K2Node_Event_8 3880C4144B19D9B0F602638392FAD6FA,),PersistentGuid=00000000000000000000000000000000,bHidden=False,bNotConnectable=False,bDefaultValueIsReadOnly=False,bDefaultValueIsIgnored=False,bAdvancedView=False,)
+	   CustomProperties Pin (PinId=94ECC7884A2DC860DBB45E8573B644E7,PinName="then",PinToolTip="\nExec",Direction="EGPD_Output",PinType.PinCategory="exec",PinType.PinSubCategoryObject=None,PinType.PinSubCategoryMemberReference=(),PinType.bIsArray=False,PinType.bIsReference=False,PinType.bIsConst=False,PinType.bIsWeakPointer=False,LinkedTo=(K2Node_CallFunction_8 399238E643A83FB3CFC85CB926B5CF9F,),PersistentGuid=00000000000000000000000000000000,bHidden=False,bNotConnectable=False,bDefaultValueIsReadOnly=False,bDefaultValueIsIgnored=False,bAdvancedView=False,)
+	   CustomProperties Pin (PinId=E1392084422728C657D6C3A54A4116F2,PinName="self",PinFriendlyName=NSLOCTEXT("K2Node", "Target", "Target"),PinToolTip="Target\nSkeletal Mesh Component Reference",PinType.PinCategory="object",PinType.PinSubCategoryObject=Class'/Script/Engine.SkeletalMeshComponent',PinType.PinSubCategoryMemberReference=(),PinType.bIsArray=False,PinType.bIsReference=False,PinType.bIsConst=False,PinType.bIsWeakPointer=False,LinkedTo=(K2Node_VariableGet_155 788C7F5B4B4B5A7362C8AEA404935EEB,),PersistentGuid=00000000000000000000000000000000,bHidden=False,bNotConnectable=False,bDefaultValueIsReadOnly=False,bDefaultValueIsIgnored=False,bAdvancedView=False,)
+	   CustomProperties Pin (PinId=28663C674854E90FE78C718CD73FB3DF,PinName="ProfileName",PinToolTip="Profile Name\nName",PinType.PinCategory="name",PinType.PinSubCategoryObject=None,PinType.PinSubCategoryMemberReference=(),PinType.bIsArray=False,PinType.bIsReference=False,PinType.bIsConst=False,PinType.bIsWeakPointer=False,DefaultValue="DocsConstraintProfile",AutogeneratedDefaultValue="None",PersistentGuid=00000000000000000000000000000000,bHidden=False,bNotConnectable=False,bDefaultValueIsReadOnly=False,bDefaultValueIsIgnored=False,bAdvancedView=False,)
+	   CustomProperties Pin (PinId=D960208D49BE62D6FDE20698BB7D27CC,PinName="bDefaultIfNotFound",PinToolTip="Default if Not Found\nBoolean",PinType.PinCategory="bool",PinType.PinSubCategoryObject=None,PinType.PinSubCategoryMemberReference=(),PinType.bIsArray=False,PinType.bIsReference=False,PinType.bIsConst=False,PinType.bIsWeakPointer=False,DefaultValue="true",AutogeneratedDefaultValue="false",PersistentGuid=00000000000000000000000000000000,bHidden=False,bNotConnectable=False,bDefaultValueIsReadOnly=False,bDefaultValueIsIgnored=False,bAdvancedView=False,)
+	End Object
+	Begin Object Class=K2Node_VariableGet Name="K2Node_VariableGet_1"
+	   VariableReference=(MemberName="SkeletalMeshComponent",bSelfContext=True)
+	   NodePosX=624
+	   NodePosY=80
+	   NodeGuid=211663904BCAEE00C6164CBF0871A8B9
+	   CustomProperties Pin (PinId=1EEABDD6494CBF38903C6691F5F5B45C,PinName="SkeletalMeshComponent",Direction="EGPD_Output",PinType.PinCategory="object",PinType.PinSubCategoryObject=Class'/Script/Engine.SkeletalMeshComponent',PinType.PinSubCategoryMemberReference=(),PinType.bIsArray=False,PinType.bIsReference=False,PinType.bIsConst=False,PinType.bIsWeakPointer=False,LinkedTo=(K2Node_CallFunction_8 4DD6FE2846E695E0F3F569B06CED0A46,),PersistentGuid=00000000000000000000000000000000,bHidden=False,bNotConnectable=False,bDefaultValueIsReadOnly=False,bDefaultValueIsIgnored=False,bAdvancedView=False,)
+	   CustomProperties Pin (PinId=D68B011B4E4D641D541E8CA766972993,PinName="self",PinFriendlyName=NSLOCTEXT("K2Node", "Target", "Target"),PinType.PinCategory="object",PinType.PinSubCategoryObject=BlueprintGeneratedClass'/Game/PhysAnimPawn.PhysAnimPawn_C',PinType.PinSubCategoryMemberReference=(),PinType.bIsArray=False,PinType.bIsReference=False,PinType.bIsConst=False,PinType.bIsWeakPointer=False,PersistentGuid=00000000000000000000000000000000,bHidden=True,bNotConnectable=False,bDefaultValueIsReadOnly=False,bDefaultValueIsIgnored=False,bAdvancedView=False,)
+	End Object
+	Begin Object Class=K2Node_CallFunction Name="K2Node_CallFunction_8"
+	   FunctionReference=(MemberParent=Class'/Script/Engine.SkeletalMeshComponent',MemberName="SetAllBodiesBelowSimulatePhysics")
+	   NodePosX=848
+	   NodeGuid=3116B9174BA9B4F48CE4CAA1E47000EC
+	   CustomProperties Pin (PinId=399238E643A83FB3CFC85CB926B5CF9F,PinName="execute",PinToolTip="\nExec",PinType.PinCategory="exec",PinType.PinSubCategoryObject=None,PinType.PinSubCategoryMemberReference=(),PinType.bIsArray=False,PinType.bIsReference=False,PinType.bIsConst=False,PinType.bIsWeakPointer=False,LinkedTo=(K2Node_CallFunction_132 94ECC7884A2DC860DBB45E8573B644E7,),PersistentGuid=00000000000000000000000000000000,bHidden=False,bNotConnectable=False,bDefaultValueIsReadOnly=False,bDefaultValueIsIgnored=False,bAdvancedView=False,)
+	   CustomProperties Pin (PinId=C132D8014BC6A4DE84C437AF07CF838E,PinName="then",PinToolTip="\nExec",Direction="EGPD_Output",PinType.PinCategory="exec",PinType.PinSubCategoryObject=None,PinType.PinSubCategoryMemberReference=(),PinType.bIsArray=False,PinType.bIsReference=False,PinType.bIsConst=False,PinType.bIsWeakPointer=False,PersistentGuid=00000000000000000000000000000000,bHidden=False,bNotConnectable=False,bDefaultValueIsReadOnly=False,bDefaultValueIsIgnored=False,bAdvancedView=False,)
+	   CustomProperties Pin (PinId=4DD6FE2846E695E0F3F569B06CED0A46,PinName="self",PinFriendlyName=NSLOCTEXT("K2Node", "Target", "Target"),PinToolTip="Target\nSkeletal Mesh Component Reference",PinType.PinCategory="object",PinType.PinSubCategoryObject=Class'/Script/Engine.SkeletalMeshComponent',PinType.PinSubCategoryMemberReference=(),PinType.bIsArray=False,PinType.bIsReference=False,PinType.bIsConst=False,PinType.bIsWeakPointer=False,LinkedTo=(K2Node_VariableGet_1 1EEABDD6494CBF38903C6691F5F5B45C,),PersistentGuid=00000000000000000000000000000000,bHidden=False,bNotConnectable=False,bDefaultValueIsReadOnly=False,bDefaultValueIsIgnored=False,bAdvancedView=False,)
+	   CustomProperties Pin (PinId=7D29E35244DEB831C54B1CAEF4A5FE5C,PinName="InBoneName",PinToolTip="In Bone Name\nName (by ref)",PinType.PinCategory="name",PinType.PinSubCategoryObject=None,PinType.PinSubCategoryMemberReference=(),PinType.bIsArray=False,PinType.bIsReference=True,PinType.bIsConst=True,PinType.bIsWeakPointer=False,DefaultValue="None",AutogeneratedDefaultValue="None",LinkedTo=(K2Node_VariableGet_181 4DF93E5141ACE87768E5D8A7543C7B07,),PersistentGuid=00000000000000000000000000000000,bHidden=False,bNotConnectable=False,bDefaultValueIsReadOnly=False,bDefaultValueIsIgnored=True,bAdvancedView=False,)
+	   CustomProperties Pin (PinId=AE4340174DD534FEA5B812AE6DAEA1EB,PinName="bNewSimulate",PinToolTip="New Simulate\nBoolean",PinType.PinCategory="bool",PinType.PinSubCategoryObject=None,PinType.PinSubCategoryMemberReference=(),PinType.bIsArray=False,PinType.bIsReference=False,PinType.bIsConst=False,PinType.bIsWeakPointer=False,DefaultValue="true",AutogeneratedDefaultValue="false",PersistentGuid=00000000000000000000000000000000,bHidden=False,bNotConnectable=False,bDefaultValueIsReadOnly=False,bDefaultValueIsIgnored=False,bAdvancedView=False,)
+	   CustomProperties Pin (PinId=80EC7602465D797C7F0D4D93D6BA9BD5,PinName="bIncludeSelf",PinToolTip="Include Self\nBoolean",PinType.PinCategory="bool",PinType.PinSubCategoryObject=None,PinType.PinSubCategoryMemberReference=(),PinType.bIsArray=False,PinType.bIsReference=False,PinType.bIsConst=False,PinType.bIsWeakPointer=False,DefaultValue="true",AutogeneratedDefaultValue="true",PersistentGuid=00000000000000000000000000000000,bHidden=False,bNotConnectable=False,bDefaultValueIsReadOnly=False,bDefaultValueIsIgnored=False,bAdvancedView=False,)
+	End Object
+	Begin Object Class=K2Node_VariableGet Name="K2Node_VariableGet_181"
+	   VariableReference=(MemberName="BoneName",MemberGuid=A5E270C141CEE32A0269309215CEBAC4,bSelfContext=True)
+	   NodePosX=640
+	   NodePosY=128
+	   NodeGuid=B116D28D4BF4E6625FF02F92D2EA8DFD
+	   CustomProperties Pin (PinId=4DF93E5141ACE87768E5D8A7543C7B07,PinName="BoneName",Direction="EGPD_Output",PinType.PinCategory="name",PinType.PinSubCategoryObject=None,PinType.PinSubCategoryMemberReference=(),PinType.bIsArray=False,PinType.bIsReference=False,PinType.bIsConst=False,PinType.bIsWeakPointer=False,AutogeneratedDefaultValue="None",LinkedTo=(K2Node_CallFunction_8 7D29E35244DEB831C54B1CAEF4A5FE5C,),PersistentGuid=00000000000000000000000000000000,bHidden=False,bNotConnectable=False,bDefaultValueIsReadOnly=False,bDefaultValueIsIgnored=False,bAdvancedView=False,)
+	   CustomProperties Pin (PinId=F259489241A7853EE35268954C561767,PinName="self",PinFriendlyName=NSLOCTEXT("K2Node", "Target", "Target"),PinType.PinCategory="object",PinType.PinSubCategoryObject=BlueprintGeneratedClass'/Game/PhysAnimPawn.PhysAnimPawn_C',PinType.PinSubCategoryMemberReference=(),PinType.bIsArray=False,PinType.bIsReference=False,PinType.bIsConst=False,PinType.bIsWeakPointer=False,PersistentGuid=00000000000000000000000000000000,bHidden=True,bNotConnectable=False,bDefaultValueIsReadOnly=False,bDefaultValueIsIgnored=False,bAdvancedView=False,)
+	End Object
+[/VAR]
+
+
+
