Copying //UE4/Dev-Documentation to Samples-Main (//UE4/Samples-Main)

diff --git a/Engine/Documentation/Source/Engine/Physics/PhAT/HowTo/ApplyPhysicalAnimationProfile/ApplyPhysicalAnimProfile.INT.udn b/Engine/Documentation/Source/Engine/Physics/PhAT/HowTo/ApplyPhysicalAnimationProfile/ApplyPhysicalAnimProfile.INT.udn
new file mode 100644
--- /dev/null
+++ b/Engine/Documentation/Source/Engine/Physics/PhAT/HowTo/ApplyPhysicalAnimationProfile/ApplyPhysicalAnimProfile.INT.udn
@@ -0,0 +1,222 @@
+﻿Availability:Docs
+Title:How-To Apply a Physical Animation Profile
+Crumbs: %ROOT%, Engine
+Description:This how-to will show you how to apply a Physical Animation Profile to a Skeletal Mesh Component using Blueprints.
+parent:Engine/Physics/PhAT/HowTo
+order:1
+Version: 4.13
+skilllevel:Intermediate
+prereq: Engine/Blueprints/QuickStart
+prereq: Engine/Physics/PhAT/HowTo/CreatePhysicalAnimationProfile
+related:Videos/PLZlv_N0_O1ga0aV9jVqJgog0VWz1cLL5f/N1tDjbFXeOo
+related: Engine/Physics/PhAT/HowTo/ApplyConstraintProfile
+topic-image:Engine/Physics/physics_topic.png
+
+Below you will find the steps for creating a simple graph to enable a **Physical Animation Profile** on a **Skeletal Mesh Component** in a **Pawn**. 
+
+## Steps
+
+1. Open or create a Blueprint with a **Skeletal Mesh Component**.
+	* Or, if your Blueprint doesn't contain a **Skeletal Mesh Component**, add one using the **Components Panel**.
+
+	![](Engine\Physics\PhAT\HowTo\ApplyConstraintProfile\addSkelMeshComp.png)
+
+1. Adjust the **Skeletal Mesh Component** Collision settings.
+	* The Collision Preset needs to be changed so that that **Skeletal Mesh Component** has **Collision Enabled** but also so it doesn't interact with Pawns:
+
+		![](collisionSettings.png)
+
+		[REGION:note]
+		You'll notice that the **Object Type** is set to **Pawn** and that we're ignoring **Pawns** in our collision channels. This fixes the issue of the **Skeletal Mesh** trying to eject itself from 
+		the collision capsule. However, if you want the **Skeletal Mesh** to collide with other **Pawns** you'll need to adjust what **Object Type** your **Skeletal Mesh** is and then change 
+		how the **Capsule** interacts with that **Object Type**. See: **[](Engine/Physics/Collision/HowTo/AddCustomCollisionType)**, for more information on creating **Custom Collision Channels**.
+		[/REGION]
+
+1. Add a **Physical Animation Component** to your Blueprint using the **Components Panel**.
+
+	![](Engine\Physics\PhAT\HowTo\ApplyConstraintProfile\addPhysAnimComp.png)
+
+1. Add a Name Variable and call it **Bone Name**.
+
+	[REGION:raw]
+	![](Engine\Physics\PhAT\HowTo\ApplyConstraintProfile\variableName.png)
+	[/REGION]
+
+1. Compile so we can set the value of the **Bone Name** variable.
+
+	![](Engine\Physics\PhAT\HowTo\ApplyConstraintProfile\compile.png)
+
+1. Change **Bone Name** Default Value to the desired target **Bone**, in this case `spine_01`.
+
+	![](Engine\Physics\PhAT\HowTo\ApplyConstraintProfile\boneName.png)
+
+1. Switch to the **Event Graph**.
+1. Locate or create the **Event BeginPlay** event node.
+
+	![](Engine\Physics\PhAT\HowTo\ApplyConstraintProfile\eventBP.png)
+
+1. Add a **Set Skeletal Mesh Component** Node connecting to **Event Beging Play**
+	* Target is the **Physical Animation Component** you added above.
+	* In **Skeletal Mesh Component** is the **Skeletal Mesh Component** in your Blueprint or the one you added above.
+
+	![](createSetSkel.png)
+
+1. Add a **Apply Physcial Animation Profile Below** Node connecting to the **Set Skeletal Mesh Component** Node.
+	* **Target** is the **Physical Animation Component**.
+	* **Body Name** will accept the **Bone Name** variable as input. In Physics Assets, Physics Bodies are named after their Bone.
+	* **Profile Name** the name of the **Physical Animation Profile** you wish to apply to the **Skeletal Mesh**.
+	* Since we are using `spine_01' as the desired bone so that just the upper body animates, **Include Self** should be checked.
+	* **Clear Not Found** is optional, but we've set it to **True** to clear any bones that are not found in the profile.
+
+	![](applyPhysAnim.png)
+
+1. Set All Bodies Below Simulate Physics Node connected to the **Apply Physcial Animation Profile Below** Node. This will complete our graph.
+	* **Target** is your **Skeletal Mesh Component**.
+	* **In Bone Name** will take the **Bone Name** variable as input.
+	* **New Simulate** should be set to **True**.
+	* Again, since we are using `spine_01` as the target bone, **Include Self** should be checked.
+
+	[OBJECT:Snippet]
+		[PARAMLITERAL:id]
+			ApplyPhysicalAnimProfBP
+		[/PARAMLITERAL]
+		[PARAMLITERAL:type]
+			blueprint
+		[/PARAMLITERAL]
+		[PARAM:image]
+			![](ApplyPhysicialAnimProf_BP.png)
+		[/PARAM]
+		[PARAM:snippet]
+			%BPSnip_001%
+		[/PARAM]
+	[/OBJECT]
+
+	**Hidden Graph Information**
+
+	[REGION:raw]
+	| Variable | Value | Description |
+	| --- | :--- | :--- |
+	| **Bone Name** %spacer2% | `spine_01` | The name of the **Bone** in the **Skeletal Mesh** to start applying the Physical Animation Profile. We also use it in the **Set All Bodies Below Simulate Physics**. |
+	[/REGION]
+
+	| Component || Description |
+	| --- | :---: | :--- |
+	| **Physical Animation Component** || This is the component that actually handles the Physical Animation of our **Skeletal Mesh**. |
+	| **Skeletal Mesh Component** || The **Skeletal Mesh Component** that will be physically animated. If working from an existing Blueprint that inherits from **Character**, it will simply be named `Mesh`. |
+
+## Result
+
+When finished, your **Skeletal Mesh Component** will be simulating physics using the **Physical Animation Component** while in runtime.
+
+![](resultApplyPAP.png)
+
+<!-- Variables -->
+
+[VAR:spacer]![](%ROOT%/spacer.png)(w:150 convert:false)[/VAR]
+
+[VAR:spacer2]![](%ROOT%/spacer.png)(w:40 convert:false)[/VAR]
+
+[VAR:BPSnip_001]
+	Begin Object Class=K2Node_Event Name="K2Node_Event_0"
+		EventReference=(MemberParent=Class'/Script/Engine.Actor',MemberName="ReceiveBeginPlay")
+		bOverrideFunction=True
+		bCommentBubblePinned=True
+		NodeGuid=5B240C834C80C2A814C96C90E55C2B1D
+		CustomProperties Pin (PinId=971FD445473845FC61CB348BF812743F,PinName="OutputDelegate",Direction="EGPD_Output",PinType.PinCategory="delegate",PinType.PinSubCategoryObject=None,PinType.PinSubCategoryMemberReference=(MemberParent=Class'/Script/Engine.Actor',MemberName="ReceiveBeginPlay"),PinType.bIsArray=False,PinType.bIsReference=False,PinType.bIsConst=False,PinType.bIsWeakPointer=False,PersistentGuid=00000000000000000000000000000000,bHidden=False,bNotConnectable=False,bDefaultValueIsReadOnly=False,bDefaultValueIsIgnored=False,bAdvancedView=False,)
+		CustomProperties Pin (PinId=A695AE4A4D3440230EDBA897398BEE45,PinName="then",Direction="EGPD_Output",PinType.PinCategory="exec",PinType.PinSubCategoryObject=None,PinType.PinSubCategoryMemberReference=(),PinType.bIsArray=False,PinType.bIsReference=False,PinType.bIsConst=False,PinType.bIsWeakPointer=False,LinkedTo=(K2Node_CallFunction_0 193A453A43AF46E1ADBD5FA4F37E154B,),PersistentGuid=00000000000000000000000000000000,bHidden=False,bNotConnectable=False,bDefaultValueIsReadOnly=False,bDefaultValueIsIgnored=False,bAdvancedView=False,)
+	End Object
+	Begin Object Class=K2Node_VariableGet Name="K2Node_VariableGet_0"
+		VariableReference=(MemberName="PhysicalAnimationComponent",bSelfContext=True)
+		NodePosY=80
+		NodeGuid=BD6F0554431375C8D13E808D4DFFD28B
+		CustomProperties Pin (PinId=5F0735EC4A38FE885F7175B29F0832CA,PinName="PhysicalAnimationComponent",Direction="EGPD_Output",PinType.PinCategory="object",PinType.PinSubCategoryObject=Class'/Script/Engine.PhysicalAnimationComponent',PinType.PinSubCategoryMemberReference=(),PinType.bIsArray=False,PinType.bIsReference=False,PinType.bIsConst=False,PinType.bIsWeakPointer=False,LinkedTo=(K2Node_CallFunction_0 7766A5D54AEB0ACB1F95B49EF70BCB55,K2Node_Knot_1 36EFEC904DDD665AA5BC8EA40F722815,),PersistentGuid=00000000000000000000000000000000,bHidden=False,bNotConnectable=False,bDefaultValueIsReadOnly=False,bDefaultValueIsIgnored=False,bAdvancedView=False,)
+		CustomProperties Pin (PinId=7142470242FB2133424E24AD697AC6C0,PinName="self",PinFriendlyName=NSLOCTEXT("K2Node", "Target", "Target"),PinType.PinCategory="object",PinType.PinSubCategoryObject=BlueprintGeneratedClass'/Game/PhysAnimPawn.PhysAnimPawn_C',PinType.PinSubCategoryMemberReference=(),PinType.bIsArray=False,PinType.bIsReference=False,PinType.bIsConst=False,PinType.bIsWeakPointer=False,PersistentGuid=00000000000000000000000000000000,bHidden=True,bNotConnectable=False,bDefaultValueIsReadOnly=False,bDefaultValueIsIgnored=False,bAdvancedView=False,)
+	End Object
+	Begin Object Class=K2Node_VariableGet Name="K2Node_VariableGet_5"
+		VariableReference=(MemberName="SkeletalMeshComponent",bSelfContext=True)
+		NodePosX=16
+		NodePosY=144
+		NodeGuid=5F5D73144D96E94DCDBB92A9F9803A95
+		CustomProperties Pin (PinId=1EEABDD6494CBF38903C6691F5F5B45C,PinName="SkeletalMeshComponent",Direction="EGPD_Output",PinType.PinCategory="object",PinType.PinSubCategoryObject=Class'/Script/Engine.SkeletalMeshComponent',PinType.PinSubCategoryMemberReference=(),PinType.bIsArray=False,PinType.bIsReference=False,PinType.bIsConst=False,PinType.bIsWeakPointer=False,LinkedTo=(K2Node_CallFunction_0 37845433463046F2DB1841A4B947D094,),PersistentGuid=00000000000000000000000000000000,bHidden=False,bNotConnectable=False,bDefaultValueIsReadOnly=False,bDefaultValueIsIgnored=False,bAdvancedView=False,)
+		CustomProperties Pin (PinId=D68B011B4E4D641D541E8CA766972993,PinName="self",PinFriendlyName=NSLOCTEXT("K2Node", "Target", "Target"),PinType.PinCategory="object",PinType.PinSubCategoryObject=BlueprintGeneratedClass'/Game/PhysAnimPawn.PhysAnimPawn_C',PinType.PinSubCategoryMemberReference=(),PinType.bIsArray=False,PinType.bIsReference=False,PinType.bIsConst=False,PinType.bIsWeakPointer=False,PersistentGuid=00000000000000000000000000000000,bHidden=True,bNotConnectable=False,bDefaultValueIsReadOnly=False,bDefaultValueIsIgnored=False,bAdvancedView=False,)
+	End Object
+	Begin Object Class=K2Node_CallFunction Name="K2Node_CallFunction_0"
+		FunctionReference=(MemberParent=Class'/Script/Engine.PhysicalAnimationComponent',MemberName="SetSkeletalMeshComponent")
+		NodePosX=256
+		NodePosY=-16
+		NodeGuid=DEDD52104BE4DC7743D0068EAEF62DD1
+		CustomProperties Pin (PinId=193A453A43AF46E1ADBD5FA4F37E154B,PinName="execute",PinToolTip="\nExec",PinType.PinCategory="exec",PinType.PinSubCategoryObject=None,PinType.PinSubCategoryMemberReference=(),PinType.bIsArray=False,PinType.bIsReference=False,PinType.bIsConst=False,PinType.bIsWeakPointer=False,LinkedTo=(K2Node_Event_0 A695AE4A4D3440230EDBA897398BEE45,),PersistentGuid=00000000000000000000000000000000,bHidden=False,bNotConnectable=False,bDefaultValueIsReadOnly=False,bDefaultValueIsIgnored=False,bAdvancedView=False,)
+		CustomProperties Pin (PinId=BE952CF64A6E491E6C0B8F8B5CEB1000,PinName="then",PinToolTip="\nExec",Direction="EGPD_Output",PinType.PinCategory="exec",PinType.PinSubCategoryObject=None,PinType.PinSubCategoryMemberReference=(),PinType.bIsArray=False,PinType.bIsReference=False,PinType.bIsConst=False,PinType.bIsWeakPointer=False,LinkedTo=(K2Node_CallFunction_20 B96C42C240335FB66D65AA84ADD838AC,),PersistentGuid=00000000000000000000000000000000,bHidden=False,bNotConnectable=False,bDefaultValueIsReadOnly=False,bDefaultValueIsIgnored=False,bAdvancedView=False,)
+		CustomProperties Pin (PinId=7766A5D54AEB0ACB1F95B49EF70BCB55,PinName="self",PinFriendlyName=NSLOCTEXT("K2Node", "Target", "Target"),PinToolTip="Target\nPhysical Animation Component Reference",PinType.PinCategory="object",PinType.PinSubCategoryObject=Class'/Script/Engine.PhysicalAnimationComponent',PinType.PinSubCategoryMemberReference=(),PinType.bIsArray=False,PinType.bIsReference=False,PinType.bIsConst=False,PinType.bIsWeakPointer=False,LinkedTo=(K2Node_VariableGet_0 5F0735EC4A38FE885F7175B29F0832CA,),PersistentGuid=00000000000000000000000000000000,bHidden=False,bNotConnectable=False,bDefaultValueIsReadOnly=False,bDefaultValueIsIgnored=False,bAdvancedView=False,)
+		CustomProperties Pin (PinId=37845433463046F2DB1841A4B947D094,PinName="InSkeletalMeshComponent",PinToolTip="In Skeletal Mesh Component\nSkeletal Mesh Component Reference",PinType.PinCategory="object",PinType.PinSubCategoryObject=Class'/Script/Engine.SkeletalMeshComponent',PinType.PinSubCategoryMemberReference=(),PinType.bIsArray=False,PinType.bIsReference=False,PinType.bIsConst=False,PinType.bIsWeakPointer=False,LinkedTo=(K2Node_VariableGet_5 1EEABDD6494CBF38903C6691F5F5B45C,),PersistentGuid=00000000000000000000000000000000,bHidden=False,bNotConnectable=False,bDefaultValueIsReadOnly=False,bDefaultValueIsIgnored=False,bAdvancedView=False,)
+	End Object
+	Begin Object Class=K2Node_Knot Name="K2Node_Knot_1"
+		NodePosX=256
+		NodePosY=128
+		NodeGuid=E4A8165E4E49B3FDC3D36AB3198BE2A7
+		CustomProperties Pin (PinId=36EFEC904DDD665AA5BC8EA40F722815,PinName="InputPin",PinType.PinCategory="object",PinType.PinSubCategoryObject=Class'/Script/Engine.PhysicalAnimationComponent',PinType.PinSubCategoryMemberReference=(),PinType.bIsArray=False,PinType.bIsReference=False,PinType.bIsConst=False,PinType.bIsWeakPointer=False,LinkedTo=(K2Node_VariableGet_0 5F0735EC4A38FE885F7175B29F0832CA,),PersistentGuid=00000000000000000000000000000000,bHidden=False,bNotConnectable=False,bDefaultValueIsReadOnly=False,bDefaultValueIsIgnored=True,bAdvancedView=False,)
+		CustomProperties Pin (PinId=E54B8E124E80A63FF92ACF8FA75AACAA,PinName="OutputPin",Direction="EGPD_Output",PinType.PinCategory="object",PinType.PinSubCategoryObject=Class'/Script/Engine.PhysicalAnimationComponent',PinType.PinSubCategoryMemberReference=(),PinType.bIsArray=False,PinType.bIsReference=False,PinType.bIsConst=False,PinType.bIsWeakPointer=False,LinkedTo=(K2Node_Knot_2 A246331D43359B8F04A7EBBEDB2EDC22,),PersistentGuid=00000000000000000000000000000000,bHidden=False,bNotConnectable=False,bDefaultValueIsReadOnly=False,bDefaultValueIsIgnored=False,bAdvancedView=False,)
+	End Object
+	Begin Object Class=K2Node_Knot Name="K2Node_Knot_2"
+		NodePosX=496
+		NodePosY=128
+		NodeGuid=59DD0FBD4738F7C8ADCAB787A410908C
+		CustomProperties Pin (PinId=A246331D43359B8F04A7EBBEDB2EDC22,PinName="InputPin",PinType.PinCategory="object",PinType.PinSubCategoryObject=Class'/Script/Engine.PhysicalAnimationComponent',PinType.PinSubCategoryMemberReference=(),PinType.bIsArray=False,PinType.bIsReference=False,PinType.bIsConst=False,PinType.bIsWeakPointer=False,LinkedTo=(K2Node_Knot_1 E54B8E124E80A63FF92ACF8FA75AACAA,),PersistentGuid=00000000000000000000000000000000,bHidden=False,bNotConnectable=False,bDefaultValueIsReadOnly=False,bDefaultValueIsIgnored=True,bAdvancedView=False,)
+		CustomProperties Pin (PinId=CC66C9584A8CC28A92B6D4AD7CD40472,PinName="OutputPin",Direction="EGPD_Output",PinType.PinCategory="object",PinType.PinSubCategoryObject=Class'/Script/Engine.PhysicalAnimationComponent',PinType.PinSubCategoryMemberReference=(),PinType.bIsArray=False,PinType.bIsReference=False,PinType.bIsConst=False,PinType.bIsWeakPointer=False,LinkedTo=(K2Node_CallFunction_20 96A091DE4DF2762BDFA90CB50CAE9BFB,),PersistentGuid=00000000000000000000000000000000,bHidden=False,bNotConnectable=False,bDefaultValueIsReadOnly=False,bDefaultValueIsIgnored=False,bAdvancedView=False,)
+	End Object
+	Begin Object Class=K2Node_VariableGet Name="K2Node_VariableGet_37"
+		VariableReference=(MemberName="BoneName",bSelfContext=True)
+		NodePosX=400
+		NodePosY=160
+		NodeGuid=7B320511484EBE8F7DDB7B8C7044CD27
+		CustomProperties Pin (PinId=62FD25D94D3A5CAE32C5E685EDA461B1,PinName="BoneName",Direction="EGPD_Output",PinType.PinCategory="name",PinType.PinSubCategoryObject=None,PinType.PinSubCategoryMemberReference=(),PinType.bIsArray=False,PinType.bIsReference=False,PinType.bIsConst=False,PinType.bIsWeakPointer=False,AutogeneratedDefaultValue="None",LinkedTo=(K2Node_Knot_4 4BF4C00042C37A0CDB25AD926D106B22,K2Node_CallFunction_20 EB57682C40C2F631928F3D88FF0B2117,),PersistentGuid=00000000000000000000000000000000,bHidden=False,bNotConnectable=False,bDefaultValueIsReadOnly=False,bDefaultValueIsIgnored=False,bAdvancedView=False,)
+		CustomProperties Pin (PinId=CC9F75D84295F4E8568F03A68C65A446,PinName="self",PinFriendlyName=NSLOCTEXT("K2Node", "Target", "Target"),PinType.PinCategory="object",PinType.PinSubCategoryObject=BlueprintGeneratedClass'/Game/PhysAnimPawn.PhysAnimPawn_C',PinType.PinSubCategoryMemberReference=(),PinType.bIsArray=False,PinType.bIsReference=False,PinType.bIsConst=False,PinType.bIsWeakPointer=False,PersistentGuid=00000000000000000000000000000000,bHidden=True,bNotConnectable=False,bDefaultValueIsReadOnly=False,bDefaultValueIsIgnored=False,bAdvancedView=False,)
+	End Object
+	Begin Object Class=K2Node_VariableGet Name="K2Node_VariableGet_1"
+		VariableReference=(MemberName="SkeletalMeshComponent",bSelfContext=True)
+		NodePosX=928
+		NodePosY=64
+		NodeGuid=5E9CBBED44802CB75DD821972402DEFD
+		CustomProperties Pin (PinId=1EEABDD6494CBF38903C6691F5F5B45C,PinName="SkeletalMeshComponent",Direction="EGPD_Output",PinType.PinCategory="object",PinType.PinSubCategoryObject=Class'/Script/Engine.SkeletalMeshComponent',PinType.PinSubCategoryMemberReference=(),PinType.bIsArray=False,PinType.bIsReference=False,PinType.bIsConst=False,PinType.bIsWeakPointer=False,LinkedTo=(K2Node_CallFunction_8 4DD6FE2846E695E0F3F569B06CED0A46,),PersistentGuid=00000000000000000000000000000000,bHidden=False,bNotConnectable=False,bDefaultValueIsReadOnly=False,bDefaultValueIsIgnored=False,bAdvancedView=False,)
+		CustomProperties Pin (PinId=D68B011B4E4D641D541E8CA766972993,PinName="self",PinFriendlyName=NSLOCTEXT("K2Node", "Target", "Target"),PinType.PinCategory="object",PinType.PinSubCategoryObject=BlueprintGeneratedClass'/Game/PhysAnimPawn.PhysAnimPawn_C',PinType.PinSubCategoryMemberReference=(),PinType.bIsArray=False,PinType.bIsReference=False,PinType.bIsConst=False,PinType.bIsWeakPointer=False,PersistentGuid=00000000000000000000000000000000,bHidden=True,bNotConnectable=False,bDefaultValueIsReadOnly=False,bDefaultValueIsIgnored=False,bAdvancedView=False,)
+	End Object
+	Begin Object Class=K2Node_CallFunction Name="K2Node_CallFunction_8"
+		FunctionReference=(MemberParent=Class'/Script/Engine.SkeletalMeshComponent',MemberName="SetAllBodiesBelowSimulatePhysics")
+		NodePosX=1152
+		NodePosY=-16
+		NodeGuid=C58AFBDD484D31A196381BA69976AB10
+		CustomProperties Pin (PinId=399238E643A83FB3CFC85CB926B5CF9F,PinName="execute",PinToolTip="\nExec",PinType.PinCategory="exec",PinType.PinSubCategoryObject=None,PinType.PinSubCategoryMemberReference=(),PinType.bIsArray=False,PinType.bIsReference=False,PinType.bIsConst=False,PinType.bIsWeakPointer=False,LinkedTo=(K2Node_CallFunction_20 B38EF5724C795C379DDB809F78D77BA1,),PersistentGuid=00000000000000000000000000000000,bHidden=False,bNotConnectable=False,bDefaultValueIsReadOnly=False,bDefaultValueIsIgnored=False,bAdvancedView=False,)
+		CustomProperties Pin (PinId=C132D8014BC6A4DE84C437AF07CF838E,PinName="then",PinToolTip="\nExec",Direction="EGPD_Output",PinType.PinCategory="exec",PinType.PinSubCategoryObject=None,PinType.PinSubCategoryMemberReference=(),PinType.bIsArray=False,PinType.bIsReference=False,PinType.bIsConst=False,PinType.bIsWeakPointer=False,PersistentGuid=00000000000000000000000000000000,bHidden=False,bNotConnectable=False,bDefaultValueIsReadOnly=False,bDefaultValueIsIgnored=False,bAdvancedView=False,)
+		CustomProperties Pin (PinId=4DD6FE2846E695E0F3F569B06CED0A46,PinName="self",PinFriendlyName=NSLOCTEXT("K2Node", "Target", "Target"),PinToolTip="Target\nSkeletal Mesh Component Reference",PinType.PinCategory="object",PinType.PinSubCategoryObject=Class'/Script/Engine.SkeletalMeshComponent',PinType.PinSubCategoryMemberReference=(),PinType.bIsArray=False,PinType.bIsReference=False,PinType.bIsConst=False,PinType.bIsWeakPointer=False,LinkedTo=(K2Node_VariableGet_1 1EEABDD6494CBF38903C6691F5F5B45C,),PersistentGuid=00000000000000000000000000000000,bHidden=False,bNotConnectable=False,bDefaultValueIsReadOnly=False,bDefaultValueIsIgnored=False,bAdvancedView=False,)
+		CustomProperties Pin (PinId=7D29E35244DEB831C54B1CAEF4A5FE5C,PinName="InBoneName",PinToolTip="In Bone Name\nName (by ref)",PinType.PinCategory="name",PinType.PinSubCategoryObject=None,PinType.PinSubCategoryMemberReference=(),PinType.bIsArray=False,PinType.bIsReference=True,PinType.bIsConst=True,PinType.bIsWeakPointer=False,DefaultValue="None",AutogeneratedDefaultValue="None",LinkedTo=(K2Node_Knot_3 336A876240769F5A4170CF8126A5CF36,),PersistentGuid=00000000000000000000000000000000,bHidden=False,bNotConnectable=False,bDefaultValueIsReadOnly=False,bDefaultValueIsIgnored=True,bAdvancedView=False,)
+		CustomProperties Pin (PinId=AE4340174DD534FEA5B812AE6DAEA1EB,PinName="bNewSimulate",PinToolTip="New Simulate\nBoolean",PinType.PinCategory="bool",PinType.PinSubCategoryObject=None,PinType.PinSubCategoryMemberReference=(),PinType.bIsArray=False,PinType.bIsReference=False,PinType.bIsConst=False,PinType.bIsWeakPointer=False,DefaultValue="true",AutogeneratedDefaultValue="false",PersistentGuid=00000000000000000000000000000000,bHidden=False,bNotConnectable=False,bDefaultValueIsReadOnly=False,bDefaultValueIsIgnored=False,bAdvancedView=False,)
+		CustomProperties Pin (PinId=80EC7602465D797C7F0D4D93D6BA9BD5,PinName="bIncludeSelf",PinToolTip="Include Self\nBoolean",PinType.PinCategory="bool",PinType.PinSubCategoryObject=None,PinType.PinSubCategoryMemberReference=(),PinType.bIsArray=False,PinType.bIsReference=False,PinType.bIsConst=False,PinType.bIsWeakPointer=False,DefaultValue="true",AutogeneratedDefaultValue="true",PersistentGuid=00000000000000000000000000000000,bHidden=False,bNotConnectable=False,bDefaultValueIsReadOnly=False,bDefaultValueIsIgnored=False,bAdvancedView=False,)
+	End Object
+	Begin Object Class=K2Node_Knot Name="K2Node_Knot_3"
+		NodePosX=976
+		NodePosY=224
+		NodeGuid=55465C7448CB280BB4657AA36C94D451
+		CustomProperties Pin (PinId=D6618E7C42476ADDD628FF855FD6E7C9,PinName="InputPin",PinType.PinCategory="name",PinType.PinSubCategoryObject=None,PinType.PinSubCategoryMemberReference=(),PinType.bIsArray=False,PinType.bIsReference=False,PinType.bIsConst=False,PinType.bIsWeakPointer=False,LinkedTo=(K2Node_Knot_4 A0911D3740ABD40031008799EB009AE0,),PersistentGuid=00000000000000000000000000000000,bHidden=False,bNotConnectable=False,bDefaultValueIsReadOnly=False,bDefaultValueIsIgnored=True,bAdvancedView=False,)
+		CustomProperties Pin (PinId=336A876240769F5A4170CF8126A5CF36,PinName="OutputPin",Direction="EGPD_Output",PinType.PinCategory="name",PinType.PinSubCategoryObject=None,PinType.PinSubCategoryMemberReference=(),PinType.bIsArray=False,PinType.bIsReference=False,PinType.bIsConst=False,PinType.bIsWeakPointer=False,LinkedTo=(K2Node_CallFunction_8 7D29E35244DEB831C54B1CAEF4A5FE5C,),PersistentGuid=00000000000000000000000000000000,bHidden=False,bNotConnectable=False,bDefaultValueIsReadOnly=False,bDefaultValueIsIgnored=False,bAdvancedView=False,)
+	End Object
+	Begin Object Class=K2Node_Knot Name="K2Node_Knot_4"
+		NodePosX=608
+		NodePosY=224
+		NodeGuid=230B64A94C55B4E5636829B6DC0D21E7
+		CustomProperties Pin (PinId=4BF4C00042C37A0CDB25AD926D106B22,PinName="InputPin",PinType.PinCategory="name",PinType.PinSubCategoryObject=None,PinType.PinSubCategoryMemberReference=(),PinType.bIsArray=False,PinType.bIsReference=False,PinType.bIsConst=False,PinType.bIsWeakPointer=False,LinkedTo=(K2Node_VariableGet_37 62FD25D94D3A5CAE32C5E685EDA461B1,),PersistentGuid=00000000000000000000000000000000,bHidden=False,bNotConnectable=False,bDefaultValueIsReadOnly=False,bDefaultValueIsIgnored=True,bAdvancedView=False,)
+		CustomProperties Pin (PinId=A0911D3740ABD40031008799EB009AE0,PinName="OutputPin",Direction="EGPD_Output",PinType.PinCategory="name",PinType.PinSubCategoryObject=None,PinType.PinSubCategoryMemberReference=(),PinType.bIsArray=False,PinType.bIsReference=False,PinType.bIsConst=False,PinType.bIsWeakPointer=False,LinkedTo=(K2Node_Knot_3 D6618E7C42476ADDD628FF855FD6E7C9,),PersistentGuid=00000000000000000000000000000000,bHidden=False,bNotConnectable=False,bDefaultValueIsReadOnly=False,bDefaultValueIsIgnored=False,bAdvancedView=False,)
+	End Object
+	Begin Object Class=K2Node_CallFunction Name="K2Node_CallFunction_20"
+		FunctionReference=(MemberParent=Class'/Script/Engine.PhysicalAnimationComponent',MemberName="ApplyPhysicalAnimationProfileBelow")
+		NodePosX=624
+		NodePosY=-16
+		NodeGuid=D1A218604F6806C6A38DF0B1C6812215
+		CustomProperties Pin (PinId=B96C42C240335FB66D65AA84ADD838AC,PinName="execute",PinToolTip="\nExec",PinType.PinCategory="exec",PinType.PinSubCategoryObject=None,PinType.PinSubCategoryMemberReference=(),PinType.bIsArray=False,PinType.bIsReference=False,PinType.bIsConst=False,PinType.bIsWeakPointer=False,LinkedTo=(K2Node_CallFunction_0 BE952CF64A6E491E6C0B8F8B5CEB1000,),PersistentGuid=00000000000000000000000000000000,bHidden=False,bNotConnectable=False,bDefaultValueIsReadOnly=False,bDefaultValueIsIgnored=False,bAdvancedView=False,)
+		CustomProperties Pin (PinId=B38EF5724C795C379DDB809F78D77BA1,PinName="then",PinToolTip="\nExec",Direction="EGPD_Output",PinType.PinCategory="exec",PinType.PinSubCategoryObject=None,PinType.PinSubCategoryMemberReference=(),PinType.bIsArray=False,PinType.bIsReference=False,PinType.bIsConst=False,PinType.bIsWeakPointer=False,LinkedTo=(K2Node_CallFunction_8 399238E643A83FB3CFC85CB926B5CF9F,),PersistentGuid=00000000000000000000000000000000,bHidden=False,bNotConnectable=False,bDefaultValueIsReadOnly=False,bDefaultValueIsIgnored=False,bAdvancedView=False,)
+		CustomProperties Pin (PinId=96A091DE4DF2762BDFA90CB50CAE9BFB,PinName="self",PinFriendlyName=NSLOCTEXT("K2Node", "Target", "Target"),PinToolTip="Target\nPhysical Animation Component Reference",PinType.PinCategory="object",PinType.PinSubCategoryObject=Class'/Script/Engine.PhysicalAnimationComponent',PinType.PinSubCategoryMemberReference=(),PinType.bIsArray=False,PinType.bIsReference=False,PinType.bIsConst=False,PinType.bIsWeakPointer=False,LinkedTo=(K2Node_Knot_2 CC66C9584A8CC28A92B6D4AD7CD40472,),PersistentGuid=00000000000000000000000000000000,bHidden=False,bNotConnectable=False,bDefaultValueIsReadOnly=False,bDefaultValueIsIgnored=False,bAdvancedView=False,)
+		CustomProperties Pin (PinId=EB57682C40C2F631928F3D88FF0B2117,PinName="BodyName",PinToolTip="Body Name\nName\n\nThe body from which we\'d like to start applying the physical animation profile. Finds all bodies below in the skeleton hierarchy. None implies all bodies",PinType.PinCategory="name",PinType.PinSubCategoryObject=None,PinType.PinSubCategoryMemberReference=(),PinType.bIsArray=False,PinType.bIsReference=False,PinType.bIsConst=False,PinType.bIsWeakPointer=False,DefaultValue="None",AutogeneratedDefaultValue="None",LinkedTo=(K2Node_VariableGet_37 62FD25D94D3A5CAE32C5E685EDA461B1,),PersistentGuid=00000000000000000000000000000000,bHidden=False,bNotConnectable=False,bDefaultValueIsReadOnly=False,bDefaultValueIsIgnored=False,bAdvancedView=False,)
+		CustomProperties Pin (PinId=D9E9561844B25F69A79A2C9FE98806C6,PinName="ProfileName",PinToolTip="Profile Name\nName\n\nThe physical animation profile we\'d like to apply. For each body in the physics asset we search for physical animation settings with this name.",PinType.PinCategory="name",PinType.PinSubCategoryObject=None,PinType.PinSubCategoryMemberReference=(),PinType.bIsArray=False,PinType.bIsReference=False,PinType.bIsConst=False,PinType.bIsWeakPointer=False,DefaultValue="DocProfile",AutogeneratedDefaultValue="None",PersistentGuid=00000000000000000000000000000000,bHidden=False,bNotConnectable=False,bDefaultValueIsReadOnly=False,bDefaultValueIsIgnored=False,bAdvancedView=False,)
+		CustomProperties Pin (PinId=7D946B3045556970B2A7519A2614C109,PinName="bIncludeSelf",PinToolTip="Include Self\nBoolean\n\nWhether to include the provided body name in the list of bodies we act on (useful to ignore for cases where a root has multiple children)",PinType.PinCategory="bool",PinType.PinSubCategoryObject=None,PinType.PinSubCategoryMemberReference=(),PinType.bIsArray=False,PinType.bIsReference=False,PinType.bIsConst=False,PinType.bIsWeakPointer=False,DefaultValue="true",AutogeneratedDefaultValue="true",PersistentGuid=00000000000000000000000000000000,bHidden=False,bNotConnectable=False,bDefaultValueIsReadOnly=False,bDefaultValueIsIgnored=False,bAdvancedView=False,)
+		CustomProperties Pin (PinId=41EFD7554F080B09E40F7199FA2BAF6F,PinName="bClearNotFound",PinToolTip="Clear Not Found\nBoolean\n\nIf true, bodies without the given profile name will have any existing physical animation settings cleared. If false, bodies without the given profile name are left untouched.",PinType.PinCategory="bool",PinType.PinSubCategoryObject=None,PinType.PinSubCategoryMemberReference=(),PinType.bIsArray=False,PinType.bIsReference=False,PinType.bIsConst=False,PinType.bIsWeakPointer=False,DefaultValue="true",AutogeneratedDefaultValue="false",PersistentGuid=00000000000000000000000000000000,bHidden=False,bNotConnectable=False,bDefaultValueIsReadOnly=False,bDefaultValueIsIgnored=False,bAdvancedView=False,)
+	End Object
+[/VAR]
\ No newline at end of file
diff --git a/Engine/Documentation/Source/Engine/Physics/PhAT/HowTo/CreateConstraintProfile/CreateConstraintProfile.INT.udn b/Engine/Documentation/Source/Engine/Physics/PhAT/HowTo/CreateConstraintProfile/CreateConstraintProfile.INT.udn
new file mode 100644
--- /dev/null
+++ b/Engine/Documentation/Source/Engine/Physics/PhAT/HowTo/CreateConstraintProfile/CreateConstraintProfile.INT.udn
@@ -0,0 +1,65 @@
+﻿Availability:Docs
+Title:How-To Create a Physics Constraint Profile
+Crumbs: %ROOT%, Engine
+Description:This how-to will show you how to create a Physics Constraint Profile using the Physics Asset Tool.
+parent:Engine/Physics/PhAT/HowTo
+order:1
+Version: 4.13
+skilllevel:Intermediate
+related:Videos/PLZlv_N0_O1ga0aV9jVqJgog0VWz1cLL5f/N1tDjbFXeOo
+topic-image::Engine/Physics/physics_topic.png
+
+In this How-To, we'll cover the basics of creating a **Constraint Profile** and adding **Physics Constraints** to it.
+
+## Steps
+
+1. Open or [Create](Engine/Physics/PhAT/HowTo/CreatePhysicsAsset) a **Physics Asset** for your **Skeletal Mesh** using the **Content Browser**.
+
+	![](Engine\Physics\PhAT\HowTo\CreatePhysicalAnimationProfile\openPhATAsset.png)
+
+	[REGION:note]
+	If opting to create a new **Physics Asset**, you'll need to set it up before continuing.
+	[/REGION]
+
+1. Open the Physics Asset window by using the **Window Menu** -> **Physics Asset** option.
+
+	[REGION:raw]
+	![](Engine\Physics\PhAT\HowTo\CreatePhysicalAnimationProfile\openPhATPhysicsAssetWindow.png)(convert:false)
+	[/REGION]
+
+1. Add Profile using the **+** button on the **Constraints Profile** entry, and set its name (in the image below we've named it `DocsConstraintProfile`).
+
+	![](physConProfAddEntry.png)
+
+1. **Close** the Physics Assset Window.
+
+1. Using the **Profile Button**, select the desired **Constraint Profile** to edit (again, we've chosen `DocsConstraintProfile`).
+
+	![](physConProfSelectProf.png)
+
+1. Switch the Physics Asset Editor to **Constraints Mode** using the **Constraints** ![](physConProfConsButton.png) button.
+
+1. Select the **Physics Constraints** you wish to include in the **Constraint Profile**.
+
+1. Press the **Add to Current Profile** button in the **Details Panel**.
+
+	![](physConProfAddToProf.png)
+
+1. Adjust the settings for the selected **Physics Constraints**.
+
+	![](physConProfAdjust.png)
+
+1. Repeat steps 7-9 for **Physics Constraints** you desire to have in the **Physical Animation Profile**.
+
+1. Save your **Physics Asset** using the **Save** button in the **Physics Asset Tool**.
+
+	[REGION:raw]
+	![](Engine\Physics\PhAT\HowTo\CreatePhysicalAnimationProfile\saveButton.png)(convert:false)
+	[/REGION]
+
+## Result
+
+The **Physics Asset** will have a profile that can be called upon from Blueprint or C++ to change its **Physical Constraint** properties.
+
+
+
diff --git a/Engine/Documentation/Source/Engine/Physics/PhAT/HowTo/CreateNewPhATPhysicsBodies/CreateNewPhATPhysicsBodies.INT.udn b/Engine/Documentation/Source/Engine/Physics/PhAT/HowTo/CreateNewPhATPhysicsBodies/CreateNewPhATPhysicsBodies.INT.udn
new file mode 100644
--- /dev/null
+++ b/Engine/Documentation/Source/Engine/Physics/PhAT/HowTo/CreateNewPhATPhysicsBodies/CreateNewPhATPhysicsBodies.INT.udn
@@ -0,0 +1,28 @@
+﻿Availability: Public
+Title:Creating a New Physics Body in PhAT
+Crumbs: %ROOT%, Engine
+Description:This how-to covers the procedures for creating new Physics Bodies in a Physics Asset using the Physics Asset Tool.
+parent:Engine/Physics/PhAT/HowTo
+order:1
+Version: 4.13
+related:Engine/Physics/PhAT/HowTo/EditPhATPhysicsBodies
+related:Engine/Physics/PhAT/HowTo/EditPhATPhysicsConstraints
+skilllevel:Beginner
+topic-image:Engine/Physics/physics_topic.png
+
+<!-- Overview Text -->
+
+## Steps
+
+![](Engine/Physics/PhAT/HowTo/PhATUG_addBody.png)(w:900)
+
+1. **Right-click** an already existing Physics Body or a Bone in the Hierarchy Panel.
+1. **Select** the type of Physics Body you wish to add from the Context Menu using the left mouse button.
+
+	If the Bone in the Hierarchy Panel does not already have a Physics Body on it, the only option in the **right-click** context menu is "Add Body."
+
+	Finally, you can just **double-click** a bone that does not have a Physics Body in the **Hierarchy** Panel to "Add Body."
+
+## Result
+
+<!-- Result Text -->
\ No newline at end of file
diff --git a/Engine/Documentation/Source/Engine/Physics/PhAT/HowTo/CreatePhysicalAnimationProfile/CreatePhysicalAnimProfile.INT.udn b/Engine/Documentation/Source/Engine/Physics/PhAT/HowTo/CreatePhysicalAnimationProfile/CreatePhysicalAnimProfile.INT.udn
new file mode 100644
--- /dev/null
+++ b/Engine/Documentation/Source/Engine/Physics/PhAT/HowTo/CreatePhysicalAnimationProfile/CreatePhysicalAnimProfile.INT.udn
@@ -0,0 +1,72 @@
+﻿Availability:Docs
+Title:How-To Create a Physical Animation Profile
+Crumbs: %ROOT%, Engine
+Description:This how-to will show you how to create a Physical Animation Profile using the Physics Asset Tool.
+parent:Engine/Physics/PhAT/HowTo
+order:1
+Version: 4.13
+skilllevel:Intermediate
+related:Videos/PLZlv_N0_O1ga0aV9jVqJgog0VWz1cLL5f/N1tDjbFXeOo
+related:Engine/Physics/PhAT/HowTo/ApplyPhysicalAnimationProfile
+topic-image:Engine/Physics/physics_topic.png
+
+In this How-To, we'll cover the basics of creating a **Physical Animation Profile** and adding **Physics Bodies** to it.
+
+## Steps
+
+1. Open or [Create](Engine/Physics/PhAT/HowTo/CreatePhysicsAsset) a **Physics Asset** for your **Skeletal Mesh** using the **Content Browser**.
+
+	![](openPhATAsset.png)
+
+	[REGION:note]
+	If opting to create a new **Physics Asset** you'll need to set it up before continuing.
+	[/REGION]
+
+1. Open the Physics Asset window by using the **Window Menu** -> **Physics Asset** option.
+
+	[REGION:raw]
+	![](openPhATPhysicsAssetWindow.png)(convert:false)
+	[/REGION]
+
+1. Add Profile using the **+** button on the **Physical Animation Profiles** entry, and set its name (in the image below we've named it `DocsProfile`).
+
+	![](physAnimSetProfileName.png)
+
+1. **Close** the Physics Assset Window.
+
+1. Using the **Profile Button** select the desired **Physical Animation Profile** to edit (again, we've chosen `DocsProfile`).
+
+	![](physAnimSetUsedProfile.png)
+
+1. Switch the Physics Asset Editor to **Physics Body Mode** using the **Bodies** ![](physAnimBodyButton.png) button.
+
+1. Select the **Physics Bodies** you wish to include in the **Physical Animation Profile**.
+
+1. Press the **Add to Profile: YourProfileName** button in the **Details Panel**.
+
+	![](physAnimAddToProfile.png)
+
+1. Adjust the settings for the selected **Physics Bodies**.
+
+	![](physAnimDetailsAdjust.png)
+
+	[REGION:tip]
+	A solid place to start for these values is 1000, 100, 1000, 100, 0, 0
+
+	![](PhysAnimStartValues.png)
+	[/REGION]
+
+1. Repeat steps 7-9 for the **Physics Bodies** you desire to have in the **Physical Animation Profile**.
+
+1. Save your **Physics Asset** using the **Save** button in the **Physics Asset Tool**.
+
+	[REGION:raw]
+	![](saveButton.png)(convert:false)
+	[/REGION]
+
+## Result
+
+The **Physics Asset** will have a profile that can be called upon from Blueprint or C++ to change the **Physical Animation** properties of its **Physics Bodies**.
+
+
+
diff --git a/Engine/Documentation/Source/Engine/Physics/PhAT/HowTo/CreatePhysicsAsset/CreatePhysicsAsset.INT.udn b/Engine/Documentation/Source/Engine/Physics/PhAT/HowTo/CreatePhysicsAsset/CreatePhysicsAsset.INT.udn
new file mode 100644
--- /dev/null
+++ b/Engine/Documentation/Source/Engine/Physics/PhAT/HowTo/CreatePhysicsAsset/CreatePhysicsAsset.INT.udn
@@ -0,0 +1,36 @@
+﻿Availability: Public
+Title:Creating a New Physics Asset
+Crumbs: %ROOT%, Engine
+Description:This how-to covers the procedures for creating a new Physics Asset.
+parent:Engine/Physics/PhAT/HowTo
+order:1
+Version: 4.13
+related:Engine/Physics/PhAT/HowTo/EditPhysicsAsset
+skilllevel:Beginner
+topic-image:Engine/Physics/physics_topic.png
+
+There are two ways to create a new **Physics Asset**: on import or by using the context menu in the **Content Browser**. Below are the steps and interfaces for both methods.
+
+## Steps
+
+When a Skeletal Mesh is imported, there is an option to generate a Physics Asset for it as it is imported. Once the imported file is processed, you will be presented with the 
+"New Asset" dialog:
+
+[REGION:raw]
+![](Engine/Physics/PhAT/HowTo/PhATUG_ImportPA.png)
+[/REGION]
+
+However, you can follow these steps if you need to create a Physics Asset for a Skeletal Mesh at a later time:
+
+1. Find the Skeletal Mesh asset you wish to add a Physics Asset to in the 
+**Content Browser**.
+1. **Right-click** on the Skeletal Mesh in the Content Browser to bring up the **Context Menu** -> **Create** -> **Create Physics Asset**.
+1. Adjust setting to your liking.
+
+	![](Shared/PhAT/PhAT_restetDefault.png)
+
+1. **Click** 'Ok'.
+
+## Result
+
+Upon creating a **Physics Asset** you will find it in the same folder as the **Skeletal Mesh** it is based on.
\ No newline at end of file
diff --git a/Engine/Documentation/Source/Engine/Physics/PhAT/HowTo/EditPhATPhysicsBodies/EditPhATPhysicsBodies.INT.udn b/Engine/Documentation/Source/Engine/Physics/PhAT/HowTo/EditPhATPhysicsBodies/EditPhATPhysicsBodies.INT.udn
new file mode 100644
--- /dev/null
+++ b/Engine/Documentation/Source/Engine/Physics/PhAT/HowTo/EditPhATPhysicsBodies/EditPhATPhysicsBodies.INT.udn
@@ -0,0 +1,68 @@
+﻿Availability: Public
+Title:Editing Physics Asset Physics Bodies
+Crumbs: %ROOT%, Engine
+Description:This how-to covers the procedures for editing the Physics Bodies of a Physics Asset.
+parent:Engine/Physics/PhAT/HowTo
+order:1
+Version: 4.13
+related:Engine/Physics/PhAT/HowTo/EditPhATPhysicsConstraints
+skilllevel:Beginner
+topic-image:Engine/Physics/physics_topic.png
+
+Below are a number of common procedures, and steps, associated with editing **Physics Bodies** in the **Physics Asset Tool**.
+
+## Editing a Physics Body
+
+1. Enable Body Editing Mode ![](Shared\Icons\icon_PHatMode_Body_40x.png).
+1. Select a Physics Body in the Viewport or in the **Hierarchy Pane**.
+	* If it is difficult to see the Physics Body you desire, use the **Viewport** -> **Modes** -> **MeshRender Mode** -> **Wireframe/None** to adjust the visibility of the 
+	Skeletal Mesh.
+1. **Move, rotate, and scale** the Physics Body using the Translation, Rotation, and Scale tools so it fits the Skeletal Mesh and Skeletal Joint.
+1. Edit the Physics Body's properties in the Details Pane.
+1. **Enable** or **Disable** Collision between the current Physics Bodies and other in the Physics Asset using the _Collision On_ and _Collision Off_ tools. Any Physics Body that can 
+collide with the currently selected Physics Body will be blue, if they cannot collide, they will be gray.
+1. If a series of bodies should act as one, such as with the twist joints that form a wrist, use the _Weld_ tool to prevent undesired physics simulation behavior. Welded Physics 
+Bodies will be yellow in color.
+1. Save ![](Shared\Icons\icon_SaveAsset_40x.png) often.
+
+See the **[Physics Asset Tool Reference](Engine/Physics/PhAT/Reference)** for more information on the properties in the Physics Asset Tool.
+
+[PUBLISH:licensee]
+[REGION:note]
+As of change list #1977557, multi select is enabled for Physics Bodies and Physics Constraints, enabling you to edit multiple bodies or constraints at the same time.
+[/REGION]
+[/PUBLISH]
+
+## Copying Physics Bodies
+
+While in either mode, properties from one Physics Body can be copied to other Physics Bodies:
+
+1. **Select** the Physics Body you wish to copy data from.
+1. **Press Ctrl + C**.
+1. **Select** the Physics Bodies you wish to apply data to.
+1. **Press Ctrl + V**.
+
+This will not overwrite the shape of a Physics Body.
+
+## Deleting Physics Bodies
+
+1. **Enable** Body Editing Mode ![](Shared\Icons\icon_PHatMode_Body_40x.png).
+1. **Select** a Physics Body you wish to delete.
+	* If it is difficult to see the Physics Body you desire, use the **Viewport -> Modes -> MeshRender Mode -> Wireframe/None** to adjust the visibility of the Skeletal Mesh.
+1. Press the **Delete** key.
+
+	![](Shared\PhAT\PhAT_DelBone.png)
+	
+## Physical Materials
+
+Each Physics Body in a Physics Asset can have a Physical Material assigned to it. To apply a Physics Material to a single Physics Body:
+
+1. **Enable** Body Editing Mode ![](Shared\Icons\icon_PHatMode_Body_40x.png).
+1. **Select** a Physics Body in the Viewport or in the **Hierarchy** Pane.
+	* If it is difficult to see the Physics Body you desire, use the **Viewport** -> **Modes** -> **MeshRender Mode** -> **Wireframe/None** to adjust the visibility of the Skeletal 
+	Mesh.
+1. In the Details pane, locate the Simple Collision Physical Material property and assign a Physics Material.
+
+[REGION:raw]
+![](Engine/Physics/PhysicalMaterials/HowTo/PhysMat_PhAT_04.png)(w:920)
+[/REGION]
diff --git a/Engine/Documentation/Source/Engine/Physics/PhAT/HowTo/EditPhATPhysicsConstraints/EditPhATPhysicsConstraints.INT.udn b/Engine/Documentation/Source/Engine/Physics/PhAT/HowTo/EditPhATPhysicsConstraints/EditPhATPhysicsConstraints.INT.udn
new file mode 100644
--- /dev/null
+++ b/Engine/Documentation/Source/Engine/Physics/PhAT/HowTo/EditPhATPhysicsConstraints/EditPhATPhysicsConstraints.INT.udn
@@ -0,0 +1,80 @@
+﻿Availability: Public
+Title:Editing Physics Asset Constraints
+Crumbs: %ROOT%, Engine
+Description:This how-to covers the procedures for editing the Constraints of a Physics Asset.
+parent:Engine/Physics/PhAT/HowTo
+order:1
+Version: 4.13
+related:Engine/Physics/PhAT/HowTo/EditPhATPhysicsBodies
+skilllevel:Beginner
+topic-image:Engine/Physics/physics_topic.png
+
+Below are a number of common procedures, and steps, associated with editing **Physics Constraints** in the **Physics Asset Tool**.
+
+## Editing Physics Constraints
+
+[REGION:note]
+The use of Physics Constraints is covered in the [Physics Constraints User Guide](Engine/Physics/Constraints/ConstraintsUserGuide) and their properties are covered in the [Physics Constraint Reference](Engine/Physics\Constraints\ConstraintsReference). This section will only cover workflows that are specific to the Physics Asset Tool or ones that have deviated substantially from the norm.
+[/REGION]
+
+1. **Enable** Constraint Editing Mode ![](Shared\Icons\icon_PHatMode_Joint_40x.png).
+1. **Select** a Physics Constraint in the Viewport or in the Hierarchy Pane.
+1. **Move and rotate** the Physics Constraint using the Translation and Rotation to create the rotational point of the "joint" the Physics Constraint will form.
+1. Edit the Physics Constraint's properties in the Details Pane.
+
+	[REGION:tip]
+	You can switch Swing1, Swing2, and Twist from Limited to Locked quickly by using the "**1**", "**2**", and "**3**" keys respectively. The "**4**" key can be used cycle through 
+	them setting one to limited and the other two to locked.
+	[/REGION]
+
+1. Save ![](Shared\Icons\icon_SaveAsset_40x.png) often.
+
+See the **[Physics Asset Tool Reference](Engine/Physics/PhAT/Reference)** for more information on the properties in the Physics Asset Tool.
+
+## Aligning Physics Constraints
+
+If you are using the Physics Constraint's Linear or Angular limits, you will be able to see their alignment:
+
+![](Engine/Physics/PhAT/HowTo/PhATUG_Limits.png)(w:300 convert:false)
+
+By then translating and rotating the Physics Constraint, you can align the limits to give the desired results. In the most basic of terms, when a Physics Constraint is limited, 
+you can see a yellow line suspended in a green arc or cone structure. The line will be "constrained" within this arc or cone.
+
+For more information on Physics Constraints and their properties, see the **[Physics Constraints User Guide](Engine/Physics\Constraints\ConstraintsUserGuide)**.
+
+## Copying Physics Constraints
+
+While in either mode, properties from one Constraint can be copied to other Constraints:
+
+1. **Select** Constraint you wish to copy data from.
+1. **Press Ctrl + C**.
+1. **Select** Constraints you wish to apply data to.
+1. **Press Ctrl + V**.
+
+## Deleting Physics Constraints
+
+[REGION:note]
+There is no easy way to recreate a Physics Constraint, be mindful of this when deleting them.
+[/REGION]
+
+1. **Enable** Constraint Editing Mode ![](Shared\Icons\icon_PHatMode_Joint_40x.png).
+1. **Select** the Physics Constraint you wish to delete.
+1. **Press** the **Delete** key.
+
+
+## Recreating a Physics Constraint
+
+[REGION:note]
+There is no easy way to recreate a Physics Constraint, be mindful of this when deleting them.
+[/REGION]
+
+Physics Constraints are only created on generation of a Physics Body and are only generated upstream. So if you remove a Physics Constraint from the shoulder, you would have to 
+remove the upper arm Physics Body (which will remove the elbow Physics Constraint), then re-create the upper arm Physics Body. This would create the shoulder Physics Constraint, 
+but not the elbow Physics Constraint, so you would have to continue this process down the arm.
+
+
+## Constraint Snapping
+
+Snapping a Physics Constraint will take the blue locator for the Physics Constraint and snap its location and rotation to the red one's location and rotation.
+
+You can snap the currently selected Physics Constraint(s) by using the Snap Constraint ![](Shared\Icons\icon_PhAT_SnapConstraint_40x.png) button.
diff --git a/Engine/Documentation/Source/Engine/Physics/PhAT/HowTo/EditPhysicsAsset/EditPhysicsAsset.INT.udn b/Engine/Documentation/Source/Engine/Physics/PhAT/HowTo/EditPhysicsAsset/EditPhysicsAsset.INT.udn
new file mode 100644
--- /dev/null
+++ b/Engine/Documentation/Source/Engine/Physics/PhAT/HowTo/EditPhysicsAsset/EditPhysicsAsset.INT.udn
@@ -0,0 +1,29 @@
+﻿Availability: Public
+Title:Editing a Physics Asset
+Crumbs: %ROOT%, Engine
+Description:This how-to covers the procedure for opening an Physics Asset for editing in the Physics Asset Tool
+parent:Engine/Physics/PhAT/HowTo
+order:1
+Version: 4.13
+related:Engine/Physics/PhAT/HowTo/EditPhATPhysicsBodies
+related:Engine/Physics/PhAT/HowTo/EditPhATPhysicsConstraints
+skilllevel:Beginner
+topic-image:Engine/Physics/physics_topic.png
+
+Below are the steps for opening a **Physics Asset** for editing in the **Physics Asset Tool**.
+
+## Steps
+
+1. Find the Physics Asset in the Content Browser.
+
+	![](Engine/Physics\PhysicalMaterials\PhysMatUserGuide\PhysMat_PhAT_02.png)
+
+1. **Double-click** to open it in the Physics Asset Tool (PhAT).
+
+	![](Engine/Physics/PhAT/HowTo/PhATUG_FullWindow.png)(w:900)
+
+## Result
+
+The **Physics Asset** will be open and ready to edit in the **Physics Asset Tool**.
+
+
diff --git a/Engine/Documentation/Source/Engine/Physics/PhAT/HowTo/FixCommonPhysicsErrors/FixCommonPhATPhysicsErrors.INT.udn b/Engine/Documentation/Source/Engine/Physics/PhAT/HowTo/FixCommonPhysicsErrors/FixCommonPhATPhysicsErrors.INT.udn
new file mode 100644
--- /dev/null
+++ b/Engine/Documentation/Source/Engine/Physics/PhAT/HowTo/FixCommonPhysicsErrors/FixCommonPhATPhysicsErrors.INT.udn
@@ -0,0 +1,45 @@
+﻿Availability: Public
+Title:Troubleshooting Common Physics Asset Errors
+Crumbs: %ROOT%, Engine
+Description:This how-to covers how to troubleshoot the two common Physics Asset issues of Exploding and Jittering.
+parent:Engine/Physics/PhAT/HowTo
+order:1
+Version: 4.13
+skilllevel:Intermediate
+topic-image:Engine/Physics/physics_topic.png
+
+While there are a number of things that can be out of your control during a physical simulation, there are two issues that are definately within your power to correct: Exploding 
+(interpenetrating physics bodies that are set to collide with eatch other) and Jitter (when the Physics Bodies refuse to sleep due to micro movements). Below are some steps for correcting 
+these issues.
+
+## Exploding
+
+This is caused by two Physics Bodies interpenetrating and the physics system attempting to correct it by applying a massive amount of force to expel the Physics Bodies from each other. 
+If a Physics Constraint is keeping the offending Physics Bodies together, the physics system will continually apply force to separate them, causing very erratic and extreme movement.
+
+![](Engine/Physics/PhAT/HowTo/brokenPhAT.png)
+
+This is normally fixed by disabling collision between the two offending Physics Bodies, or changing their positions and/or scale to make sure they do not interpenetrate. Also, welded 
+Physics Bodies act as one regardless of interpenetration.
+
+## Jitter
+
+If the Physics Asset is mostly "collapsed" but still shaking and twitching on the ground, but not violently bouncing around, there are a few values you can tweak to get it to stop 
+and sleep. 
+
+Before doing anything, try simulating with No Gravity, this will show you if any of your Physics Constraints are miss aligned and trying to correct before the Physics Asset has 
+even hit the ground.
+
+Often times, a small amount of Linear and Angular Damping on all the Physics Asset's Physics Bodies is more than enough to stop the asset from shaking. However, using large 
+values in Linear Damping will cause a Physics Body slow its movement through the world, even due to gravity. Damping, as it is implemented here, is not "drag" per se. It is 
+more like moving through a viscous substance.
+
+If jitter is still an issue, perhaps you have a number of really small Physics Bodies, check the **[Substepping](Engine/Physics/Substepping)** documentation for information on having the 
+physics simulation calculate in-between solutions.
+
+See the **[](Engine/Physics/PhAT/Reference)** and **[](Engine/Physics/Constraints/ConstraintsReference)** for more details on the properties for Physics Constraints.
+
+
+
+
+
diff --git a/Engine/Documentation/Source/Engine/Physics/PhAT/HowTo/HowTo.INT.udn b/Engine/Documentation/Source/Engine/Physics/PhAT/HowTo/HowTo.INT.udn
--- a/Engine/Documentation/Source/Engine/Physics/PhAT/HowTo/HowTo.INT.udn
+++ b/Engine/Documentation/Source/Engine/Physics/PhAT/HowTo/HowTo.INT.udn
@@ -1,334 +1,57 @@
 ﻿Availability: Public
-Title:Physics Asset Tool How To Guide
+Title:Physics Asset Tool How-To Landing Page
 Crumbs: %ROOT%, Engine, Engine/Physics, Engine/Physics/PhAT
-Description:How to use the Physics Asset Tool to create Physics Asset for Skeletal Meshes.
-Version: 4.9
+Description:Landing page for the Physics Asset Tool how-to section.
+Navigation:topic
+parent:Engine/Physics/PhAT
+order:1
+Version: 4.13
 
 [VAR:Topic]
 [OBJECT:Topic]
 	[PARAM:image]
-		![%Engine/Physics/PhAT/UserGuide:title%](Engine/Physics/physics_topic.png)
+		![%Engine/Physics/PhAT/HowTo:title%](Engine/Physics/physics_topic.png)
 	[/PARAM]
 	[PARAM:icon]
-		![](Engine/Physics/PhAT/HowTo/%ROOT%/physics_icon.png)(convert:false)
+		![](%ROOT%/physics_icon.png)(convert:false)
 	[/PARAM]
 	[PARAM:title]
-		%Engine/Physics/PhAT/UserGuide:title%
+		%Engine/Physics/PhAT/HowTo:title%
 	[/PARAM]
 	[PARAM:description]
-		%Engine/Physics/PhAT/UserGuide:description%
+		%Engine/Physics/PhAT/HowTo:description%
 	[/PARAM]
 	[PARAM:path]
-		[RELATIVE:Engine/Physics/Phat/UserGuide]
+		[RELATIVE:Engine/Physics/Phat/HowTo]
 	[/PARAM]
 [/OBJECT]
 [/VAR]
 [VAR:TopicCompact]
 [OBJECT:TopicCompact]
 	[PARAM:image]
-		![%Engine/Physics/PhAT/UserGuide:title%](Engine/Physics/physics_topic.png)
+		![%Engine/Physics/PhAT/HowTo:title%](Engine/Physics/physics_topic.png)
 	[/PARAM]
 	[PARAM:icon]
-		![](Engine/Physics/PhAT/HowTo/%ROOT%/physics_icon.png)(convert:false)
+		![](%ROOT%/physics_icon.png)(convert:false)
 	[/PARAM]
 	[PARAM:title]
-		%Engine/Physics/PhAT/UserGuide:title%
+		%Engine/Physics/PhAT/HowTo:title%
 	[/PARAM]
 	[PARAM:description]
-		%Engine/Physics/PhAT/UserGuide:description%
+		%Engine/Physics/PhAT/HowTo:description%
 	[/PARAM]
 	[PARAM:path]
-		[RELATIVE:Engine/Physics/Phat/UserGuide]
+		[RELATIVE:Engine/Physics/Phat/HowTo]
 	[/PARAM]
 [/OBJECT]
 [/VAR]
 
+This page contains links to guides on how-to perform common tasks in the **Physics Asset Tool** (**PhAT**).
 
+## Beginner
 
+[DIR(output:"fancy" parent:"Engine/Physics/PhAT/HowTo" org:"hierarchy" end:"1" skill_level:"beginner")]
 
-The **Physics Asset Tool** is used for creating **Physics Assets** for use with **Skeletal Meshes**. One can start from nothing and build to a full ragdoll setup or use the automation 
-tools to create a basic set of Physics Bodies and Physics Constraints.
-
-This document will cover step by step procedures for the most common activities in the Physics Asset Tool.
-
-## New Physics Asset
-
-When a Skeletal Mesh is imported, there is an option to generate a Physics Asset for it as it is imported. Once the imported file is processed, you will be presented with the 
-"New Asset" dialog:
-
-[REGION:raw]
-![](Engine/Physics/PhAT/HowTo/PhATUG_ImportPA.png)
-[/REGION]
-
-However, you can follow these steps if you need to create a Physics Asset for a Skeletal Mesh at a later time:
-
-1. Find the Skeletal Mesh asset you wish to add a Physics Asset to in the 
-**Content Browser**.
-1. **Right-click** on the Skeletal Mesh in the Content Browser to bring up the **Context Menu** -> **Create** -> **Create Physics Asset**.
-1. Adjust setting to your liking.
-
-	![](Shared\PhAT\PhAT_restetDefault.png)
-
-1. **Click** 'Ok'.
-
-## Editing a Physics Asset
-
-1. Find the Physics Asset in the Content Browser.
-
-	![](Engine/Physics\PhysicalMaterials\PhysMatUserGuide\PhysMat_PhAT_02.png)
-
-1. **Double-click** to open it in the Physics Asset Tool (PhAT).
-
-	![](Engine/Physics/PhAT/HowTo/PhATUG_FullWindow.png)(w:900)
-
-
-## Physics Bodies
-
-### Editing Physics Bodies
-
-1. Enable Body Editing Mode ![](Shared\Icons\icon_PHatMode_Body_40x.png).
-1. Select a Physics Body in the Viewport or in the **Hierarchy Pane**.
-	* If it is difficult to see the Physics Body you desire, use the **Viewport** -> **Modes** -> **MeshRender Mode** -> **Wireframe/None** to adjust the visibility of the 
-	Skeletal Mesh.
-1. **Move, rotate, and scale** the Physics Body using the Translation, Rotation, and Scale tools so it fits the Skeletal Mesh and Skeletal Joint.
-1. Edit the Physics Body's properties in the Details Pane.
-1. **Enable** or **Disable** Collision between the current Physics Bodies and other in the Physics Asset using the _Collision On_ and _Collision Off_ tools. Any Physics Body that can 
-collide with the currently selected Physics Body will be blue, if they cannot collide, they will be gray.
-1. If a series of bodies should act as one, such as with the twist joints that form a wrist, use the _Weld_ tool to prevent undesired physics simulation behavior. Welded Physics 
-Bodies will be yellow in color.
-1. Save ![](Shared\Icons\icon_SaveAsset_40x.png) often.
-
-See the **[Physics Asset Tool Reference](Engine/Physics/PhAT/Reference)** for more information on the properties in the Physics Asset Tool.
-
-[PUBLISH:licensee]
-[REGION:note]
-As of change list #1977557, multi select is enabled for Physics Bodies and Physics Constraints, enabling you to edit multiple bodies or constraints at the same time.
-[/REGION]
-[/PUBLISH]
-
-### Deleting Physics Bodies
-
-1. **Enable** Body Editing Mode ![](Shared\Icons\icon_PHatMode_Body_40x.png).
-1. **Select** a Physics Body you wish to delete.
-	* If it is difficult to see the Physics Body you desire, use the **Viewport -> Modes -> MeshRender Mode -> Wireframe/None** to adjust the visibility of the Skeletal Mesh.
-1. Press the **Delete** key.
-
-	![](Shared\PhAT\PhAT_DelBone.png)
-
-### Creating New Physics Bodies
-
-![](Engine/Physics/PhAT/HowTo/PhATUG_addBody.png)(w:900)
-
-1. **Right-click** an already existing Physics Body or a Bone in the Hierarchy Panel.
-1. **Select** the type of Physics Body you wish to add from the Context Menu using the left mouse button.
-
-	If the Bone in the Hierarchy Panel does not already have a Physics Body on it, the only option in the **right-click** context menu is "Add Body."
-
-	Finally, you can just **double-click** a bone that does not have a Physics Body in the **Hierarchy** Panel to "Add Body."
-
-
-### Enabling and Disabling Collision Between Physics Bodies
-
-[REGION:simpletable]
-| --- | --- |
-|[INCLUDE:#ImgColOn]| [INCLUDE:#ColOn] |
-|[INCLUDE:#ImgColOff]| [INCLUDE:#ColOff] |
-[/REGION]
-
-<!--
-
-[EXCERPT:ImgColOn]
-[REGION:lightbox]
-[![](Shared\PhAT\PhAT_collisionOn.png)(w:350)](Shared\PhAT\PhAT_collisionOn.png)
-[/REGION]
-[/EXCERPT:ImgColOn]
-
-[EXCERPT:ImgColOff]
-[REGION:lightbox]
-[![](Shared\PhAT\PhAT_collisionOff.png)(w:350)](Shared\PhAT\PhAT_collisionOff.png)
-[/REGION]
-[/EXCERPT:ImgColOff]
-
-[EXCERPT:ColOn]
-[REGION:raw]
-**Enabling Collision**
-
-1. **Select** 2 or more Physics Bodies using **Ctrl + Left Mouse click**.
-1. **Click** the "Collision On" ![](Shared\Icons\icon_PhAT_EnableCollision_40x.png) button.
-
-
-[/REGION]
-[/EXCERPT:ColOn]
-
-[EXCERPT:ColOff]
-**Disabling Collision** 
-
-1. **Select** 2 or more Physics Bodies using **Ctrl + Left Mouse click**.
-1. **Click** the "Collision Off" ![](Shared\Icons\icon_PhAT_DisableCollision_40x.png) button.
-
-[/EXCERPT:ColOff]
-
--->
-
-
-### Welding
-
-Welding Physics Bodies will make multiple bodies interact as one, as well as lock together the Skeletal Mesh Joints they are associated with. To Weld one Physics Body to another:
-
-1. **Select** 2 or more Physics Bodies using **Ctrl + Left Mouse click**.
-1. **Click** the "Weld" ![](Shared\Icons\icon_PhAT_Weld_40x.png) button.
-
-
-Any Physics Body welded to the currently selected Physics Body will show in Yellow.
-
-
-### Physical Materials
-
-Each Physics Body in a Physics Asset can have a Physical Material assigned to it. To apply a Physics Material to a single Physics Body:
-
-1. **Enable** Body Editing Mode ![](Shared\Icons\icon_PHatMode_Body_40x.png).
-1. **Select** a Physics Body in the Viewport or in the **Hierarchy** Pane.
-	* If it is difficult to see the Physics Body you desire, use the **Viewport** -> **Modes** -> **MeshRender Mode** -> **Wireframe/None** to adjust the visibility of the Skeletal 
-	Mesh.
-1. In the Details pane, locate the Simple Collision Physical Material property and assign a Physics Material.
-
-[REGION:raw]
-![](Engine/Physics/PhysicalMaterials/HowTo/PhysMat_PhAT_04.png)(w:920)
-[/REGION]
-
-
-## Constraints
-
-[REGION:note]
-The use of Physics Constraints is covered in the [Physics Constraints User Guide](Engine/Physics\Constraints\ConstraintsUserGuide) and their properties are covered in the [Physics Constraint Reference](Engine/Physics\Constraints\ConstraintsReference). This section will only cover workflows that are specific to the Physics Asset Tool or ones that have deviated substantially from the norm.
-[/REGION]
-
-### Editing Physics Constraints
-
-1. **Enable** Constraint Editing Mode ![](Shared\Icons\icon_PHatMode_Joint_40x.png).
-1. **Select** a Physics Constraint in the Viewport or in the Hierarchy Pane.
-1. **Move and rotate** the Physics Constraint using the Translation and Rotation to create the rotational point of the "joint" the Physics Constraint will form.
-1. Edit the Physics Constraint's properties in the Details Pane.
-
-	[REGION:tip]
-	You can switch Swing1, Swing2, and Twist from Limited to Locked quickly by using the "**1**", "**2**", and "**3**" keys respectively. The "**4**" key can be used cycle through 
-	them setting one to limited and the other two to locked.
-	[/REGION]
-
-1. Save ![](Shared\Icons\icon_SaveAsset_40x.png) often.
-
-
-
-See the **[Physics Asset Tool Reference](Engine/Physics/PhAT/Reference)** for more information on the properties in the Physics Asset Tool.
-
-
-
-
-### Aligning Physics Constraints
-
-If you are using the Physics Constraint's Linear or Angular limits, you will be able to see their alignment:
-
-![](Engine/Physics/PhAT/HowTo/PhATUG_Limits.png)(w:300 convert:false)
-
-By then translating and rotating the Physics Constraint, you can align the limits to give the desired results. In the most basic of terms, when a Physics Constraint is limited, 
-you can see a yellow line suspended in a green arc or cone structure. The line will be "constrained" within this arc or cone.
-
-For more information on Physics Constraints and their properties, see the **[Physics Constraints User Guide](Engine/Physics\Constraints\ConstraintsUserGuide)**.
-
-
-### Deleting Physics Constraints
-
-[REGION:note]
-There is no easy way to recreate a Physics Constraint, be mindful of this when deleting them.
-[/REGION]
-
-1. **Enable** Constraint Editing Mode ![](Shared\Icons\icon_PHatMode_Joint_40x.png).
-1. **Select** the Physics Constraint you wish to delete.
-1. **Press** the **Delete** key.
-
-
-### Recreating a Physics Constraint
-
-[REGION:note]
-There is no easy way to recreate a Physics Constraint, be mindful of this when deleting them.
-[/REGION]
-
-Physics Constraints are only created on generation of a Physics Body and are only generated upstream. So if you remove a Physics Constraint from the shoulder, you would have to 
-remove the upper arm Physics Body (which will remove the elbow Physics Constraint), then re-create the upper arm Physics Body. This would create the shoulder Physics Constraint, 
-but not the elbow Physics Constraint, so you would have to continue this process down the arm.
-
-
-### Constraint Snapping
-
-Snapping a Physics Constraint will take the blue locator for the Physics Constraint and snap its location and rotation to the red one's location and rotation.
-
-You can snap the currently selected Physics Constraint(s) by using the Snap Constraint ![](Shared\Icons\icon_PhAT_SnapConstraint_40x.png) button.
-
-
-## Both Physics Bodies and Physics Constraints
-
-### Copying Bodies and Constraints
-
-While in either mode, properties from one Physics Body or Constraint can be copied to other Physics Bodies or Constraints:
-
-1. **Select** the Physics Body or Constraint you wish to copy data from.
-1. **Press Ctrl + C**.
-1. **Select** the Physics Bodies or Constraints you wish to apply data to.
-1. **Press Ctrl + V**.
-
-
-This will not overwrite the shape of a Physics Body.
-
-
-## Testing
-
-Clicking the Simulation ![](Shared\Icons\icon_PhAT_PlaySim_40x.png) button will allow you to test your Physics Asset.
-
-![](Engine/Physics/PhAT/HowTo/SimPhAT.png)(w:600)
-
-The downward triangle next to the Simulation button will give you access to **Real Simulation** and **No Gravity**:
-
-* **Real Simulation** - Simulates the entire Physics Asset with gravity enabled.
-* **No Gravity** - Simulates the entire Physics Asset but does not turn on gravity, leaving you to **ctrl+click** to poke the Physics Bodies in a zero gravity environment.
-
-	This is useful for finding any interpenetrating Physics Bodies or Limited Physics Constraints already outside of their limits.
-
-You can also simulate a chain of joints by toggling **Selected Simulation** ![](Engine/Physics/PhAT/HowTo/selectedToggle_button.png) on, then clicking the **Simulation** ![](Shared\Icons\icon_PhAT_PlaySim_40x.png) button.
-
-* **Selected Simulation** - Does mostly what it sounds like, it only simulates the Physics Bodies you have selected (using **ctrl+left click** you can select more than one 
-Physics Body) and those down the hierarchy from the selected Physics Bodies. So if you select the shoulder, the entire arm will simulate.
-
-
-
-## Exploding
-
-This is caused by two Physics Bodies interpenetrating and the physics system attempting to correct it by applying a massive amount of force to expel the Physics Bodies from each other. 
-If a Physics Constraint is keeping the offending Physics Bodies together, the physics system will continually apply force to separate them, causing very erratic and extreme movement.
-
-![](Engine/Physics/PhAT/HowTo/brokenPhAT.png)
-
-This is normally fixed by disabling collision between the two offending Physics Bodies, or changing their positions and/or scale to make sure they do not interpenetrate. Also, welded 
-Physics Bodies act as one regardless of interpenetration.
-
-## Jitter
-
-If the Physics Asset is mostly "collapsed" but still shaking and twitching on the ground, but not violently bouncing around, there are a few values you can tweak to get it to stop 
-and sleep. 
-
-Before doing anything, try simulating with No Gravity, this will show you if any of your Physics Constraints are miss aligned and trying to correct before the Physics Asset has 
-even hit the ground.
-
-Often times, a small amount of Linear and Angular Damping on all the Physics Asset's Physics Bodies is more than enough to stop the asset from shaking. However, using large 
-values in Linear Damping will cause a Physics Body slow its movement through the world, even due to gravity. Damping, as it is implemented here, is not "drag" per se. It is 
-more like moving through a viscous substance.
-
-If jitter is still an issue, perhaps you have a number of really small Physics Bodies, check the **[Substepping](Engine/Physics/Substepping)** documentation for information on having the 
-physics simulation calculate in-between solutions.
-
-See the **[](Engine/Physics/PhAT/Reference)** and **[](Engine/Physics/Constraints/ConstraintsReference)** for more details on the properties for Physics Constraints.
-
-
-
-
-
+## Intermediate
 
+[DIR(output:"fancy" parent:"Engine/Physics/PhAT/HowTo" org:"hierarchy" end:"1" skill_level:"intermediate")]
\ No newline at end of file
diff --git a/Engine/Documentation/Source/Engine/Physics/PhAT/HowTo/TestPhysicsAssets/TestPhysicsAssets.INT.udn b/Engine/Documentation/Source/Engine/Physics/PhAT/HowTo/TestPhysicsAssets/TestPhysicsAssets.INT.udn
new file mode 100644
--- /dev/null
+++ b/Engine/Documentation/Source/Engine/Physics/PhAT/HowTo/TestPhysicsAssets/TestPhysicsAssets.INT.udn
@@ -0,0 +1,29 @@
+﻿Availability: Public
+Title:Testing Physics Assets
+Crumbs: %ROOT%, Engine
+Description:This how-to covers the procedure for testing your Physics Assets in the Physics Asset Tool.
+parent:Engine/Physics/PhAT/HowTo
+order:1
+Version: 4.13
+skilllevel:Intermediate
+topic-image:Engine/Physics/physics_topic.png
+
+This page covers the basics of testing a **Physics Asset** in the **Physics Asset Tool**.
+
+## Testing
+
+Clicking the Simulation ![](Shared/Icons/icon_PhAT_PlaySim_40x.png) button will allow you to test your Physics Asset.
+
+![](Engine/Physics/PhAT/HowTo/SimPhAT.png)(w:600)
+
+The downward triangle next to the Simulation button will give you access to **Real Simulation** and **No Gravity**:
+
+* **Real Simulation** - Simulates the entire Physics Asset with gravity enabled.
+* **No Gravity** - Simulates the entire Physics Asset but does not turn on gravity, leaving you to **ctrl+click** to poke the Physics Bodies in a zero gravity environment.
+
+	This is useful for finding any interpenetrating Physics Bodies or Limited Physics Constraints already outside of their limits.
+
+You can also simulate a chain of joints by toggling **Selected Simulation** ![](Engine/Physics/PhAT/HowTo/selectedToggle_button.png) on, then clicking the **Simulation** ![](Shared/Icons/icon_PhAT_PlaySim_40x.png) button.
+
+* **Selected Simulation** - Does mostly what it sounds like, it only simulates the Physics Bodies you have selected (using **ctrl+left click** you can select more than one 
+Physics Body) and those down the hierarchy from the selected Physics Bodies. So if you select the shoulder, the entire arm will simulate.
diff --git a/Engine/Documentation/Source/Engine/Physics/PhAT/HowTo/WeldingPhATPhysicsBodies/WeldingPhysicsBodies.INT.udn b/Engine/Documentation/Source/Engine/Physics/PhAT/HowTo/WeldingPhATPhysicsBodies/WeldingPhysicsBodies.INT.udn
new file mode 100644
--- /dev/null
+++ b/Engine/Documentation/Source/Engine/Physics/PhAT/HowTo/WeldingPhATPhysicsBodies/WeldingPhysicsBodies.INT.udn
@@ -0,0 +1,21 @@
+﻿Availability: Public
+Title:Welding Physics Bodies in PhAT
+Crumbs: %ROOT%, Engine
+Description:This how-to covers the procedure for welding two physics bodies together in the Physics Asset Tool.
+parent:Engine/Physics/PhAT/HowTo
+order:1
+Version: 4.13
+related:Engine/Physics/PhAT/HowTo/EditPhATPhysicsBodies
+skilllevel:Intermediate
+topic-image:Engine/Physics/physics_topic.png
+
+This how-to contains the procedures for welding (combining, merging, etc...) two or more **Physics Bodies** together in the **Physics Asset Tool**.
+
+### Welding
+
+Welding Physics Bodies will make multiple bodies interact as one, as well as lock together the Skeletal Mesh Joints they are associated with. To Weld one Physics Body to another:
+
+1. **Select** 2 or more Physics Bodies using **Ctrl + Left Mouse click**.
+1. **Click** the "Weld" ![](Shared\Icons\icon_PhAT_Weld_40x.png) button.
+
+Any Physics Body welded to the currently selected Physics Body will show in Yellow.
diff --git a/Engine/Documentation/Source/Engine/Physics/PhAT/Reference/Properties.INT.udn b/Engine/Documentation/Source/Engine/Physics/PhAT/Reference/Properties.INT.udn
--- a/Engine/Documentation/Source/Engine/Physics/PhAT/Reference/Properties.INT.udn
+++ b/Engine/Documentation/Source/Engine/Physics/PhAT/Reference/Properties.INT.udn
@@ -2,47 +2,12 @@ Availability: Public
 Title:Physics Asset Properties Reference
 Crumbs: %ROOT%, Engine, Engine/Physics, Engine/Physics/PhAT
 Description:A listing of properties and descriptions for Physics Bodies, Physics Constraints, and Physics Assets accessible in the Physics Asset Tool.
-Version: 4.9
-
-[VAR:Topic]
-[OBJECT:Topic]
-	[PARAM:image]
-		![%Engine/Physics/PhAT/Properties:title%](Engine/Physics/physics_topic.png)
-	[/PARAM]
-	[PARAM:icon]
-		![](%ROOT%/physics_icon.png)(convert:false)
-	[/PARAM]
-	[PARAM:title]
-		%Engine/Physics/PhAT/Properties:title%
-	[/PARAM]
-	[PARAM:description]
-		%Engine/Physics/PhAT/Properties:description%
-	[/PARAM]
-	[PARAM:path]
-		[RELATIVE:Engine/Physics/PhAT/Properties]
-	[/PARAM]
-[/OBJECT]
-[/VAR]
-[VAR:TopicCompact]
-[OBJECT:TopicCompact]
-	[PARAM:image]
-		![%Engine/Physics/PhAT/Properties:title%](Engine/Physics/physics_topic.png)
-	[/PARAM]
-	[PARAM:icon]
-		![](%ROOT%/physics_icon.png)(convert:false)
-	[/PARAM]
-	[PARAM:title]
-		%Engine/Physics/PhAT/Properties:title%
-	[/PARAM]
-	[PARAM:description]
-		%Engine/Physics/PhAT/Properties:description%
-	[/PARAM]
-	[PARAM:path]
-		[RELATIVE:Engine/Physics/PhAT/Properties]
-	[/PARAM]
-[/OBJECT]
-[/VAR]
-
+parent:Engine/Physics/PhAT
+order:1
+Version: 4.13
+topic-image:Engine/Physics/physics_topic.png
+tags:Physics
+SkillLevel: Beginner
 
 [VAR:spacer]![](%ROOT%/spacer.png)(w:150 convert:false)[/VAR]
 
@@ -151,14 +116,27 @@ When no Physics Body or Physics Constraint is selected (regardless of mode), the
 [REGION:raw]
 | Property | %spacer% | Description |
 | --- | :---: | :--- |
-| **Physics Blend** || %doxygen:UPhATSimOptions::PhysicsBlend% |
-| **Update Joints from Animation** || |
-| **Poke Pause Time** || %doxygen:UPhATSimOptions::PokePauseTime% |
-| **Poke Blend Time** || %doxygen:UPhATSimOptions::PokeBlendTime% |
+| **Physics Blend** || Enables you to manually control the physics/animation blend while simulating in the **Physics Asset Tool**. |
+| **Update Joints from Animation** || Let's you manually control the physics/animation. |
+| **Physics Update Mode** || [INCLUDE:#UpdateJointsfromAnimationEnum] |
+| **Poke Pause Time** || Time between poking a ragdoll and starting to blend back. |
+| **Poke Blend Time** || Time taken to blend from physics to animation. |
 [/REGION]
 
 [/VAR]
 
+<!--
+[EXCERPT:UpdateJointsfromAnimationEnum]
+
+Determines whether simulation of root body updates component transform.
+
+| Value | %spacer2% | Description |
+| --- | :---: | :--- |
+| **Simulation Updates Component Transfrom** || Simulation updates the component's transform. |
+| **Component Trasform is Kinematic** || Simulation doesn't update the component's transform. Instead it is driven by animation, blueprint, etc... |
+[/EXCERPT]
+-->
+
 [VAR:p12]
 
 ![](PhATDef2.png)
@@ -166,8 +144,8 @@ When no Physics Body or Physics Constraint is selected (regardless of mode), the
 [REGION:raw]
 | Property | %spacer% | Description |
 | --- | :---: | :--- |
-|**Floor Gap** || %doxygen:UPhATSimOptions::FloorGap% |
-|**Grav Scale** || %doxygen:UPhATSimOptions::GravScale% |
+|**Floor Gap** || The gap between the floor and physics objects in the **Physics Asset Tool**. Changing this requires a restart of the **Physics Asset Tool**. |
+|**Grav Scale** || Scale factor for the gravity used in the simulation while in the **Physics Asset Tool**. |
 |**Max FPS** || Max FPS for simulation in PhAT. This is helpful for targeting the same FPS as your game. -1 means this is disabled.  |
 |**Time Dilation** || This will dilate the time by scale.  |
 [/REGION]
@@ -181,11 +159,11 @@ When no Physics Body or Physics Constraint is selected (regardless of mode), the
 [REGION:raw]
 | Property | %spacer% | Description |
 | --- | :---: | :--- |
-|**Handle Linear Damping** || %doxygen:UPhATSimOptions::HandleLinearDamping% |
-|**Handle Linear Stiffness** || %doxygen:UPhATSimOptions::HandleLinearStiffness% |
-|**Handle Angular Damping** || %doxygen:UPhATSimOptions::HandleAngularDamping% |
-|**Handle Angular Stiffness** || %doxygen:UPhATSimOptions::HandleAngularStiffness% |
-|**Interpolation Speed** || %doxygen:UPhATSimOptions::InterpolationSpeed%|
+|**Handle Linear Damping** || Linear damping of mouse spring forces. |
+|**Handle Linear Stiffness** || Linear stiffness of mouse spring forces. |
+|**Handle Angular Damping** || Angular damping of mouse spring forces. |
+|**Handle Angular Stiffness** || Angular stiffness of mouse spring forces. |
+|**Interpolation Speed** || How quickly we interpolate the physics target transform for mouse spring forces. |
 [/REGION]
 
 [/VAR]
@@ -197,7 +175,7 @@ When no Physics Body or Physics Constraint is selected (regardless of mode), the
 [REGION:raw]
 | Property | %spacer% | Description |
 | --- | :---: | :--- |
-| **Poke Strength** || %doxygen:UPhATSimOptions::PokeStrength% |
+| **Poke Strength** || Strength of the impulse used when poking with the left mouse button. |
 [/REGION]
 
 [/VAR]
@@ -209,8 +187,8 @@ When no Physics Body or Physics Constraint is selected (regardless of mode), the
 [REGION:raw]
 | Property | %spacer% | Description |
 | --- | :---: | :--- |
-| **Show Constraints as Points** || %doxygen:UPhATSimOptions::bShowConstraintsAsPoints% |
-| **Show Names in Hierarchy** || %doxygen:UPhATSimOptions::bShowNamesInHierarchy% |
+| **Show Constraints as Points** || Whether or not to draw constraints as points. |
+| **Show Names in Hierarchy** || Whether or not to draw bone names in the viewport. |
 [/REGION]
 
 [/VAR]
@@ -222,7 +200,7 @@ When no Physics Body or Physics Constraint is selected (regardless of mode), the
 [REGION:raw]
 | Property | %spacer% | Description |
 | --- | :---: | :--- |
-| **Constraint Draw Size** || %doxygen:UPhATSimOptions::ConstraintDrawSize% |
+| **Constraint Draw Size** || Controls how large constraints are drawn in the **Physics Asset Tool**. |
 [/REGION]
 
 [/VAR]
@@ -247,25 +225,25 @@ When the Physics Asset Tool is in Body Mode and a Physics Body is selected, the
 
 [OBJECT:modulesectionexpanded]
 	[PARAMLITERAL:id]
-		Body1
+		Phys1
 	[/PARAMLITERAL]
 	[PARAM:heading]
-		Body Setup
+		Physics
 	[/PARAM]
 	[PARAM:content]
-		%p21%
+		%p22%
 	[/PARAM]
 [/OBJECT]
 
 [OBJECT:modulesectionexpanded]
 	[PARAMLITERAL:id]
-		Phys1
+		Body1
 	[/PARAMLITERAL]
 	[PARAM:heading]
-		Physics
+		Body Setup
 	[/PARAM]
 	[PARAM:content]
-		%p22%
+		%p21%
 	[/PARAM]
 [/OBJECT]
 
@@ -281,6 +259,18 @@ When the Physics Asset Tool is in Body Mode and a Physics Body is selected, the
 	[/PARAM]
 [/OBJECT]
 
+[OBJECT:modulesectionexpanded]
+	[PARAMLITERAL:id]
+		PhysAnim
+	[/PARAMLITERAL]
+	[PARAM:heading]
+		Physical Animation
+	[/PARAM]
+	[PARAM:content]
+		%p24%
+	[/PARAM]
+[/OBJECT]
+
 [/VAR]
 
 [VAR:p21]
@@ -290,8 +280,9 @@ When the Physics Asset Tool is in Body Mode and a Physics Body is selected, the
 [REGION:raw]
 | Property | %spacer% | Description |
 | --- | :---: | :--- |
-| **Bone Name** || %doxygen:UBodySetup::BoneName% |
-| **Consider for Bounds** || %doxygen:UBodySetup::bConsiderForBounds% |
+| **Primitives** || The basic shapes that compose this Bone's collision. |
+| **Bone Name** || Used in the PhysicsAsset case, associates this Body with a Bone in a Skeletal Mesh. |
+| **Consider for Bounds** || Should this bone be used in the calculation for the bounding volume of the Skeletal Mesh. |
 [/REGION]
 
 [/VAR]
@@ -304,12 +295,12 @@ When the Physics Asset Tool is in Body Mode and a Physics Body is selected, the
 | Property | %spacer% | Description |
 | --- | :---: | :--- |
 | **Mass in KG** || Mass of the body in KG. |
-| **Linear Damping** || %doxygen:FBodyInstance::LinearDamping% |
-| **Angular Damping** || %doxygen:FBodyInstance::AngularDamping% |
-| **Enable Gravity** || %doxygen:FBodyInstance::bEnableGravity% |
-| **Physics Type** || %doxygen:UBodySetup::PhysicsType% |
+| **Linear Damping** || "Drag" force added to reduce linear movement. |
+| **Angular Damping** || "Drag" force added to reduce angular movement.  |
+| **Enable Gravity** || If object should have the force of gravity applied.  |
+| **Physics Type** || If this is set to **Unfixed** it will use physics. If this is set to **Fixed** it will use kinematic. **Default** will inherit from **OwnerComponent**'s behavior. |
 | **Double Sided Geometry** || If true, the physics triangle mesh will use double sided faces when doing scene queries. Thsi is useful for planes and single sided meshes that need traces to work on both sides. |
-| **Simple Collision Physical Material** || %doxygen:UBodySetup::PhysMaterial% |
+| **Simple Collision Physical Material** || The **Physical Material** to use for simple collision on this body. |
 [/REGION]
 
 ### Advanced
@@ -319,20 +310,19 @@ When the Physics Asset Tool is in Body Mode and a Physics Body is selected, the
 [REGION:raw]
 | Property | %spacer% | Description |
 | --- | :---: | :--- |
-| **Start Awake** || %doxygen:FBodyInstance::bStartAwake% |
-| **Center of Mass Offset** || %doxygen:FBodyInstance::COMNudge% |
-| **Mass Scale** || %doxygen:FBodyInstance::MassScale% |
+| **Start Awake** || If the physics body should start awake (and moving), or asleep (stationary until interacted with). |
+| **Center of Mass Offset** || User specified offset for the center of mass of this object, from the calculated location. |
+| **Mass Scale** || Per-instance scaling of mass. |
 | **Max Angular Velocity** || The maximum angular velocity for this instance. |
-| **Use Async Scene** || %doxygen:FBodyInstance::bUseAsyncScene% |
-| **Generate Wake Events** ||  |
+| **Generate Wake Events** || Should **wake**/**sleep** events fire when this object is woken up or put to sleep by the physics simulation. |
 | **Max Depenetration Velocity** || The maximum velocity used to depenetrate this object.  |
-| **Walkable Slope Override** || %doxygen:UBodySetup::WalkableSlopeOverride% |
+| **Walkable Slope Override** || Custom walkable slope setting for this body. See the [Walkable Slope](Engine/Physics/WalkableSlope) documentation for usage information. |
 | **Walkable Slope Behavior** || Behavior of this surface (whether we affect the walkable slope). Determines how the Walkable Slope Angle operates on the existing walkable angle, whether it increases or decreases it. |
-| **Walkable Slope Angle** || %doxygen:FWalkableSlopeOverride::WalkableSlopeAngle% |
-| **Custom Sleep Threshold Multiplier** || |
-| **Position Solver Iteration Count** || %doxygen:FBodyInstance::PositionSolverIterationCount% |
-| **Velocity Solver Iteration Count** || %doxygen:FBodyInstance::VelocitySolverIterationCount% |
-| **Sleep Family** || %doxygen:FBodyInstance::SleepFamily% |
+| **Walkable Slope Angle** || Override walkable slope, applying the rules of the Walkable Slope Behavior. |
+| **Custom Sleep Threshold Multiplier** || If the **SleepFamily** is set to custom, multiply the natural sleep threshold by this amount. A higher number will cause the body to sleep sooner. |
+| **Position Solver Iteration Count** || This physics body's solver iteration count for position. Increasing this will be more CPU intensive, but better stabilized. |
+| **Velocity Solver Iteration Count** || This physics body's solver iteration count for velocity. Increasing this will be more CPU intensive, but better stabilized. |
+| **Sleep Family** || The set of values used in considering when put this body to sleep. |
 [/REGION]
 
 [/VAR]
@@ -344,9 +334,9 @@ When the Physics Asset Tool is in Body Mode and a Physics Body is selected, the
 [REGION:raw]
 | Property | %spacer% | Description |
 | --- | :---: | :--- |
-| **Simulation Generates Hit Events** || %doxygen:FBodyInstance::bNotifyRigidBodyCollision% |
-| **Phys Material Override** || %doxygen:FBodyInstance::PhysMaterialOverride% |
-| **Collision Response** || %doxygen:UBodySetup::CollisionReponse% |
+| **Simulation Generates Hit Events** || Should 'Hit' events fire when this object collides during physics simulation. |
+| **Phys Material Override** || Allows you to override the PhysicalMaterial to use for simple collision on this body. |
+| **Collision Response** || See the [Collision Response Reference](Engine/Physics/Collision/Reference) documentation for more information. |
 
 [/REGION]
 ### Advanced
@@ -356,11 +346,18 @@ When the Physics Asset Tool is in Body Mode and a Physics Body is selected, the
 [REGION:raw]
 | Property | %spacer% | Description |
 | --- | :---: | :--- |
-| **Use CCD** || %doxygen:FBodyInstance::bUseCCD% |
+| **Use CCD** || If true Continuous Collision Detection (CCD) will be used for this component. |
 [/REGION]
 
 [/VAR]
 
+[VAR:p24]
+
+![](PhATBody4.png)
+
+Please see the [](Engine/Physics/PhAT/HowTo) for more information on using **Physical Animations**.
+
+[/VAR]
 
 ---
 
@@ -461,10 +458,9 @@ When the Physics Asset Tool is in Body Mode and a Physics Body is selected, the
 [REGION:raw]
 | Property | %spacer% | Description |
 | --- | :---: | :--- |
-| **Joint Name** || %doxygen:FConstraintInstance::JointName% |
-| **Constraint Bone 1** || %doxygen:FConstraintInstance::ConstraintBone1% |
-| **Constraint Bone 2** || %doxygen:FConstraintInstance::ConstraintBone2% |
-| **Disable Collision** || Disable collision between bodies joined by this constraint. |
+| **Joint Name** || The name of the bone (**Physics Body**) this **Physics Constraint** is associated with. |
+| **Constraint Bone 1** || Name of the first bone (**Physics Body**) this **Physics Constraint** is connecting. This will be the "child" bone in a **Physics Asset**. |
+| **Constraint Bone 2** || Name of the second bone (**Physics Body**) this **Physics Constraint** is connecting. This will be the "parent" bone in a **Physics Asset**. |
 [/REGION]
 [/VAR]
 
diff --git a/Engine/Documentation/Source/Engine/Physics/Physics.INT.udn b/Engine/Documentation/Source/Engine/Physics/Physics.INT.udn
--- a/Engine/Documentation/Source/Engine/Physics/Physics.INT.udn
+++ b/Engine/Documentation/Source/Engine/Physics/Physics.INT.udn
@@ -4,7 +4,7 @@ Crumbs: %ROOT%, Engine
 Description:Subsystem used to calculate collision and simulate physical actors.
 Version: 4.9
 parent:Engine
-order:6
+order:8
 
 [VAR:Topic]
 [OBJECT:Topic]
diff --git a/Engine/Documentation/Source/Engine/Physics/PhysicsBodies/PhysicsBodies.INT.udn b/Engine/Documentation/Source/Engine/Physics/PhysicsBodies/PhysicsBodies.INT.udn
--- a/Engine/Documentation/Source/Engine/Physics/PhysicsBodies/PhysicsBodies.INT.udn
+++ b/Engine/Documentation/Source/Engine/Physics/PhysicsBodies/PhysicsBodies.INT.udn
@@ -64,17 +64,5 @@ The following links have information on Physics Bodies.
 [/OBJECT]
 
 
-[OBJECT:TopicButtonList]
-	[PARAM:icon]
-		![](%ROOT%/start_icon.png)(convert:false)
-	[/PARAM]
-	[PARAM:title]
-		Related
-	[/PARAM]
-	[PARAM:links]
-
-	[/PARAM]
-[/OBJECT]
-
 
 
diff --git a/Engine/Documentation/Source/Engine/Plugins/Ansel/AnselPhotography.INT.udn b/Engine/Documentation/Source/Engine/Plugins/Ansel/AnselPhotography.INT.udn
new file mode 100644
--- /dev/null
+++ b/Engine/Documentation/Source/Engine/Plugins/Ansel/AnselPhotography.INT.udn
@@ -0,0 +1,26 @@
+Availability:Public
+TItle:Nvidia Ansel Photography Plugin
+Description:Integration of Nvidia's Ansel in-game photography tool that enables capturing and sharing enhanced images from games
+Crumbs:
+Parent:Engine/Plugins
+Order:1
+type:landing
+tags:Plugins
+tags:Ansel
+version:4.14
+topic-image:Ansel_topic.png
+reference-image:ansel_example.gif
+
+[REGION:asyncgif]
+![](ansel_example.gif)
+[/REGION]
+
+[INCLUDE:Engine/Plugins/Ansel/Overview#intro]
+
+## Starting Out
+
+[DIR(output:"topic" parent:"Engine/Plugins/Ansel" type:"overview")]
+
+## Reference
+
+[DIR(output:"topic" parent:"Engine/Plugins/Ansel" type:"reference")]
diff --git a/Engine/Documentation/Source/Engine/Plugins/Ansel/Overview/AnselOverview.INT.udn b/Engine/Documentation/Source/Engine/Plugins/Ansel/Overview/AnselOverview.INT.udn
new file mode 100644
--- /dev/null
+++ b/Engine/Documentation/Source/Engine/Plugins/Ansel/Overview/AnselOverview.INT.udn
@@ -0,0 +1,119 @@
+Availability:Public
+TItle:Ansel Plugin Overview
+Description:Overview of the Unreal Engine 4 integration of Nvidia's Ansel in-game photography tool that enables capturing and sharing enhanced images from games
+Crumbs:
+Parent:Engine/Plugins/Ansel
+Order:1
+type:overview
+tags:Plugins
+tags:Ansel
+version:4.14
+topic-image:Engine/Plugins/Ansel/Ansel_topic.png
+
+[EXCERPT:intro]
+Nvidia's Ansel technology makes it simple for PC gamers to capture and share enhanced images from their games. Integrating the Ansel UE4 plugin into your game is quick and easy.
+
+An end-user overview of Ansel is available on [Nvidia's Ansel Site](http://www.geforce.com/hardware/technology/ansel).
+
+The Ansel features available in Unreal Engine 4 are:
+
+* **Free Camera** - Position the camera for ideal composition before taking a picture
+* **Super Resolution** - Render at higher than game resolution for ultimate quality 
+* **Post Processing** - Apply filters to alter the mood 
+* **360 Capture** - Capture 360 panoramic images that can be viewed on PC or with a VR headset 
+
+## Hardware and Software Requirements
+
+| Item | Requirement |
+| ---- | ----------- |
+| Unreal Engine Version | 4.14 or later |
+| Operating System | Microsoft Windows |
+| GPU | NVIDIA GTX 780 or better [REGION:none]A full compatibility list is available on the [Ansel Supported GPUs page](http://www.geforce.com/hardware/technology/ansel/supported-gpus)[/REGION] |
+| Display Driver | Nvidia Geforce 372.90 or later |
+[/EXCERPT:intro]
+
+## Enabling Ansel Photography for Your Game
+
+**To enable the Ansel Photography feature in your game:**
+
+1.	In the **Settings** dropdown of the Level Editor toolbar, choose Use **Plugins** to open the **Plugins Browser**.
+
+	![](Settings_Menu.png)
+
+1.	Select the **Photography** category and locate the **Nvidia Ansel Photography Plugin**. Toggle the _Enabled_ option.
+
+	![](Plugin_Enabled.png)
+
+1.	Restart Unreal Editor.
+
+Your game should now have a working Ansel Photography Mode which the player can invoke. For instruction on how to test the Ansel integration in your game, see the [](Engine/Plugins/Ansel/Testing) page.
+
+## Capture Types
+
+Ansel offers the following capture types:
+
+| Capture Type | Example | Description |
+| --- | --- | --- |
+| **Screenshot** | [REGION:lightbox][![](capture_normal.png)(w:320)](capture_normal.png)[/REGION] | A standard screenshot using the current viewport resolution |
+| **Superresolution** | [REGION:lightbox][![](capture_super.jpg)(w:320)](capture_super.jpg)[/REGION] | A high resolution screenshot using a configurable multiple of the current viewport resolution |
+| **360** | [REGION:lightbox][![](capture_360.jpg)(w:320)](capture_360.jpg)[/REGION] | A 360 degree view screenshot using a configurable resolution |
+| **Stereo** | [REGION:lightbox][![](capture_stereo.jpg)(w:320)](capture_stereo.jpg)[/REGION] | A stereo screenshot using the current viewport resolution for each viewpoint |
+| **360 Stereo** | [REGION:lightbox][![](capture_360stereo.jpg)(w:320)](capture_360stereo.jpg)[/REGION] | A 360 degree view stereo screenshot using a configurable resolution  |
+
+To select a capture type, use the **Capture type** slider:
+
+![](capture_select.png)(w:960)
+
+
+[REGION:note]
+Not all filters are valid with multipart Capture types, e.g. 360 and Highres. You may therefore see filters, or aspects of a filter, removed in the final picture.
+
+Similarly, not all UE4 postprocessing effects are valid with multipart Capture types.  Standard UE4 postprocessing effects which are problematic in this case are automatically removed in the final picture.
+[/REGION]
+
+## Filters
+
+A number of filters can be applied to adjust the look of the scene.
+
+[REGION:fullwidth]
+[OBJECT:SliderSequence]
+	[PARAMLITERAL:id]
+	filter_slider
+	[/PARAMLITERAL]
+	[PARAMLITERAL:max]
+	8
+	[/PARAMLITERAL]
+	[PARAMLITERAL:ext]
+	jpg
+	[/PARAMLITERAL]
+	[PARAM:caption]
+	Drag to see the available filters in action
+	[/PARAM]
+	[PARAM:image_list]
+	![](filter_slider_1.png)
+	![](filter_slider_2.png)
+	![](filter_slider_3.png)
+	![](filter_slider_4.png)
+	![](filter_slider_5.png)
+	![](filter_slider_6.png)
+	![](filter_slider_7.png)
+	![](filter_slider_8.png)
+	[/PARAM]
+[/OBJECT]
+[/REGION]
+
+To select a filter, use the **Filter** slider:
+
+![](filter_select.png)(w:960)
+
+[REGION:tip]
+Some filters, such as the **Custom** filter, have additional settings that can be used to adjust the filter even further.
+[/REGION]
+
+## Customizing the Player Experience
+
+The default Ansel plugin configuration is meant to provide a sane out-of-the-box experience for typical projects but for the best player experience it should be configured for your game's needs and potentially integrated with your game logic. See the [](Engine/Plugins/Ansel/Reference/ConsoleVariables) and [](Engine/Plugins/Ansel/Reference/BlueprintAPI) for information on the available options and logic hooks.
+
+## Publishing your game with Ansel photography
+
+Studios, please contact the [NVIDIA Ansel team](mailto:[REDACTED_EMAIL]) to arrange for your title's Ansel support to be universally whitelisted for all players.
diff --git a/Engine/Documentation/Source/Engine/Plugins/Ansel/Reference/BlueprintAPI/AnselBlueprintAPI.INT.udn b/Engine/Documentation/Source/Engine/Plugins/Ansel/Reference/BlueprintAPI/AnselBlueprintAPI.INT.udn
new file mode 100644
--- /dev/null
+++ b/Engine/Documentation/Source/Engine/Plugins/Ansel/Reference/BlueprintAPI/AnselBlueprintAPI.INT.udn
@@ -0,0 +1,160 @@
+Availability:Public
+TItle:Ansel Blueprint Node Reference
+Description:Reference of the configuration variables available for the Nvidia Ansel Photography Plugin
+Crumbs:
+Parent:Engine/Plugins/Ansel
+Order:1
+type:reference
+tags:Plugins
+tags:Ansel
+version:4.14
+topic-image:BlueprintAPI_Topic.png
+
+You can customize the end user's experience with the Ansel Photography plugin in order to improve that experience, i.e. by disabling certain effects that don't work well in captures, or to control the experience, such as limiting the movement of the camera. Several Blueprint events and functions are available to implement logic to customize the player's experience.
+
+## Events
+
+### On Photography Session Start
+
+Event triggered upon entering Photography mode (before pausing, if `r.Photography.AutoPause` is 1).
+
+[REGION:graph]
+    [OBJECT:BlueprintNode]
+        [PARAMLITERAL:type]
+            event
+        [/PARAMLITERAL]
+        [PARAMLITERAL:title]
+            Event On Photography Session Start
+        [/PARAMLITERAL]
+        [PARAM:inputs]
+        [/PARAM]
+        [PARAM:outputs]
+            [OBJECT:BlueprintPin]
+                [PARAMLITERAL:type]
+                    exec
+                [/PARAMLITERAL]
+                [PARAMLITERAL:id]
+                    Exec
+                [/PARAMLITERAL]
+                [PARAMLITERAL:title]
+                [/PARAMLITERAL]
+            [/OBJECT]
+        [/PARAM]
+    [/OBJECT]
+[/REGION]
+
+### On Photography Session End
+
+Event triggered upon leaving Photography mode (after unpausing, if `r.Photography.AutoPause` is 1).
+
+[REGION:graph]
+    [OBJECT:BlueprintNode]
+        [PARAMLITERAL:type]
+            event
+        [/PARAMLITERAL]
+        [PARAMLITERAL:title]
+            Event On Photography Session End
+        [/PARAMLITERAL]
+        [PARAM:inputs]
+        [/PARAM]
+        [PARAM:outputs]
+            [OBJECT:BlueprintPin]
+                [PARAMLITERAL:type]
+                    exec
+                [/PARAMLITERAL]
+                [PARAMLITERAL:id]
+                    Exec
+                [/PARAMLITERAL]
+                [PARAMLITERAL:title]
+                [/PARAMLITERAL]
+            [/OBJECT]
+        [/PARAM]
+    [/OBJECT]
+[/REGION]
+
+### On Photography Multi Part Capture Start
+
+Event triggered upon the start of a multi-part photograph capture (i.e. a stereoscopic or 360-degree shot).  This is an ideal time to turn off rendering effects that tile badly (UI, subtitles, vignette, very aggressive bloom, etc; most of these are automatically disabled when `r.Photography.AutoPostprocess` is 1).
+
+[REGION:graph]
+    [OBJECT:BlueprintNode]
+        [PARAMLITERAL:type]
+            event
+        [/PARAMLITERAL]
+        [PARAMLITERAL:title]
+            Event On Photography Multi Part Capture Start
+        [/PARAMLITERAL]
+        [PARAM:inputs]
+        [/PARAM]
+        [PARAM:outputs]
+            [OBJECT:BlueprintPin]
+                [PARAMLITERAL:type]
+                    exec
+                [/PARAMLITERAL]
+                [PARAMLITERAL:id]
+                    Exec
+                [/PARAMLITERAL]
+                [PARAMLITERAL:title]
+                [/PARAMLITERAL]
+            [/OBJECT]
+        [/PARAM]
+    [/OBJECT]
+[/REGION]
+
+### On Photography Multi Part Capture End
+
+Event triggered upon the end of a multi-part photograph capture, when manual free-roaming photographic camera control is about to be returned to the user.  Here you may re-enable whatever was turned off within OnPhotographyMultiPartCaptureStart.
+
+[REGION:graph]
+    [OBJECT:BlueprintNode]
+        [PARAMLITERAL:type]
+            event
+        [/PARAMLITERAL]
+        [PARAMLITERAL:title]
+            Event On Photography Multi Part Capture End
+        [/PARAMLITERAL]
+        [PARAM:inputs]
+        [/PARAM]
+        [PARAM:outputs]
+            [OBJECT:BlueprintPin]
+                [PARAMLITERAL:type]
+                    exec
+                [/PARAMLITERAL]
+                [PARAMLITERAL:id]
+                    Exec
+                [/PARAMLITERAL]
+                [PARAMLITERAL:title]
+                [/PARAMLITERAL]
+            [/OBJECT]
+        [/PARAM]
+    [/OBJECT]
+[/REGION]
+
+## Overridable Functions
+
+**To override the functions listed below:**
+
+1. Create a new **Blueprint Class** with **Player Camera Manager** as the parent and open it in the **Blueprint Editor**.
+
+1. In the **My Blueprint** panel, place the mouse cursor over the **Functions** heading and click **Override**.
+
+	![](BIE_Override.png)
+
+1. Select the function you wish to override from the list of overridable functions. The function appears in the **Functions** list and the graph for it is opened so you can edit it.
+
+### Photography Camera Modify
+
+This allows a PlayerCameraManager subclass to constrain or otherwise modify the camera during free-camera photography. For example, a blueprint may wish to limit the distance from the camera's original point, or forbid the camera from passing through walls.
+This function gives you the original, previous, and proposed camera positions and takes your modified camera position as a result.
+
+[REGION:raw]
+![](PhotoCameraModify.png)
+[/REGION]
+
+#### Example
+
+Below, is an example implementation of **Photography Camera Modify** which limits the photography camera to within 200 Unreal Units of its original position:
+
+[REGION:lightbox]
+[![](camera_constraint.png)(w:480)](camera_constraint.png)
+[/REGION]
\ No newline at end of file
