Copying //UE4/Dev-Documentation to Samples-Main (//UE4/Samples-Main)

diff --git a/Engine/Documentation/Source/Engine/Plugins/Ansel/Reference/ConsoleVariables/AnselConsoleVars.INT.udn b/Engine/Documentation/Source/Engine/Plugins/Ansel/Reference/ConsoleVariables/AnselConsoleVars.INT.udn
new file mode 100644
--- /dev/null
+++ b/Engine/Documentation/Source/Engine/Plugins/Ansel/Reference/ConsoleVariables/AnselConsoleVars.INT.udn
@@ -0,0 +1,24 @@
+Availability:Public
+TItle:Ansel Console Variable Reference
+Description:Reference of the configuration variables available for the Nvidia Ansel Photography Plugin
+Crumbs:
+Parent:Engine/Plugins/Ansel
+Order:3
+type:reference
+tags:Plugins
+tags:Ansel
+version:4.14
+topic-image:Console_Topic.png
+
+The following console variables are available for configuring the behavior of the Nvidia Ansel Photography Plugin and can be set manually from the console, through your .ini files, or dynamically by a Blueprint. Most of them should be configured during development instead of at runtime, though `r.Photography.Allow` is useful to set dynamically.
+
+| Variable | Default | Description |
+| --- | --- | --- |
+| `r.Photography.Available` | 1 | (Read-only) If 1, the photography system is potentially available to the user.  Otherwise, a functioning back-end is not available. This is set to 1 when the plugin is enabled. |
+| `r.Photography.Allow` | 1 | If 1, allows a photography session to potentially be started by the player.  Set this dynamically to permit or forbid photography per-level, per-cutscene, etc. |
+| `r.Photography.EnableMultipart` | 1 | If 1, allow the photography system to take high-resolution shots that need to be rendered in tiles which are later stitched together. |
+| `r.Photography.SettleFrames` | 10 | The number of frames to let the rendering 'settle' before taking a photo.  Useful to allow temporal AA/smoothing to work well; if not using any temporal effects, can be lowered for faster capture. |
+| `r.Photography.TranslationSpeed` | 100.0 | Normal speed (in Unreal units per second) at which to move the roaming photography camera. |
+| `r.Photography.AutoPostprocess` | 1 | If 1, the photography system will attempt to automatically disable HUD, subtitles, and some standard postprocessing effects during photography sessions/captures which are known to give poor photography results.  Set to 0 to manage all postprocessing tweaks manually from the PlayerCameraManager Blueprint callbacks. [REGION:note]Blueprint callbacks will be called regardless of AutoPostprocess value.[/REGION] |
+| `r.Photography.AutoPause` | 1 | If 1, the photography system will attempt to ensure that the level is paused while in photography mode.  Set to 0 to manage pausing and unpausing manually from the PlayerCameraManager Blueprint callbacks.  [REGION:note]Blueprint callbacks will be called regardless of AutoPause value.[/REGION] |
+| `r.Photography.PersistEffects` | 0 | If 1, custom postprocessing effects enabled in photography mode are permitted to persist in the game after a photography session has ended.  Changes to this value might not be applied until the next photography session starts. |
diff --git a/Engine/Documentation/Source/Engine/Plugins/Ansel/Testing/AnselTesting.INT.udn b/Engine/Documentation/Source/Engine/Plugins/Ansel/Testing/AnselTesting.INT.udn
new file mode 100644
--- /dev/null
+++ b/Engine/Documentation/Source/Engine/Plugins/Ansel/Testing/AnselTesting.INT.udn
@@ -0,0 +1,63 @@
+Availability:Public
+TItle:Ansel Capture Testing
+Description:Overview of the Unreal Engine 4 integration of Nvidia's Ansel in-game photography tool that enables capturing and sharing enhanced images from games
+Crumbs:
+Parent:Engine/Plugins/Ansel
+Order:1
+type:overview
+tags:Plugins
+tags:Ansel
+version:4.14
+topic-image:Ansel_Testing_Topic.png
+
+**To test the Nvidia Ansel Photography Plugin with your game:**
+
+1. Run the NvCameraConfiguration.exe application and ensure that **NvCamera Status** is set to _Enable_, which is the default.
+
+		C:\Program Files\NVIDIA Corporation\Ansel\Tools\NvCameraConfiguration.exe
+
+1. Run the NvCameraEnable.exe application passing it `whitelisting-everything` to allow your game to use Ansel. This persists across reboots and may be reset by typing `whitelisting-default` in place of `whitelisting-everything`.
+
+		"C:\Program Files\NVIDIA Corporation\Ansel\Tools\NvCameraEnable.exe" whitelisting-everything
+
+1. In Unreal Editor, select **Play > Standalone Game** in the Level Editor toolbar to run as a Standalone Game. Ansel will not activate when playing in the Editor.
+
+1. Press **Alt + F2** to activate Ansel photography in the game. When a Session has been successfully started, the game world pauses and the Ansel UI will appear over the game view.
+
+	![](Engine/Plugins/Ansel/Testing/unlit.png)(w:640)
+
+	[REGION:note]
+	Ansel will not activate while VR or Split-screen are active in the game
+	[/REGION]
+
+	[REGION:note]
+	In Editor builds, **F2** is the keybinding that switches to Unlit rendering mode.  To restore the more photo-friendly Lit mode, press **F3**.
+	[/REGION]
+
+1. Select a capture type using the **Capture type** slider. 
+
+1. Press the **Snap** button. 
+
+	[REGION:note]
+	Some pictures may take significant time to produce, especially highres shots of large dimensions. If the game uses streaming the streaming performance may be affected when shots involving many parts are being stitched together.
+	[/REGION]
+
+
+## In-Game Controls
+
+| Keyboard / Mouse | Gamepad | Description |
+| --- | --- | --- |
+| [REGION:tablesection]Camera Movement[/REGION] |||
+| W | Left Analog Stick Up | Moves the Camera forward |
+| A | Left Joystick Left | Moves the Camera left |
+| S | Left Joystick Down | Moves the Camera backward |
+| D | Left Joystick Right | Moves the Camera right |
+| X | Left Trigger | Moves the Camera up |
+| Y | Right Trigger | Moves the Camera down |
+| Shift (Hold) | Right Analog Stick (Press) | Accelerates Camera movement speed |
+| [REGION:tablesection]Camera Movement[/REGION] |||
+| Left-Click + Move Mouse | Right Analog Stick | Rotates the Camera (yaw and pitch) |
+
+
+
+
diff --git a/Engine/Documentation/Source/Engine/Plugins/Plugins.INT.udn b/Engine/Documentation/Source/Engine/Plugins/Plugins.INT.udn
new file mode 100644
--- /dev/null
+++ b/Engine/Documentation/Source/Engine/Plugins/Plugins.INT.udn
@@ -0,0 +1,11 @@
+Availability:Public
+TItle:Built-In Plugins
+Description:Developer guides to plugins provided with Unreal Engine 4
+Crumbs:
+Parent:Engine
+Order:100
+type:landing
+version:4.14
+topic-image:Plugins_Topic.png
+
+[DIR(output:"topic" type:"landing" parent:"Engine/Plugins" org:"flat")]
diff --git a/Engine/Documentation/Source/Engine/Rendering/LightingAndShadows/CapsuleShadows/BestPractices/CapsuleShadowsBestPractices.INT.udn b/Engine/Documentation/Source/Engine/Rendering/LightingAndShadows/CapsuleShadows/BestPractices/CapsuleShadowsBestPractices.INT.udn
--- a/Engine/Documentation/Source/Engine/Rendering/LightingAndShadows/CapsuleShadows/BestPractices/CapsuleShadowsBestPractices.INT.udn
+++ b/Engine/Documentation/Source/Engine/Rendering/LightingAndShadows/CapsuleShadows/BestPractices/CapsuleShadowsBestPractices.INT.udn
@@ -65,20 +65,22 @@ make sure that it looks right. You may also want to have several different light
 [![](FootPlacement1.png)(w:600)](FootPlacement1.png)
 [/REGION]
 
-In this setup the capsules for the feet are close to relative size of the feet and have some overlap with the capsules for the shins. The shadows are soft and not overly dark.
+In this setup, the capsules for the feet are close to relative size of the feet and have some overlap with the capsules for the shins. The shadows are soft and not overly dark.
 
 [REGION:lightbox]
 [![](FootPlacement2.png)(w:600)](FootPlacement2.png)
 [/REGION]
  
-In this setup the capsules are too large for the feet and extend beyond the back of the heel leaving the soft shadows to appear darker than they should.
+In this setup, the capsules are too large for the feet and extend beyond the back of the heel, leaving the soft shadows to appear darker than they should.
 
 
 ### Limit Capsule Bodies
 
-Capsule Shadows can be considered an advanced form of "blob" shadows. You can simply use a single capsule or as many as you need to get better shadowing accuracy. With that in mind, start by 
-determining the needs of your character's movement for your game and whether they will ever be near a surface by using a cover system or falling down. If they will never be near surfaces consider
-not having capsule bodies for the arms of your character, since the performance costs is determined by the number of capsule bodies.
+Capsule Shadows can be considered an advanced form of "blob" shadows. You can simply use a single capsule, or as many as you need to get better shadowing accuracy. With that in mind, start by 
+determining the needs of your character's movement for your game and consider what parts of the character will get close to surfaces. For example, a character who takes cover or lies prone in a third-person shooter 
+will likely need capsule shadows on the head, arms and torso, whereas a character who stays upright and doesn't press against walls might not need those capsules, and performance can be 
+improved by not having capsule bodies for the arms of your character, since the performance costs is determined by the number of capsule bodies.
+
 
 [REGION:lightbox]
 [![](BlobShadow.png)(w:400)](BlobShadow.png)
@@ -99,8 +101,7 @@ if accuracy is not a concern.
 	[/PARAM]
 [/OBJECT]
 
-In this comparison, the arms have been removed from the character's representation. In the directly lit character areas not near surfaces you can see that it only makes a slight visual difference. 
-
+In this comparison, the arms have been removed from the character's representation. The visual effect of the arm capsules is slight in the directly-lit areas, and almost imperceptible in the darker regions. 
 
 ### Indirect Lighting Cache (Volume Light Samples)
 
@@ -110,7 +111,6 @@ soft shadow. Sometimes this requires tweaking the placement of the light samples
 To adjust the **Volume Light Sample Placement** you must enable the **World Settings** panel by going to **Settings** > **World Settings** in the main toolbar. Then in the World Settings panel under the 
 **Lightmass** section you can change the **Volume Light Sample Placement**. Use lower values to make the placement of samples more dense or higher values to make them more spread out.
 
-
 You can visualize these samples in your level by going to the **Viewport** > **Show** > **Visualize** > **Volume Lighting Samples**. 
 
 [OBJECT:ComparisonSlider]
@@ -127,9 +127,4 @@ this may be different for your game.
 
 [REGION:warning]
 Values too low will cause a "popping" as the capsule moves between the samples instead of providing a soft smooth transition between them, since the placement of the samples is so close together.
-[/REGION]
-
-
-
-
-
+[/REGION]
\ No newline at end of file
diff --git a/Engine/Documentation/Source/Engine/Rendering/LightingAndShadows/CapsuleShadows/CapsuleShadowsLanding.INT.udn b/Engine/Documentation/Source/Engine/Rendering/LightingAndShadows/CapsuleShadows/CapsuleShadowsLanding.INT.udn
--- a/Engine/Documentation/Source/Engine/Rendering/LightingAndShadows/CapsuleShadows/CapsuleShadowsLanding.INT.udn
+++ b/Engine/Documentation/Source/Engine/Rendering/LightingAndShadows/CapsuleShadows/CapsuleShadowsLanding.INT.udn
@@ -4,8 +4,8 @@ description:Landing page for Capsule Shadows in Unreal Engine 4
 crumbs:
 version:4.13
 
-**Capsule Shadows** have been developed to support soft shadow casting for Skeletal Meshes by creating a capsule representation of your character using a custom Physics Asset. 
-The soft shadows serve to ground the character in indirectly areas and give soft shadows in directly lit areas. This shadowing technique can be used with Unreal Engine 4.11 or later. 
+**Capsule Shadows** support soft shadow casting for Skeletal Meshes by creating a capsule representation of your character using a custom Physics Asset. 
+The soft shadows serve to ground the character in indirectly-lit areas and give soft shadows in directly-lit areas. This shadowing technique is available in Unreal Engine 4.11 and higher. 
 
 
 [OBJECT:TopicButtonList]
diff --git a/Engine/Documentation/Source/Engine/Rendering/LightingAndShadows/CapsuleShadows/Overview/CapsuleShadowsOverview.INT.udn b/Engine/Documentation/Source/Engine/Rendering/LightingAndShadows/CapsuleShadows/Overview/CapsuleShadowsOverview.INT.udn
--- a/Engine/Documentation/Source/Engine/Rendering/LightingAndShadows/CapsuleShadows/Overview/CapsuleShadowsOverview.INT.udn
+++ b/Engine/Documentation/Source/Engine/Rendering/LightingAndShadows/CapsuleShadows/Overview/CapsuleShadowsOverview.INT.udn
@@ -4,81 +4,74 @@ Crumbs: %ROOT%, Engine,
 Description: Capsule Shadows Overview
 Version: 4.13
 
-[VAR:Topic]
-[OBJECT:Topic]
-	[PARAM:image]
-		![%Engine/Rendering/LightingAndShadows/CapsuleShadows/CapsuleShadowsOverview:title%]()
-	[/PARAM]
-	[PARAM:icon]
-		![](%ROOT%/light_icon.png)(convert:false)
-	[/PARAM]
-	[PARAM:title]
-		%Engine/Rendering/LightingAndShadows/CapsuleShadows/CapsuleShadowsOverview:title%
-	[/PARAM]
-	[PARAM:description]
-		%Engine/Rendering/LightingAndShadows/CapsuleShadows/CapsuleShadowsOverview:description%
-	[/PARAM]
-	[PARAM:path]
-		[RELATIVE:Engine/Rendering/LightingAndShadows/CapsuleShadows/CapsuleShadowsOverview]
-	[/PARAM]
-[/OBJECT]
-[/VAR]
-[VAR:TopicCompact]
-[OBJECT:TopicCompact]
-	[PARAM:image]
-		![%Engine/Rendering/LightingAndShadows/CapsuleShadows/CapsuleShadowsOverview:title%]()
-	[/PARAM]
-	[PARAM:icon]
-		![](%ROOT%/light_icon.png)(convert:false)
-	[/PARAM]
-	[PARAM:title]
-		%Engine/Rendering/LightingAndShadows/CapsuleShadows/CapsuleShadowsOverview:title%
-	[/PARAM]
-	[PARAM:description]
-		%Engine/Rendering/LightingAndShadows/CapsuleShadows/CapsuleShadowsOverview:description%
-	[/PARAM]
-	[PARAM:path]
-		[RELATIVE:Engine/Rendering/LightingAndShadows/CapsuleShadows/CapsuleShadowsOverview]
-	[/PARAM]
-[/OBJECT]
-[/VAR]
 
+[REGION:center]
+![](image_9.png)(w:900)
+[/REGION]
 
 [TOC(start:2 end:3)]
 
 
-**Capsule Shadows** is a technique that is used to generate soft shadows in indirectly and directly lit areas by using a capsule representation of your characters with a Physics Asset.
-This solves the problem of making a character feel grounded in indirectly lit areas and by also giving soft shadow representations for those which are directly lit.
+**Capsule Shadows** are used to generate soft shadows in lit areas by using a capsule representation of a character with a Physics Asset.
+This feature can make a character feel grounded in both direct and indirect lighting by producing consistent soft shadows.
 
 
 ## How do they work?
 
-Capsule Shadows start by using a unique Physics Asset that can be assigned to any Skeletal Mesh as a **Shadow Physics Asset**. This asset is unique to each Skeletal Mesh as an 
-accurate representation of your character. Once this information is known the [Indirect Lighting Samples](Engine/Rendering/LightingAndShadows/IndirectLightingCache/) that are 
-generated after a light build can be used to cast a cone towards the direction of the light intensity. As the character nears this point the cone is used to check if any capsules 
-are occluding the cone. If they are, then all the capsule shadows can be accumulated to get this end result of soft shadows.
+Capsule Shadows start by using a unique Physics Asset that can be assigned to any Skeletal Mesh as a **Shadow Physics Asset**. This asset represents an approximation of the Skeletal Mesh 
+with respect to the shadows it should cast. Once this information is known, the [Indirect Lighting Samples](Engine/Rendering/LightingAndShadows/IndirectLightingCache/) that are 
+generated after a light build can be used to cast a cone towards the direction of the light source. As the character nears this point, the cone is checked against the capsules to 
+see if there are any overlaps. If there are, then all the capsule shadows can be accumulated to build soft shadows for the character.
 
-<!---
-The Capsule Shadow technique allows a Physics Asset to be assigned to any Skeletal Mesh as a **Shadow Physics Asset**. This asset is unique to each Skeletal Mesh so that it can represent your character 
-accurately for soft shadowing. Character shadowing has generally been handled by using standard shadow mapping techniques, which can only have sharp shadows with slight blurring. Capsule Shadows make 
-it possible to have these soft shadow representations for your character so that it can seamlessly blend with the environment.
---->
-
-[REGION:lightbox]
-[![](CapsuleShadowRepresentation.png)(w:500)](CapsuleShadowRepresentation.png)
+[REGION:imagetable]
+| ![](CS_Without.png)(w:370) | ![](CS_With.png)(w:370) |
+| ---- | ---- |
+| Without Capsule Shadows | With Capsule Shadows |
 [/REGION]
-[REGION:caption]
-_Click the image for full size view_
+
+[REGION:left]
+[INCLUDE:Engine/Rendering/LightingAndShadows/CapsuleShadows\Overview#CapsuleShadowSlider]
 [/REGION]
 
-The scene above breaks down how Capsule Shadows are created for your character. 
+<!---
+[EXCERPT:CapsuleShadowSlider]
+[OBJECT:SliderSequence]
+	[PARAMLITERAL:max]
+	5
+	[/PARAMLITERAL]
+	[PARAMLITERAL:ext]
+	jpg
+	[/PARAMLITERAL]
+	[PARAMLITERAL:id]
+	FS
+	[/PARAMLITERAL]
+	[PARAM:caption]
+	**Drag the slider to see what it takes to render the Capsule Shadows for a character.**
+	[/PARAM]
+	[PARAM:image_list]
+	![](CS_1.png)(w:750)
+	![](CS_2.png)(w:750)
+	![](CS_3.png)(w:750)
+	![](CS_4.png)(w:750)
+	![](CS_6.png)(w:750)
+	[/PARAM]
+[/OBJECT]
+[/EXCERPT:CapsuleShadowSlider]
+-->
 
-1. Once lighting has been built for the scene the Indirect Lighting Cache samples that are generated for your level's surfaces is used to gather the directionality and intensity of the light. 
-2. The cone is placed near a light sample and from this can determine the direction the shadow will cast and how dark the shadow should be. If the light sample is brighter than ones in more covered areas the cone will be larger so that the shadows will be softer with less directionality.
-3. The cone checks if the capsules are occluding any part of it. If the capsules are occluding, the cone direction and intensity of the lighting samples is used to draw the soft shadows.
-4. Finally, we have the soft shadow that is created from the capsules that are occluding the cone.
+The Physics asset that is created for Capsule Shadows and assigned to the Skeletal Mesh is used to represent the character in these indirectly-lit areas. Note that this is a mock-up representation of the capsules for the character. 
 
+1. Once lighting has been built for the scene, the character will have no shadowing in areas that have indirect lighting. 
+2. The Physics asset that is created for Capsule Shadows and assigned to the Skeletal Mesh is used to represent the character in these indirectly-lit areas. Note that this is a mock-up representation of the capsules for the character.
+3. With lighting built, the precomputed [volume lighting samples](Engine/Rendering/LightingAndShadows/IndirectLightingCache/) are placed so that Skeletal Actors can blend in with the indirectly-lit environment by referencing the 
+lighting information stored in the placed samples.
+4. These light samples are used to draw a cone towards the capsule representaiton of the character based on the directional intensity of the light. The cone can be used to mask out shadowed areas to get a soft shadow representation for 
+the character.  
+5. The intesnity and directionality of the capsule shadow representation can be computed to ground the character in the environment giving us the final render with soft shadowing for movable actors in indirectly-lit areas. 
 
+[REGION:note]
+Capsule Shadows can also work in Dynamic lights, when **Capsule Direct Shadow** is enabled, to produce a soft shadow in directly-lit areas. See the [Capsule Direct Shadowing](Engine\Rendering\LightingAndShadows\CapsuleShadows\Overview#CapsulesforDirectShadowing) section of this page.
+[/REGION]
 
 #### Capsule Representation of Character
 
@@ -97,7 +90,7 @@ _Capsule representation of a Skeletal Mesh character_
 [/REGION]
 
 [REGION:note]
-The shadows on the ground are not Capsule Shadows. 
+The shadows on the ground are not representative Capsule Shadows. 
 [/REGION]
 
 
@@ -108,7 +101,7 @@ You can enable Capsule Shadows for each Skeletal Mesh in your level by enabling
 ![](EnableCapsuleShadows.png)(w:350)
 
 
-#### Capsules for Direct Shadowing
+#### Capsules Direct Shadowing
 
 **Capsule Direct Shadows** enables soft shadowing for your character when in areas that are directly lit by any **Movable** or **Stationary** light source. 
 
@@ -132,7 +125,7 @@ For directly lit Capsule Shadows you can adjust the **Light Source Angle** of th
 	[/PARAM]
 [/OBJECT]
 
-#### Capsules for Indirect Shadowing
+#### Capsules Indirect Shadowing
 
 **Capsule Indirect Shadows** enables soft shadowing for your character in areas that are indirectly lit by any **Static** or **Stationary** light source. Capsule shadowing is created using the 
 [Indirect Lighting Cache](Engine/Rendering/LightingAndShadows/IndirectLightingCache/) that is generated during a light build. 
@@ -145,13 +138,13 @@ _Click the image for full size_
 [/REGION]
 
 
-## Directonality and Intensity of Indirect Shadows
+## Directionality and Intensity of Indirect Shadows
 
 The directionality and intensity of the capsule shadowing is determined by the light samples generated after a light build. The brighter intensity of the lighting sample the less intense and softer the shadow will be since light is coming from many directions. 
 
 ### Closed Area Shadowing
 
-When you have a opening for light to pour through the intensity of the light will be brighter at the opening and dimmer at the back part of this area. 
+When you have a opening for light to pour through, the intensity of the light will be brighter at the opening and dimmer at the back part of this area. 
 
 [REGION:lightbox]
 [![](CS_CaveFront.png)(w:500)](CS_CaveFront.png)
@@ -169,13 +162,12 @@ In the front of this area the character's shadow is much softer since there is m
 _Click the image for full size_
 [/REGION]
 
-However, at the back of this area the shadow is more intense since it is more dimly lit and has more directionality. The intensity of the light is coming from a single opening towards the front which means there will be
-less light traveling to the back of this area.
-
+However, at the back of this area the shadow is more intense, since it is more dimly lit and has more directionality. The majority of the light is coming from a single opening towards the front which means there will be
+less light traveling to the back of this area, and the character is lit more strongly from one direction.
 
 ### Open Area Shadowing
 
-In open areas that receive indirect shadowing, the Capsule Shadows will appear much softer since the light can come from all around.
+In open areas that receive indirect shadowing, Capsule Shadows will appear much softer, since the light can come from all around.
 
 [REGION:lightbox]
 [![](CS_openarea.png)(w:500)](CS_openarea.png)
@@ -189,9 +181,10 @@ _Click the image for full size_
 
 ## Performance
 
-The cost of Capsule Shadows in your game is cheaper and more efficient than standard shadow mapping techniques because of tile defered culling which cull the capsule influences to only the tiles that could possibly 
-receive shadows from that capsule. This means that most of the areas on the screen knew they did not have to do any work. The tiles that are possibly receiving shadows will have to compute for the given pixel 
-the number of capsules that directly affect it. The performance benefit of Capsule Shadows comes by only testing these tiles that are affected and by having fewer capsules to compute.
+Capsule Shadows are more efficient than standard shadow mapping techniques because of tile deferred culling which cull the capsule influences to only the tiles that could possibly 
+receive shadows from that capsule. This means that most of the areas on the screen were unaffected and did not require additional rendering work. The tiles that are potentially 
+receiving shadows will have to compute, for each pixel, the number of capsules that directly affect that pixel. Capsule Shadows perform better when there are fewer affected tiles on 
+the screen, as well as when there are fewer capsules.
 
 [OBJECT:ComparisonSlider]
 	[PARAM:before]
@@ -208,8 +201,8 @@ it needs to consider all the triangles in the shadow map, which means you can on
 
 ## Limitations
 
-* Requires DX-11 due to using tiled deferred implementation using compute shaders 
-* Self-shadowing artifacts for arbitrary mesh shapes 
-* Can only use Spheres and Sphylls for capsule representation
-* Only support for Skeletal Meshes, but could be made to work with static meshes
+* Requires DX 11 due to tiled deferred implementation using compute shaders.
+* Self-shadowing artifacts can exist for arbitrary mesh shapes.
+* Can only use Spheres and Sphylls for capsule representation.
+* Only support for Skeletal Meshes, but could be made to work with static meshes.
 * When capsule shadows get so soft that they become ambient occlusion, there's an artifact in the shadowing that causes a hard line. 
\ No newline at end of file
diff --git a/Engine/Documentation/Source/Engine/Rendering/LightingAndShadows/CapsuleShadows/Reference/CapsuleShadowsReference.INT.udn b/Engine/Documentation/Source/Engine/Rendering/LightingAndShadows/CapsuleShadows/Reference/CapsuleShadowsReference.INT.udn
--- a/Engine/Documentation/Source/Engine/Rendering/LightingAndShadows/CapsuleShadows/Reference/CapsuleShadowsReference.INT.udn
+++ b/Engine/Documentation/Source/Engine/Rendering/LightingAndShadows/CapsuleShadows/Reference/CapsuleShadowsReference.INT.udn
@@ -4,76 +4,33 @@ Crumbs: %ROOT%, Engine,
 Description: Capsule Shadows Reference
 Version: 4.13
 
-[VAR:Topic]
-[OBJECT:Topic]
-	[PARAM:image]
-		![%Engine/Rendering/LightingAndShadows/CapsuleShadows/Reference:title%]()
-	[/PARAM]
-	[PARAM:icon]
-		![](%ROOT%/light_icon.png)(convert:false)
-	[/PARAM]
-	[PARAM:title]
-		%Engine/Rendering/LightingAndShadows/CapsuleShadows/Reference:title%
-	[/PARAM]
-	[PARAM:description]
-		%Engine/Rendering/LightingAndShadows/CapsuleShadows/Reference:description%
-	[/PARAM]
-	[PARAM:path]
-		[RELATIVE:Engine/Rendering/LightingAndShadows/CapsuleShadows/Reference]
-	[/PARAM]
-[/OBJECT]
-[/VAR]
-[VAR:TopicCompact]
-[OBJECT:TopicCompact]
-	[PARAM:image]
-		![%Engine/Rendering/LightingAndShadows/CapsuleShadows/Reference:title%]()
-	[/PARAM]
-	[PARAM:icon]
-		![](%ROOT%/light_icon.png)(convert:false)
-	[/PARAM]
-	[PARAM:title]
-		%Engine/Rendering/LightingAndShadows/CapsuleShadows/Reference:title%
-	[/PARAM]
-	[PARAM:description]
-		%Engine/Rendering/LightingAndShadows/CapsuleShadows/Reference:description%
-	[/PARAM]
-	[PARAM:path]
-		[RELATIVE:Engine/Rendering/LightingAndShadows/CapsuleShadows/Reference]
-	[/PARAM]
-[/OBJECT]
-[/VAR]
 
-
-
-[TOC(start:2 end:2)]
-
-
-
-The following console commands allow you to adjust the quality of various features for Capsule Shadows.
+Use the following console commands to adjust the quality and various settings for how Capsule Shadows are represented in your game. 
 
 <!---
-| Console Variable | Description |
-| ---------------- | ----------- |
-| [r.CapsuleIndirectConeAngle](#Parameters) |	Light source angle used when the indirect shadow direction is derived from precomputed indirect lighting (no stationary skylight present). Default Value: 0.392699 |
-| [r.CapsuleIndirectShadowMinVisibility](#capsuleindirectshadowminvisibility) |	Minimum visibility caused by capsule indirect shadows. This is used to keep the indirect shadows from going fully black. Default Value: 0.1 |
-| r.CapsuleMaxDirectOcclusionDistance |	Maximum cast distance for direct shadows from capsules. �This has a big impact on performance. Default Value: 400 |
-| r.CapsuleMaxIndirectOcclusionDistance |	Maximum cast distance for indirect shadows from capsules. This has a big impact on performance. Default Value: 200 |
-| r.CapsuleShadowFadeAngleFromVertical |	Angle from vertical up to start fading out the indirect shadow, to avoid self shadowing artifacts. Default Value: 1.047198 |
-| r.CapsuleShadows |	Whether to allow capsule shadowing on skinned components with bCastCapsuleDirectShadow or bCastCapsuleIndirectShadow enabled. |
-| r.CapsuleShadowsFullResolution |	Whether to compute capsule shadows at full resolution. Disabled by Default. |
-| r.CapsuleSkyAngleScale |	Scales the light source angle derived from the precomputed unoccluded sky vector (stationary skylight present) Default Value: 0.6 |
+| Console Variable | Description | Range | 
+| ---------------- | ----------- | ----- |
+| [r.CapsuleIndirectConeAngle](#Parameters) |	Light source angle used when the indirect shadow direction is derived from precomputed indirect lighting (no stationary skylight present). Default Value: 0.392699 | 0.1 - 1.0 |
+| [r.CapsuleIndirectShadowMinVisibility](#capsuleindirectshadowminvisibility) |	Minimum visibility caused by capsule indirect shadows. This is used to keep the indirect shadows from going fully black. Default Value: 0.1 | 0.1 - 0.9 |
+| r.CapsuleMaxDirectOcclusionDistance |	Maximum cast distance for direct shadows from capsules. �This has a big impact on performance. Default Value: 400 | 200-400 |
+| r.CapsuleMaxIndirectOcclusionDistance |	Maximum cast distance for indirect shadows from capsules. This has a big impact on performance. Default Value: 200 | 200-400 |
+| r.CapsuleShadowFadeAngleFromVertical |	Angle from vertical up to start fading out the indirect shadow, to avoid self shadowing artifacts. Default Value: 1.047198 | 1.0 - 2.0 |
+| r.CapsuleShadows |	Whether to allow capsule shadowing on skinned components with bCastCapsuleDirectShadow or bCastCapsuleIndirectShadow enabled. |  |
+| r.CapsuleShadowsFullResolution |	Whether to compute capsule shadows at full resolution. Disabled by Default. | - |
+| r.CapsuleSkyAngleScale |	Scales the light source angle derived from the precomputed unoccluded sky vector (stationary skylight present) Default Value: 0.6 | 0.0 - 1.0 |
 --->
 
 
 ### Capsule Indirect Cone Angle
 
-The **Capsule Indirect Cone Angle** command allows you to specify the light source angle when using precomputed indirect lighting. If a Skylight is used in the scene its mobility must be set to Static in order to use this command. 
+This command allows you to specify the light source angle for capsule's shadows when in indirectly-lit areas only and will not change the light source angle for any Directional Light in the scene.
+The light source can have its mobility set to **Static** or **Stationary**, but if a Skylight is used in the scene its mobility must be set to **Static** in order for this command to work. 
 
-Type the following command followed by a value to adjust the light source angle.
+Enter the following command followed by a value between 0.1 and 1.0 to adjust the light source angle that is used for the capsule's shadow representation. This command will only affect capsule shadows.
 
 	r.CapsuleIndirectConeAngle
 	
-		
+			
 [OBJECT:ComparisonSlider]
 	[PARAM:before]
 	![Indirect Cone Angle = 0.392699 (Default)](IndirectConeAngle_Default.png)(w:600)
@@ -84,15 +41,16 @@ Type the following command followed by a value to adjust the light source angle.
 [/OBJECT]
 
 
-### Capsule Indirect Shadow Min Visibility
+### Capsule Indirect Shadow Minimum Visibility
 
-The **Capsule Indirect Shadow Minimum Visibility** command allows you to specify a value for the capsule's indirect shadows. This is used to keep the indirect
-shadows from going fully black. Lower values will soften the shadows intensity, whereas higher values will darken the shadows.
+This command allows you to adjust how visible the capsule's shadows are when in indirectly-lit areas. This is used to keep the indirect shadows from going fully black.
+When using lower values the shadows will soften their intensity, whereas higher values will darken the shadows. 
 
-Type the following command followed by a value to adjust the indirect shadow's visibility.
+Enter the following command followed by a value between 0.1 and 0.9 to adjust the minimum visibility for the indirect shadows.
 
 	r.CapsuleIndirectShadowMinVisibility
 	
+
 [OBJECT:ComparisonSlider]
 	[PARAM:before]
 	![Indirect Shadow Minimum Visibility = 0.1 (Default)](DirectIndirectShadowsDefault.png)(w:600)
@@ -103,18 +61,17 @@ Type the following command followed by a value to adjust the indirect shadow's v
 [/OBJECT]
 
 
-
 ### Capsule Occlusion Distance
 
-The **Capsule Occlusion Distance** command allows you to specify the maximum distance from the capsule that it can cast shadows. Values can be set for Capsule Shadows that are in directly and indirectly lit areas.
+The commands for Direct Occlusion Distance and Indirect Occlusion Distance allows you to specify the maximum distance that the capsules can cast shadows in indirectly-lit and directly-lit areas.
 
 [Region:warning]
-Adjusting this value will have a significant impact on performance.
+By adjusting these values you can significantly impact performance.
 [/REGION]
 
-#### Direct Occlusion Distance
+##### Direct Occlusion Distance
 
-Type the following command followed by a value for distance that direct shadows should cast from the capsules.
+Enter the following command with a value for the distance, in centimeters, that the capsules should cast directly-lit shadows. 
 
 	r.CapsuleMaxDirectOcclusionDistance 
 
@@ -127,9 +84,9 @@ Type the following command followed by a value for distance that direct shadows
 	[/PARAM]
 [/OBJECT]
 
-#### Indirect Occlusion Distance
+##### Indirect Occlusion Distance
 
-Type the following command followed by a value for distance that indirect shadows should cast from the capsules.
+Enter the following command with a value for the distance, in centimeters, that the capsules should cast indirectly-lit shadows.
 
 	r.CapsuleMaxIndirectOcclusionDistance 
 		
@@ -142,23 +99,22 @@ Type the following command followed by a value for distance that indirect shadow
 	[/PARAM]
 [/OBJECT]
 
-
-
 ### Capsule Shadow Fade Angle From Vertical
 
-The **Capsule Shadow Fade Angle from Vertical** command allows you to set a value to start fading out indirect shadows. This can help avoid self shadowing artifacts.
+This command is helpful to avoid any self-shadowing artifacts that can happen in indirectly-lit areas where vertically aligned capsules can produce darker shadowing. Adjusting 
+this value will fade the vertical influence angle.
 
-Type the following command followed by a value for fading the indirect shadows. Higher values will fade out more quickly.
+Enter the following command followed by a value between 1.0 and 2.0 to adjust the fade for any self-shadowing artifacts that happen. Higher values will fade 
 
 	r.CapsuleShadowFadeAngleFromVertical 
 
 		
 [OBJECT:ComparisonSlider]
 	[PARAM:before]
-	![Shadow Fade Angle from Vertical = 1.047198 (Default)](FadeAngleVertical_Default.png)(w:600)
+	![Shadow Fade Angle from Vertical = 1.047198 (Default)](CSFadeVerticalDefault.png)(w:600)
 	[/PARAM]
 	[PARAM:after]
-	![Shadow Fade Angle from Vertical = 1.9](FadeAngleVertical_1_9.png)(w:600)
+	![Shadow Fade Angle from Vertical = 1.9](CSFadeVertical1_9.png)(w:600)
 	[/PARAM]
 [/OBJECT]
 
@@ -189,13 +145,11 @@ Type the following to enable full resolution of the Capsule Shadows.
 
 ### Capsule Sky Angle
 
-The **Capsule Sky Angle** commands allow you to specify minimum light source angle for the sky and the scale. The commands require you to have a 
-Stationary Skylight present in your scene and an unoccluded sky to see the results. 
-
+The commands for the Minimum Sky Angle and the Sky Angle Scale allow you to adjust the light source angle when you have a Stationary Skylight in your scene. 
 
-#### Mininmum Sky Angle
+##### Mininmum Sky Angle
 
-Type the following to set the minimum value for the Stationary Skylight's light source angle.
+Enter the following command with a degree value between 0 and 90 to set the Stationary Skylight's minimimum light source angle. 
 
 	r.CapsuleMinSkyAngle
 
@@ -210,12 +164,11 @@ Type the following to set the minimum value for the Stationary Skylight's light
 [/OBJECT]
 
 
-#### Sky Angle Scale
+##### Sky Angle Scale
 
-Type the following to set the scale of the light source angle derived from the Stationary Skylight.
+Enter the follwing command with a value between 0 and 1.0 to scale the light source angle derived from the Stationary Skylight in your scene.            
 
 	r.CapsuleShadowsSkyAngleScale
-
 		
 [OBJECT:ComparisonSlider]
 	[PARAM:before]
diff --git a/Engine/Documentation/Source/Engine/Rendering/LightingAndShadows/ContactShadows/ContactShadows.INT.udn b/Engine/Documentation/Source/Engine/Rendering/LightingAndShadows/ContactShadows/ContactShadows.INT.udn
new file mode 100644
--- /dev/null
+++ b/Engine/Documentation/Source/Engine/Rendering/LightingAndShadows/ContactShadows/ContactShadows.INT.udn
@@ -0,0 +1,114 @@
+Availability: Docs
+Title: Contact Shadows
+Crumbs: 
+Description: This page provides an overview of Contact Shadows in UE4.
+Version: 4.14
+Parent: Engine/Rendering/LightingAndShadows
+Order: 6
+Tags: Rendering
+Type: Reference
+reference-image: ContactShadow.gif
+Related: Engine/Rendering/LightingAndShadows/LightTypes/Point
+
+[VAR:Params]
+[PARAM:image]
+	![%Engine/Rendering/LightingAndShadows/ContactShadows:title%](ContactShadow_Topic.png)
+[/PARAM]
+[PARAM:icon]
+	![](%ROOT%/tutorial_icon.png)(convert:false)
+[/PARAM]
+[PARAM:title]
+	%Engine/Rendering/LightingAndShadows/ContactShadows:title%
+[/PARAM]
+[PARAM:path]
+	[RELATIVE:Engine/Rendering/LightingAndShadows/ContactShadows]
+[/PARAM]
+[/VAR]
+
+[TOC(start:1 end:4)]
+
+![](ContactShadow_Banner.png)
+
+When creating scenery and characters for your application, there are times when you'll want to add visual depth to 
+the rendering. Adding **Contact Shadows** is a great way to improve the visual depth and fidelity of your scene 
+because they provide a more accurate approximation of shadowing, allowing you to add a contoured shadow that 
+might not be achieved with other shadowing algorithms.
+
+##Adding Detail to your Characters
+
+One use case for enabling Contact Shadows for a given [Point Light](Engine/Rendering/LightingAndShadows/LightTypes/Point) 
+is to render an additional Level of Detail (LOD) for your characters. The following images should give you a good idea of 
+what Contact Shadows can add to a character. As you drag the comparison slider from right to left, you'll notice that the 
+character located within the radius of a Point Light with enabled Contact Shadows exhibits additional details that aren't 
+present when Contact Shadows are turned off.
+
+[OBJECT:ComparisonSlider]
+	[PARAM:before]
+	![Contact Shadow Off](ContactShadow_Off.png)(w:900)
+	[/PARAM]
+	[PARAM:after]
+	![Light's Contact Shadow Length = 0.1](ContactShadow_On.png)(w:900)
+	[/PARAM]
+[/OBJECT]
+
+When you turn on Contact Shadows, you're directing the renderer to execute a per-pixel screen space algorithm on a 
+per-light basis. This means that the Contact Shadow algorithm is executing a light computation pass, 
+performing scene depth buffer ray marching so as to determine whether the queried pixel is to be occluded from the 
+Point Light source that has Contact Shadows enabled.
+
+##Scenes with Contact Shadows
+
+Another use case for enabling Contact Shadows is to eliminate the need for computing parallax occlusion mapping shadows 
+within a material's pixel shader that only supports one light. The following image provides a comparison of parallax 
+occlusion mapping materials with and without Contact Shadows enabled.
+
+![](ContactShadow_Comparison.png)(w:900)
+
+[REGION:note]
+It's important to note that your parallax occlusion mapping materials should output the pixel depth offset.
+[/REGION]
+
+The following is an example of Contact Shadows interacting with a parallax occlusion mapping material.
+
+[REGION:asyncgif]
+![image alt text](ContactShadow.gif)(w:900)
+[/REGION]
+
+[REGION:caption]
+Contact shadow length = 0.1 on both lights.
+[/REGION]
+
+##Enabling Contact Shadows
+
+By default, Contact Shadows are disabled on Point Lights, whereby the Point Light's Contact Shadow length is initialized to 0. 
+To enable Contact Shadows, perform the following steps:
+
+1.	First, drag a **Point Light** into your scene.
+
+	![](CastShadows_Step1.png)(w:1007)
+
+1.	Now, expand the **Light** section of the **Details** panel for your newly created Point Light Component by clicking on the expansion arrow.
+
+	[REGION:lightbox]
+	[![](CastShadows_Step2.png)(w:360 h:640)](CastShadows_Step2.png)
+	[/REGION]
+
+	[REGION:caption]
+	Click for full image.
+	[/REGION]
+
+1.	To enable contact shadows, set the **Contact Shadow Length** to a value greater than 0.
+
+	[REGION:lightbox]
+	[![](CastShadows_Step3.png)(w:360 h:640)](CastShadows_Step3.png)
+	[/REGION]
+
+	[REGION:caption]
+	Click for full image.
+	[/REGION]
+
+By setting the Contact Shadow's length to a value greater than 0, you're directing the renderer to trace rays from a pixel's position toward the 
+light source through the scene's depth buffer. Typically, you'll want to set the Contact Shadow's length to a maximum value of 1, where 1 means 
+that the ray traverses the entire screen. Whereas, setting the Contact Shadow's length to a value of 0.5 means that the ray traverses 
+only half of the screen. Note that the number of samples taken in the scene depth buffer remain constant, which means that you'll see more noise 
+(artifacts) as you increase the Contact Shadow's length. Basically, a Contact Shadow length of 0.75 generates more noise than a length of 0.1.
\ No newline at end of file
diff --git a/Engine/Documentation/Source/Engine/Rendering/LightingAndShadows/LightMobility/DynamicLights/DynamicLights.INT.udn b/Engine/Documentation/Source/Engine/Rendering/LightingAndShadows/LightMobility/DynamicLights/DynamicLights.INT.udn
--- a/Engine/Documentation/Source/Engine/Rendering/LightingAndShadows/LightMobility/DynamicLights/DynamicLights.INT.udn
+++ b/Engine/Documentation/Source/Engine/Rendering/LightingAndShadows/LightMobility/DynamicLights/DynamicLights.INT.udn
@@ -47,7 +47,7 @@ Version: 4.9
 ![](Engine/Rendering/LightingAndShadows/lighting_overview_banner_001.png)(convert:false)
 [/REGION]
 
-[TOC(start:2 end:2)]
+[TOC(start:2 end:3)]
 
 
 
@@ -71,22 +71,26 @@ On any light, under the **Transform** category, you will find a pull down proper
 
 When a point or spot light is not moving, we can store off the shadow map for that light and reuse it next frame. This makes shadow casting movable Point and Spot lights much more affordable in games where the environment is often not moving. In the following document we will take a look at how you can use this feature in your Unreal Engine 4 (UE4) projects.
 
-## Enabling Shadow Map Caching 
-Shadow Map Caching should be automatically be enabled in your UE4 project. To check and see if Shadow Map Caching is enabled you can do the following:
+### Shadow Map Caching & Performance
+Shadow Map Caching should be an automatically enabled feature in any UE4 project. To check and see the performance you get from using Shadow Map Caching you can enable and disable Shadow Map Caching by doing the following:
 
-1. Select all of the lights in your project's level that you want to cast dynamic shadows.
+[region:note]
+For following demonstration of Shadow Map Caching the Sun Temple project was used. You can find the Sun Temple map in the Learn tab of the Epic games launcher.
+[/region]
+
+1. Select all of the lights that you want to have cast dynamic shadows in your project's level.
 
 	![](SMC_Select_All_Lights.png)
 
-1. Set the lights **Mobility**  to **Movable** and make sure that **Cast Shadows** is enabled.
+1. Set the lights' **Mobility**  to **Movable** and make sure that **Cast Shadows** is enabled.
 
 	![](SMC_Set_Mob_Shadow_Cast.png)
 
-1. Press the **Backtick (`)** key to open up the Unreal Console and type **Stat Shadowrendering** to see the current cost of dynamic shadows. 
+1. Press the **Backtick (\`)** key to open up the **Console** and type `Stat Shadowrendering` to see the current cost of dynamic shadows. 
 
 	![](SMC_Console_ShadowRendering.png)
 
-1. Press the Backtick (`) key again to open the Console and type **r.Shadow.CacheWholeSceneShadows 0** to disable dynamic shadow caching.
+1. Press the Backtick (\`) key again to open the **Console** and type `r.Shadow.CacheWholeSceneShadows 0` to disable dynamic shadow caching.
 
 	![](SMC_Disable_Shadow_Cache.png)
 	
@@ -94,7 +98,7 @@ Shadow Map Caching should be automatically be enabled in your UE4 project. To ch
 	 Make sure to note the the numbers in the **CallCount** and **InclusiveAug** sections.
 	[/region]
 
-1. Now press the Backtick (`) key one more time to open the Console and type **r.Shadow.CacheWholeSceneShadows 1** to re-enable shadow caching. If you compare the Callout and InclusiveAug numbers you can see the impact this feature has on dynamic shadow performance.
+1. Now press the Backtick key one more time to open the **Console** and type `r.Shadow.CacheWholeSceneShadows 1` to re-enable shadow caching. If you compare the **CallCount** and **InclusiveAug** numbers you can see the impact this feature has on dynamic shadow performance.
 
 	[OBJECT:ComparisonSlider]
 		[PARAM:before]
@@ -105,26 +109,26 @@ Shadow Map Caching should be automatically be enabled in your UE4 project. To ch
 		[/PARAM]
 	[/OBJECT]
 
-## Performance
-The impact that this features has on performance is quite significant. Here are the results of testing this feature using an NVIDIA 970 GTX GPU at a screen resolution of 1920x1200.
+### Performance
+Shadow Map Caching has a significant impact on performance. Here are the results of testing this optmization out using an NVIDIA 970 GTX GPU at a screen resolution of 1920x1200.
 
-* Before enabling this feature three shadow casting point lights without caching took 14.89ms to render Shadow Depths.
+* Before enabling this feature, three shadow casting point lights without caching took 14.89ms to render Shadow Depths.
 
-* When Cached Shadow Maps where enabled the same three shadow casting lights took .9ms to render which is about **16 times faster**!
+* When Cached Shadow Maps were enabled, the same three shadow casting lights took .9ms to render which is about **16 times faster**!
 	[region:note]
-	Note that it still costs 2 ms to render the contributions of the 33 point lights, which can be optimized in other ways but is not affected by this change.
+	Note that it still costs 2 ms to render the contributions of the 33 point lights, which can be optimized in other ways but are not affected by this change.
 	[/region]
 * You can control the maximum amount memory used by the Shadow Map Cache using **r.Shadow.WholeSceneShadowCacheMb**
 
-## Limitations
-While enabling Shadow Map Caching can decrease the cost of using dynamic shadows in your UE4 projects, there are some limitations that could cause rendering artifacts if used with non-supported features. In the following section we will take a look at what the limitations of Shadow Map Casting are and what you can do about them.
+### Limitations
+While Shadow Map Caching can decrease the cost of using dynamic shadows in your UE4 projects, there are some limitations that could cause rendering artifacts if used with non-supported features. In the following section we will take a look at what the limitations of Shadow Map Caching are and what you can do about them.
 
 
 * By default, caching can only happen when an object meets the following requirements:
 
-	* Primitives have mobility their **Mobility** set to **Static** or **Stationary**. 
+	* Primitives have their mobility **Mobility** set to **Static** or **Stationary**. 
 	* Materials used in the level do not use **World Position Offset**.
-	* Light need to be either a **Point** or **Spot** light, have its **Mobility** set to **Movable** and have **Shadow Casting** enabled.
-	* Lights have to remain stationary and can not move at all.
+	* Light need to be either a **Point** or **Spot** light, have its **Mobility** set to **Movable**, and have **Shadow Casting** enabled.
+	* Lights have to remain at one location.
 	* Materials that use animated **Tessellation** or **Pixel Depth Offset** can cause artifacts as their shadow depths are cached.
 
diff --git a/Engine/Documentation/Source/Engine/Rendering/LightingAndShadows/LightTypes/Point/Point.INT.udn b/Engine/Documentation/Source/Engine/Rendering/LightingAndShadows/LightTypes/Point/Point.INT.udn
--- a/Engine/Documentation/Source/Engine/Rendering/LightingAndShadows/LightTypes/Point/Point.INT.udn
+++ b/Engine/Documentation/Source/Engine/Rendering/LightingAndShadows/LightTypes/Point/Point.INT.udn
@@ -2,7 +2,7 @@ Availability: Public
 Title:Point Lights
 Crumbs: %ROOT%, Engine, Engine/Rendering/LightingAndShadows, Engine/Rendering/LightingAndShadows/LightTypes
 Description:The basics of Point Lights.
-Version: 4.9
+Version: 4.14
 
 [VAR:Topic]
 [OBJECT:Topic]
@@ -113,6 +113,7 @@ Property | Description |
 | **Min Roughness** | Min roughness effective for this light, used for softening specular highlights. |
 | **Shadow Bias** | Controls how accurate the shadows from this light are. |
 | **Shadow Filter Sharpen** | How much to sharpen shadow filtering for this light. |
+| **Contact Shadow Length** | The length of screen space to ray trace for sharp contact shadows. A value of zero disables this option. |
 | **Cast Translucent Shadows** | Whether this light is allowed to cast dynamic shadows through translucent objects. |
 | **Affect Dynamic Indirect Lighting** | Whether the light should be injected into the Light Propagation Volume.|
 | **Cast Static Shadows** | Whether this light casts static shadows. |
diff --git a/Engine/Documentation/Source/Engine/Rendering/LightingAndShadows/LightingAndShadows.INT.udn b/Engine/Documentation/Source/Engine/Rendering/LightingAndShadows/LightingAndShadows.INT.udn
--- a/Engine/Documentation/Source/Engine/Rendering/LightingAndShadows/LightingAndShadows.INT.udn
+++ b/Engine/Documentation/Source/Engine/Rendering/LightingAndShadows/LightingAndShadows.INT.udn
@@ -5,6 +5,9 @@ Description:Using Lighting and shadowing geometry, using Global Illumination, an
 Navigation:topic
 Related:Engine/Rendering/ParticleSystems/ParticleLights
 Version: 4.9
+parent:Engine/Rendering
+order:3
+type:landing
 
 [VAR:Topic]
 [OBJECT:Topic]
@@ -64,6 +67,7 @@ Version: 4.9
 %Engine/Rendering/LightingAndShadows/IndirectLightingCache:Topic%
 %Engine/Rendering/LightingAndShadows/LightMobility/StaticLights:Topic%
 %Engine/Rendering/LightingAndShadows/LightMobility/StationaryLights:Topic%
+%Engine/Rendering/LightingAndShadows/PrecomputedLightingScenarios:Topic%
 [/REGION]
 
 
diff --git a/Engine/Documentation/Source/Engine/Rendering/LightingAndShadows/LitTranslucency/LitTranslucency.INT.udn b/Engine/Documentation/Source/Engine/Rendering/LightingAndShadows/LitTranslucency/LitTranslucency.INT.udn
--- a/Engine/Documentation/Source/Engine/Rendering/LightingAndShadows/LitTranslucency/LitTranslucency.INT.udn
+++ b/Engine/Documentation/Source/Engine/Rendering/LightingAndShadows/LitTranslucency/LitTranslucency.INT.udn
@@ -81,3 +81,23 @@ _The left sphere is translucent, the right is opaque, and they are both setup as
 
 * Lit translucent surfaces are missing direct specular.
 * Lit translucent surfaces get all their direct lighting through the translucency volume lighting texture, which causes them to be lower resolution than needed for most surface materials (glass, water).
+
+### Per-Pixel Translucent Lighting
+
+In the deferred renderer, the new forward shading functionality can now be used on translucent surfaces to get specular highlights from multiple lights and image-based reflections from parallax-corrected reflection captures.
+
+[OBJECT:EmbeddedVideo]
+	[PARAMLITERAL:width]
+	640
+	[/PARAMLITERAL]
+	[PARAMLITERAL:height]
+	360
+	[/PARAMLITERAL]
+	[PARAMLITERAL:videoid]
+	r6u908IE4aA
+	[/PARAMLITERAL]
+[/OBJECT]
+
+To enable Per-Pixel Translucent Lighting set the Lighting Mode to **Surface ForwardShading** and then make sure that **Screen Space Reflections** are enabled.
+
+![](PerPixel_TransLighting_Settings.png)
\ No newline at end of file
diff --git a/Engine/Documentation/Source/Engine/Rendering/LightingAndShadows/PrecomputedLightingScenarios/PrecomputedLightingScenarios.INT.udn b/Engine/Documentation/Source/Engine/Rendering/LightingAndShadows/PrecomputedLightingScenarios/PrecomputedLightingScenarios.INT.udn
new file mode 100644
--- /dev/null
+++ b/Engine/Documentation/Source/Engine/Rendering/LightingAndShadows/PrecomputedLightingScenarios/PrecomputedLightingScenarios.INT.udn
@@ -0,0 +1,77 @@
+Availability: Docs
+Title: Precomputed Lighting Scenarios
+Crumbs: %ROOT%, Engine, Engine/Rendering/LightingAndShadows
+Description: This page provides an overview of Precomputed Lighting Scenarios in UE4.
+Version: 4.14
+Parent: Engine/Rendering/LightingAndShadows
+Tags: Rendering
+Type: Reference
+topic-image:PCLS_Header.png
+
+[TOC(start:1 end:4)]
+
+Unreal Engine 4 (UE4) has support for using different **Precomputed Lighting Scenarios** in levels. This allows a single level to store and display multiple lighting setups, giving you the flexibility of dynamic lighting, but at the fixed cost of precomputed lighting. 
+Having the ability to change between different Precomputed Lighting Scenarios is especially important for Virtual Reality (VR) or Architectural Visualization projects that require high quality renderings with the benefits of fast performance measures. 
+As you read through this reference, you'll learn how to use Precomputed lighting in your UE4 projects.
+
+![](PCLS_Comp_Image_00.png)
+
+[REGION:caption]
+In the above image; the Directional Light, Sky Light, and Skybox have been placed in a Lighting Scenario level called `Day Scenario`. The Spotlights for the streetlights have been placed in a different Lighting Scenario level called `Night Scenario`.
+[/REGION]
+
+## Precomputed Lighting Scenarios Limitations
+While Precomputed Lighting scenarios have a lot of advantages, there are some disadvantages and limitations that you'll need to be aware of while using them. 
+In the following section, we'll go over some of these limitations and how you can avoid (or work around) them.
+
+* Only one Lighting Scenario level should be visible at any time in a game.
+* When a Lighting Scenario level is present, Lightmap data from all Sublevels will be placed inside of it so that only the Day Scenario Lightmaps are loaded when it’s daytime. As a result, Lightmaps will no longer be streamed by Sublevels.
+* A Reflection Capture update is forced when making a Lighting Scenario level visible, which can increase the level load time.
+
+##  Using Lighting Scenarios
+
+To use Lighting Scenarios in your UE4 projects, you will need to do the following:
+
+1. First go to **Window** > **Levels** to open up the **Levels Manager**.
+
+    ![](PCLS_Open_Levels.png)
+
+1. With the **Levels Manager** open, right-click on a Sublevel in the **Levels** menu, go to **Lighting Scenario**, and select the **Change to Lighting Scenario** option to make the level a **Lighting Scenario** level.
+
+    ![](PCLS_LS_Level_Option.png)
+
+    [region:note]
+    When a **Lighting Scenario** level is visible, its Lightmaps will be applied to the world.
+    [/region]
+
+1. Make sure the Level Streaming method is set to Blueprint by right-clicking on the Sublevel, going to **Change Streaming Method**, and selecting **Blueprint** (if it's not already selected).
+
+    ![](PCLS_Level_Streaming_Method.png)
+    
+1. Now place any light or [](Engine/Content/Types/StaticMeshes) your project needs into either of the lighting levels and then build each level's lighting as you would with any other level.
+
+    ![](PCLS_Build_Lighting.png)
+
+1. Once the lighting has finished building, open up the **Persistent Level's** Blueprint and add a **Load Stream Level** node and connect it to the **Event Begin Play** node. 
+
+    ![](PCLS_Level_Load_Nodes.png)
+
+1. Connect the **Event Begin Play** node to the **Load Stream Level** node and then input the name of the Level you want to load. Also, make sure to check both **Make Visible After Load** and **Should Block on Load** to make sure you can see the newly loaded level.
+
+    ![](PCLS_Steam_Setup.png)
+
+1. Press the **Play** button to launch the project and when the first level loads, it should now be using the Day level lighting. To use Night level lighting, you can use the same setup, but you'll need to change the level name to your night time level (instead of the day time level).
+
+    [OBJECT:ComparisonSlider]
+        [PARAM:before]
+        ![Day Lighting](PCLS_Day.png)
+        [/PARAM]
+        [PARAM:after]
+        ![Night Lighting](PCLS_Night.png)
+        [/PARAM]
+    [/OBJECT]
+
+
+Although there are some notable limitations, using Precomputed Lighting Scenarios in UE4 can provide your project with many benefits; such as improved performance and the ability to change to baked lighting (to meet your project's needs).  
+
+ 
\ No newline at end of file
diff --git a/Engine/Documentation/Source/Engine/Rendering/LightingAndShadows/RayTracedDistanceFieldShadowing/RayTracedDistanceFieldShadowing.INT.udn b/Engine/Documentation/Source/Engine/Rendering/LightingAndShadows/RayTracedDistanceFieldShadowing/RayTracedDistanceFieldShadowing.INT.udn
--- a/Engine/Documentation/Source/Engine/Rendering/LightingAndShadows/RayTracedDistanceFieldShadowing/RayTracedDistanceFieldShadowing.INT.udn
+++ b/Engine/Documentation/Source/Engine/Rendering/LightingAndShadows/RayTracedDistanceFieldShadowing/RayTracedDistanceFieldShadowing.INT.udn
@@ -51,7 +51,7 @@ Version: 4.9
 
 
 
-Ray Traced Distance Field shadows leverage the properties of a distance field representation of the scene to compute efficient area shadowing from dynamic meshes.  This uses the same data as [**Distance Field Ambient Occlusion**](Engine/Rendering/LightingAndShadows/DistanceFieldAmbientOcclusion), and therefore has many of the same limitations.
+Ray Traced Distance Field shadows leverage the properties of a distance field representation of the scene to compute efficient area shadowing from dynamic meshes.  This uses the same data as [Distance Field Ambient Occlusion](Engine/Rendering/LightingAndShadows/DistanceFieldAmbientOcclusion), and therefore has many of the same limitations.
 
 To calculate shadowing, a ray is traced from the point being shaded through the scene's signed distance fields toward each light.  The closest distance to an occluder is used to approximate a cone trace for about the same cost as the ray trace.  This allows high quality area shadowing from spherical light source shapes.
 
@@ -76,7 +76,7 @@ The editor draws the source shape of a light with yellow lines.
 Ray Traced soft shadows from a point light with a Source Radius of 100.
 [/REGION]
 
-**[_Video of a moving point light using Ray Traced Distance Field shadowing._](https://www.youtube.com/watch?v=-xai0iBffUc&list=UUBobmJyzsJ6Ll7UbfhI4iwQ)**
+[_Video of a moving point light using Ray Traced Distance Field shadowing._](https://www.youtube.com/watch?v=-xai0iBffUc&list=UUBobmJyzsJ6Ll7UbfhI4iwQ)
 
 ## Directional lights
 
@@ -103,7 +103,7 @@ Ray Traced Distance Field shadows behave much more gracefully in the distance, d
 [/OBJECT]
 
 
-**[_Video of changing time of day using Ray Traced Distance Field shadowing._](https://www.youtube.com/watch?v=4249b94KtyA&list=UUBobmJyzsJ6Ll7UbfhI4iwQ)**
+[_Video of changing time of day using Ray Traced Distance Field shadowing._](https://www.youtube.com/watch?v=4249b94KtyA&list=UUBobmJyzsJ6Ll7UbfhI4iwQ)
 
 ## Quality
 
@@ -113,7 +113,7 @@ These shadows are computed at half resolution with a depth-aware upsample.  Temp
 
 ## Limitations
 
-The limitations of this technique are similar to [**Distance Field Ambient Occlusion**](Engine/Rendering/LightingAndShadows/DistanceFieldAmbientOcclusion)'s limitations:
+The limitations of this technique are similar to [Distance Field Ambient Occlusion](Engine/Rendering/LightingAndShadows/DistanceFieldAmbientOcclusion)'s limitations:
 
 * Only Static Mesh components, Instanced Static Mesh components and foliage (once "Affect Distance Field Lighting" has been enabled) generate distance fields and can cast shadows.  However, skeletal meshes can still composite their shadows with 'Cast Inset Shadow'.  Landscape is best handled with Far cascades enabled by "Far Shadow Cascade Count".
 * Only slight non-uniform scaling (squishing) of meshes is supported.  Artifacts in the penumbras will appear with higher amounts of mesh squishing.
diff --git a/Engine/Documentation/Source/Engine/Rendering/LightingAndShadows/ReflectionEnvironment/ReflectionEnvironment.INT.udn b/Engine/Documentation/Source/Engine/Rendering/LightingAndShadows/ReflectionEnvironment/ReflectionEnvironment.INT.udn
--- a/Engine/Documentation/Source/Engine/Rendering/LightingAndShadows/ReflectionEnvironment/ReflectionEnvironment.INT.udn
+++ b/Engine/Documentation/Source/Engine/Rendering/LightingAndShadows/ReflectionEnvironment/ReflectionEnvironment.INT.udn
@@ -2,7 +2,7 @@ Availability:Public
 Title: Reflection Environment
 Crumbs: %ROOT%, Engine, Engine/Rendering/LightingAndShadows
 Description:System for capturing and displaying local glossy reflections.
-Version: 4.12
+Version: 4.14
 
 [TOC(start:2)]
 
@@ -158,6 +158,27 @@ to see results.
 	[/PARAM]
 [/OBJECT]
 
+## Reflection Capture Lightmap Mixing
+
+When you use Reflection Capture Actors, UE4 mixes the indirect Specular from the Reflection Capture with the indirect Diffuse lighting from lightmaps. This helps to reduce leaking since the reflection cubemap was only captured at one point in space, but the lightmaps were computed on all the receiver surfaces and contain local shadowing information.
+
+[OBJECT:ComparisonSlider]
+    [PARAM:before]
+    ![Reduce Lightmap Mixing On Smooth Surfaces Off](ReduceLightmapMixingOnSmoothSurfaces_Off.png)
+    [/PARAM]
+    [PARAM:after]
+    ![Reduce Lightmap Mixing On Smooth Surfaces On.png](ReduceLightmapMixingOnSmoothSurfaces_On.png)
+    [/PARAM]
+[/OBJECT]
+
+While lightmap mixing works well for rough surfaces, this method breaks down on smooth surfaces as reflections from Reflection Capture Actors no longer match reflections from other methods like Screen Space Reflections or Planar Reflections. Because of this, lightmap mixing is no longer applied to very smooth surfaces. A surface with Roughness value of 0.3 will get full lightmap mixing, fading out to no lightmap mixing by Roughness 0.1 and below. This also allows Reflection Captures and Screen Space Reflections to match better and make it harder to spot transitions between the two.
+
+### Lightmap Mixing and Existing Content
+
+By default, lightmap mixing will be enabled which means it will affect existing content. In cases where you had reflections leaking on smooth surfaces, that leaking will be more apparent. To solve this, you can either place additional Reflection Capture Actors around the level to help reduce the leaking. Or you can revert to the old lightmap mixing behavior by going to **Edit** > **Project Settings** >  **Rendering** > **Reflections** and then un-check the **Reduce lightmap mixing on smooth surfaces**.
+
+![](Engine\Rendering\LightingAndShadows\ReflectionEnvironment\Enable_LM_Mixing.png)
+
 ## High Quality Reflections
 
 While the default reflection quality settings strike a good balance between performance and visual quality, there could be instances where you want to achieve even higher quality reflections. The following sections describe the available methods for achieving high quality reflections.
@@ -176,7 +197,7 @@ An important factor in achieving high quality reflections is how accurately the
 	[![](SM_HP_Vertex_Normals_00.png)(w:400 h:350)](SM_HP_Vertex_Normals_00.png)
 	[/REGION]
 
-1. At the bottom of the **LOD0**, there is a section called **Build Settings**. Click on the small triangle next Build Settings to expand the Build Settings options.
+1. At the bottom of the **LOD0**, there is a section called **Build Settings**. Click on the small triangle next to Build Settings to expand the Build Settings options.
 
 	[REGION:lightbox]
 	[![](SM_HP_Vertex_Normals_01.png)(w:400 h:350)](SM_HP_Vertex_Normals_01.png)
diff --git a/Engine/Documentation/Source/Engine/Rendering/Materials/CustomizedUVs/CustomizedUVs.INT.udn b/Engine/Documentation/Source/Engine/Rendering/Materials/CustomizedUVs/CustomizedUVs.INT.udn
--- a/Engine/Documentation/Source/Engine/Rendering/Materials/CustomizedUVs/CustomizedUVs.INT.udn
+++ b/Engine/Documentation/Source/Engine/Rendering/Materials/CustomizedUVs/CustomizedUVs.INT.udn
@@ -3,45 +3,11 @@ Title:Customized UVs
 Crumbs: %ROOT%, Engine, Engine/Rendering/Materials
 Description:Feature that allows running calculations in the vertex shader to increase performance over running them per-pixel.
 Version: 4.9
+parent:Engine/Rendering/Materials
+order:12
+topic-image:CUVs_topic.png
+type:overview
 
-[VAR:Topic]
-[OBJECT:Topic]
-	[PARAM:image]
-		![%Engine/Rendering/Materials/CustomizedUVs:title%](Engine/Rendering/Materials/CustomizedUVs/CUVs_topic.png)
-	[/PARAM]
-	[PARAM:icon]
-		![](%ROOT%/tech_icon.png)(convert:false)
-	[/PARAM]
-	[PARAM:title]
-		%Engine/Rendering/Materials/CustomizedUVs:title%
-	[/PARAM]
-	[PARAM:description]
-		%Engine/Rendering/Materials/CustomizedUVs:description%
-	[/PARAM]
-	[PARAM:path]
-		[RELATIVE:Engine/Rendering/Materials/CustomizedUVs]
-	[/PARAM]
-[/OBJECT]
-[/VAR]
-[VAR:TopicCompact]
-[OBJECT:TopicCompact]
-	[PARAM:image]
-		![%Engine/Rendering/Materials/CustomizedUVs:title%](Engine/Rendering/Materials/CustomizedUVs/CUVs_topic.png)
-	[/PARAM]
-	[PARAM:icon]
-		![](%ROOT%/tech_icon.png)(convert:false)
-	[/PARAM]
-	[PARAM:title]
-		%Engine/Rendering/Materials/CustomizedUVs:title%
-	[/PARAM]
-	[PARAM:description]
-		%Engine/Rendering/Materials/CustomizedUVs:description%
-	[/PARAM]
-	[PARAM:path]
-		[RELATIVE:Engine/Rendering/Materials/CustomizedUVs]
-	[/PARAM]
-[/OBJECT]
-[/VAR]
 
 [TOC(start:2)]
 
diff --git a/Engine/Documentation/Source/Engine/Rendering/Materials/Editor/MaterialEditor.INT.udn b/Engine/Documentation/Source/Engine/Rendering/Materials/Editor/MaterialEditor.INT.udn
--- a/Engine/Documentation/Source/Engine/Rendering/Materials/Editor/MaterialEditor.INT.udn
+++ b/Engine/Documentation/Source/Engine/Rendering/Materials/Editor/MaterialEditor.INT.udn
@@ -3,26 +3,11 @@ Title:Material Editor Reference
 Crumbs: %ROOT%, Engine, Engine/Rendering/Materials
 Description:Guide to using the Material Editor for creating shaders.
 Version: 4.9
+parent:Engine/Rendering/Materials
+topic-image:Engine/Rendering/Materials/MatEditorUserGuideTopic.png
+type:reference
+order:2
 
-[VAR:Topic]
-[OBJECT:Topic]
-	[PARAM:image]
-		![%Engine/Rendering/Materials/:title%](Engine/Rendering/Materials_topic.png)
-	[/PARAM]
-	[PARAM:icon]
-		![](%ROOT%/texture_icon.png)(convert:false)
-	[/PARAM]
-	[PARAM:title]
-		Material Editor
-	[/PARAM]
-	[PARAM:description]
-		Guide to using the Material Editor for creating shaders.
-	[/PARAM]
-	[PARAM:path]
-		[RELATIVE:Engine/Rendering/Materials/Editor]
-	[/PARAM]
-[/OBJECT]
-[/VAR]
 
 [TOC(start:2)]
 
diff --git a/Engine/Documentation/Source/Engine/Rendering/Materials/ExpressionReference/ExpressionReference.INT.udn b/Engine/Documentation/Source/Engine/Rendering/Materials/ExpressionReference/ExpressionReference.INT.udn
--- a/Engine/Documentation/Source/Engine/Rendering/Materials/ExpressionReference/ExpressionReference.INT.udn
+++ b/Engine/Documentation/Source/Engine/Rendering/Materials/ExpressionReference/ExpressionReference.INT.udn
@@ -3,6 +3,10 @@ title:Material Expression Reference
 Crumbs: %ROOT%, Engine, Engine/Rendering/Materials
 Description:Reference for all material expressions available for use in the Material Editor for creating complex node-based shader networks.
 Version: 4.9
+parent:Engine/Rendering/Materials
+order:9
+topic-image:exp_ref_topic.png
+type:overview
 
 [REGION:banner]
 ![Materials Compendium](compendium.png)
diff --git a/Engine/Documentation/Source/Engine/Rendering/Materials/ExpressionReference/Utility/UtilityExpressions.INT.udn b/Engine/Documentation/Source/Engine/Rendering/Materials/ExpressionReference/Utility/UtilityExpressions.INT.udn
--- a/Engine/Documentation/Source/Engine/Rendering/Materials/ExpressionReference/Utility/UtilityExpressions.INT.udn
+++ b/Engine/Documentation/Source/Engine/Rendering/Materials/ExpressionReference/Utility/UtilityExpressions.INT.udn
@@ -251,22 +251,7 @@ The **Noise** expression creates a procedural noise field, giving you control ov
 
 ## Vector Noise
 
-The **Vector Noise** expression creates a variety of procedural noise effects with 3D output.
-
-| Item | Description |
-| -------- | ----------- |
-| [REGION:tablesection]Properties[/REGION] ||
-| **Function** | [INCLUDE:#VectorNoiseFunctions] |
-| **Quality** | A look/performance setting. Lower values are faster but may look worse, higher values are slower but may look better. |
-| **Tiling** | For noise functions that support it, allows noise to tile. This is more expensive, but useful when baking noise into a seamless wrapping texture. |
-| **Tile Size** | When tiling, how often should the noise repeat. For Perlin noise variants, the Tile Size must be a multiple of three. |
-|| [REGION:tablesection]Inputs[/REGION] ||
-| **Position** | Allows the texture size to be adjusted via a 3D vector. |
-
-![Cellnoise.png](Cellnoise.png)
-![PerlinGradient.png](PerlinGradient.png)
-![Voronoi.png](Voronoi.png)
-
+[INCLUDE:Engine\Rendering\Materials\ExpressionReference\Vector#VectorNoiseTest]
 
 ## QualitySwitch
 
@@ -427,14 +412,14 @@ The image below shows what data the DistanceFieldGradient will use when the vari
 <!--
 [EXCERPT:NoiseFunctions]
 
-|Function| Description|
-|---------|-----------|
+| Function |     | Description|
+| -------- | --- | ----------- |
 |Simplex | Texture Based|High quality for direct use and bumps about 77 instructions per level, 4 Texture lookups, Can not tile.|
 |Gradient | Texture Based| High quality for direct use and bumps. Non-tiled: about 61 instructions per level, 8 Texture lookups. Tiling: about 74 instructions per level, 8 Texture lookups. Even "non tiled' mode has a repeat of 128. Useful repeat size range <= 128. Formerly labeled as Perlin Noise.|
 |Fast Gradient | 3D Texture | High quality for direct use, **BAD** for bumps. about 16 instructions per level, 1 Texture lookup. Always tiles with a repeat of 16, "Tiling" mode is not an option for Fast Gradient noise.|
 | Gradient | Computational | High quality for direct use and bumps. Non-tiled: about 80 instructions per level, no Textures. Tiling: about 143 instructions per level, no Textures. |
 | Value | Computational | Low quality, but pure computation. Non-tiled: about 53 instructions per level, no Textures. Tiling: about 118 instructions per level, no Textures. Formerly mislabeled as Gradient noise. |
-| Voronoi | Also know as Worley or Cellular noise. Quality=1 searches 8 cells, Quality=2 searches 16 cells, Quality=3 searches 27 cells, Quality=4 searches 32 cells. All are about about 20 instructions per cell searched. |
+| Voronoi | | Also know as Worley or Cellular noise. Quality=1 searches 8 cells, Quality=2 searches 16 cells, Quality=3 searches 27 cells, Quality=4 searches 32 cells. All are about about 20 instructions per cell searched. |
 [/EXCERPT]
 -->
 
@@ -450,3 +435,6 @@ The image below shows what data the DistanceFieldGradient will use when the vari
 | Voronoi | The same algorithm and instruction counts as the Voronoi function in the *Noise* expression, but RGB is the location of the closest seed point in each Voronoi cell, and A is the distance to that seed point. |
 [/EXCERPT]
 -->
+
+
+
diff --git a/Engine/Documentation/Source/Engine/Rendering/Materials/ExpressionReference/Vector/VectorExpressions.INT.udn b/Engine/Documentation/Source/Engine/Rendering/Materials/ExpressionReference/Vector/VectorExpressions.INT.udn
--- a/Engine/Documentation/Source/Engine/Rendering/Materials/ExpressionReference/Vector/VectorExpressions.INT.udn
+++ b/Engine/Documentation/Source/Engine/Rendering/Materials/ExpressionReference/Vector/VectorExpressions.INT.udn
@@ -140,3 +140,99 @@ In the example above, the preview sphere would seem to scale up and down with si
 [/EXCERPT:VertexNormalWS]
 
 
+## Vector Noise 
+
+[EXCERPT:VectorNoiseTest]
+![](Vector_Noise_Example.png)
+
+The Vector Noise Material expression adds several more 3D or 4D vector noise results to use in your Materials. Due to the run-time expense of these functions, it is recommended that once a look is developed with them, all or part of the computation be baked into a Texture using the [Draw Material to Render Target Blueprint](Engine\Rendering\RenderTargets\BlueprintRenderTargets) feature introduced in Unreal Engine 4.13 and later. These Material graph Expressions allow procedural looks to be developed in engine on final assets, providing an alternative to creating procedurally generated Textures with an external tool to apply to assets in UE4. Inside of the Vector Noise Material Expression, you will find the following Vector Noise types.
+
+|Image|Item|Description|
+|-----|--------|---------------|
+|![](Cellnoise.png)(w:128)|Cellnoise|Returns a random color for each cell in a 3D grid (i.e. from the mathematical floor operation applied to the node input). The results are always consistent for a given position, so can provide a reliable way to add randomness to a Material. This Vector Noise function is extremely cheap to compute, so it is not necessary to bake it into a Texture for performance.|
+|![](VectorNoise.png)(w:128)|Perlin 3D Noise|Returns a random color for each cell in a 3D grid (i.e. from the mathematical floor operation applied to the node input). The results are always consistent for a given position, so can provide a reliable way to add randomness to a Material. This Vector Noise function is extremely cheap to compute, so it is not necessary to bake it into a Texture for performance.|
+|![](GradientNoise.png)(w:128)|Perlin Gradient|Computes the analytical 3D gradient of a scalar Perlin Simplex Noise. The output is four channels, where the first three (RGB) are the gradient, and the fourth (A) is the scalar noise. This noise type is useful for bumps on a surface or for flow maps.|
+|![](CurlNoise.png)(w:128)|Perlin Gradient|Computes the analytical 3D curl of a vector Perlin Simplex Noise (aka Curl Noise). The output is a 3D signed curl vector and is useful for fluid or particle flow.|
+|![](VoronoiNoise.png)(w:128)|Voronoi|Computes the same Voronoi noise as the scalar Noise material node. The scalar Voronoi noise scatters seed points in 3D space and returns the distance to the closest one. The Vector Noise version returns the location of the closest seed point in RGB, and the distance to it in A. Especially coupled with Cellnoise, this can allow some randomized behavior per Voronoi cell.|
+
+Below is a simple stone bed material using the distance component of the Vector Noise / Voronoi to modulate some surface bumps and blend in moss in the cracks, and the seed position together with Vector Noise / Cellnoise to change the color and bump height per rock.
+
+![](Stone_Example.png)
+
+The derivative-based operations Perlin Curl and Perlin Gradient can be added together in octaves, just as regular Perlin noise can. For derivatives of more complex expressions, it is necessary to compute the gradient of the result of the expression. To help with this, place the expression to compute into a Material Function and use it with the following helper nodes.
+ 
+|Item|Description|
+|-----|-----------|
+|Prepare3DDeriv|Uses positions offset in a tetrahedral pattern to compute 3D derivatives. Evaluate the same 3D function at each offset position produced by this function, then feed the resulting values into Compute3DDeriv.|
+|Compute3DDeriv|Uses positions offset in a tetrahedral pattern to compute 3D derivatives. Use with Prepare3DDeriv.|
+|GradFrom3DDeriv|Computes 3D gradient vector from result of Prepare3DDeriv/Compute3DDeriv.|
+|CurlFrom3DDeriv|Computes curl of at 3D vector field from result of Prepare3DDeriv/Compute3DDeriv.|
+
+[region:note]
+These helper Material Functions use four evaluations of the base expression spaced in a tetrahedral pattern to approximate these derivative-based operations.
+[/region]
+
+Below you will find descriptions for the various noise functions you will find in the Vector Noise Material Expression.
+
+| Item | Description |
+| -------- | ----------- |
+| [REGION:tablesection]Properties[/REGION] ||
+| **Function** | [INCLUDE:#VectorNoiseFunctions] |
+| **Quality** | A look/performance setting. Lower values are faster but may look worse, higher values are slower but may look better. |
+| **Tiling** | For noise functions that support it, allows noise to tile. This is more expensive, but useful when baking noise into a seamless wrapping texture. |
+| **Tile Size** | When tiling, how often should the noise repeat. For Perlin noise variants, the Tile Size must be a multiple of three. |
+|| [REGION:tablesection]Inputs[/REGION] ||
+| **Position** | Allows the texture size to be adjusted via a 3D vector. |
+
+* **Cell Noise** Material Example:
+	[REGION:lightbox]
+	[![](Engine\Rendering\Materials\ExpressionReference\Utility\Cellnoise.png)(w:504)](Engine\Rendering\Materials\ExpressionReference\Utility\Cellnoise.png)
+	[/REGION]
+
+	[REGION:caption]
+	Click for full image.
+	[/REGION]
+
+* **Perlin Gradient**  Material Example:
+	[REGION:lightbox]
+	[![](Engine\Rendering\Materials\ExpressionReference\Utility\PerlinGradient.png)(w:504)](Engine\Rendering\Materials\ExpressionReference\Utility\PerlinGradient.png)
+	[/REGION]
+
+	[REGION:caption]
+	Click for full image.
+	[/REGION]
+
+* **Voronoi**  Material Example:
+	[REGION:lightbox]
+	[![](Engine\Rendering\Materials\ExpressionReference\Utility\Voronoi.png)(w:504)](Engine\Rendering\Materials\ExpressionReference\Utility\Voronoi.png)
+	[/REGION]
+
+	[REGION:caption]
+	Click for full image.
+	[/REGION]
+		
+[/EXCERPT:VectorNoiseTest]
+
+<!--
+[EXCERPT:VectorNoiseFunctions]
+
+|Function| Description|
+|---------|-----------|
+| Cellnoise | Random color for each integer grid cell in 3D space. About 10 instructions. |
+| Perlin 3D Noise | Computational Perlin noise with 3D output, each channel output ranges from -1 to 1. About 83 instructions if only the red channel is used, 125 instructions if all three channels are used|
+| Perlin Gradient | Computes the gradient of a Perlin noise function. RGB output contains the gradient vector, A is the scalar noise. About 106 instructions. |
+| Perlin Curl | Computes a 3D curl noise. Output is the mathematical curl of Perlin 3D Noise. About 162 instructions. |
+| Voronoi | The same algorithm and instruction counts as the Voronoi function in the *Noise* expression, but RGB is the location of the closest seed point in each Voronoi cell, and A is the distance to that seed point. |
+[/EXCERPT]
+-->
+
+
+
+
+
+
+
+
+
+
+
diff --git a/Engine/Documentation/Source/Engine/Rendering/Materials/Functions/Functions.INT.udn b/Engine/Documentation/Source/Engine/Rendering/Materials/Functions/Functions.INT.udn
--- a/Engine/Documentation/Source/Engine/Rendering/Materials/Functions/Functions.INT.udn
+++ b/Engine/Documentation/Source/Engine/Rendering/Materials/Functions/Functions.INT.udn
@@ -3,6 +3,10 @@ Title:Material Functions
 Crumbs: %ROOT%, Engine, Engine/Rendering/Materials
 Description:Snippets of material graphs that can be saved in packages and reused across multiple materials.
 Version: 4.9
+parent:Engine/Rendering/Materials
+order:10
+topic-image:Engine/Rendering/Materials/FunctionTopic.png
+type:overview
 
 [REGION:full_feature]
 [/REGION]
diff --git a/Engine/Documentation/Source/Engine/Rendering/Materials/HowTo/AnimatingUVCoords/AnimatingUVCoordinates.INT.udn b/Engine/Documentation/Source/Engine/Rendering/Materials/HowTo/AnimatingUVCoords/AnimatingUVCoordinates.INT.udn
--- a/Engine/Documentation/Source/Engine/Rendering/Materials/HowTo/AnimatingUVCoords/AnimatingUVCoordinates.INT.udn
+++ b/Engine/Documentation/Source/Engine/Rendering/Materials/HowTo/AnimatingUVCoords/AnimatingUVCoordinates.INT.udn
@@ -1,7 +1,7 @@
 Availability: Public
-Title:Material Editor - How To animate UV Coordinates
+Title:Animating UV Coordinates
 Crumbs: %ROOT%, Engine, Engine/Rendering/Materials, Engine/Rendering/Materials/HowTo
-Description:A How To Guide for animating UV Coordinates
+Description:Guide for animating UV Coordinates
 Related: Engine/Rendering/Materials/ExpressionReference/Coordinates
 Related: Engine/Rendering/Materials/CustomizedUVs
 Related: Engine/Rendering/LightingAndShadows/LightFunctions
diff --git a/Engine/Documentation/Source/Engine/Rendering/Materials/HowTo/BumpOffset/BumpOffset.INT.udn b/Engine/Documentation/Source/Engine/Rendering/Materials/HowTo/BumpOffset/BumpOffset.INT.udn
--- a/Engine/Documentation/Source/Engine/Rendering/Materials/HowTo/BumpOffset/BumpOffset.INT.udn
+++ b/Engine/Documentation/Source/Engine/Rendering/Materials/HowTo/BumpOffset/BumpOffset.INT.udn
@@ -1,7 +1,7 @@
 Availability: Public
-Title:Material Editor - How To use Bump Offset 
+Title:Using Bump Offset 
 Crumbs: %ROOT%, Engine, Engine/Rendering/Materials, Engine/Rendering/Materials/HowTo
-Description: A How To guide for using the Bump Offset node in your Material.
+Description: Guide for using the Bump Offset node in your Materials.
 Related: Resources/Showcases/RealisticRendering
 Related: Engine/Rendering/Materials/ExpressionReference/Utility
 Related: Engine/Rendering/Materials
diff --git a/Engine/Documentation/Source/Engine/Rendering/Materials/HowTo/ClearCoatDualNormal/ClearCoatDualNormal.INT.udn b/Engine/Documentation/Source/Engine/Rendering/Materials/HowTo/ClearCoatDualNormal/ClearCoatDualNormal.INT.udn
--- a/Engine/Documentation/Source/Engine/Rendering/Materials/HowTo/ClearCoatDualNormal/ClearCoatDualNormal.INT.udn
+++ b/Engine/Documentation/Source/Engine/Rendering/Materials/HowTo/ClearCoatDualNormal/ClearCoatDualNormal.INT.udn
@@ -1,6 +1,6 @@
 Availability: Public
-Title:Material Editor - How To use Dual Normals with Clear Coat
-Description: A How-To guide for using Dual Normal maps in your Clear Coat Material.
+Title:Using Dual Normals with Clear Coat
+Description: Guide for using dual Normal Maps in your Clear Coat Material.
 Crumbs: %ROOT%, Engine, Engine/Rendering/Materials, Engine/Rendering/Materials/HowTo
 Version: 4.12
 SkillLevel: Intermediate
@@ -105,7 +105,7 @@ With the Dual - Normal Clear Coat option now enabled, it is now time to create a
 	|Texture Coordinate| N/A|N/A|
 	|Multiply| N/A|N/A|
 	|FlattenNormal| N/A|N/A|
-	|ClearCoatBottomNormal| N/A|N/A|
+	|ClearCoatBottomNormal (ClearCoatNormalCustomOutput)| N/A|N/A|
 
 	![](Dual_CC_Normal_ClearCoat_Setup.png)
 
diff --git a/Engine/Documentation/Source/Engine/Rendering/Materials/HowTo/ColoredTransluscentShadows/ColoredTransluscentShadows.INT.udn b/Engine/Documentation/Source/Engine/Rendering/Materials/HowTo/ColoredTransluscentShadows/ColoredTransluscentShadows.INT.udn
--- a/Engine/Documentation/Source/Engine/Rendering/Materials/HowTo/ColoredTransluscentShadows/ColoredTransluscentShadows.INT.udn
+++ b/Engine/Documentation/Source/Engine/Rendering/Materials/HowTo/ColoredTransluscentShadows/ColoredTransluscentShadows.INT.udn
@@ -1,7 +1,7 @@
 Availability: Public
-Title:Material Editor - How To use Colored Translucent Shadows
+Title:Using Colored Translucent Shadows
 Crumbs: %ROOT%, Engine, Engine/Rendering/Materials
-Description:A How To Guide for setting up and using Colored Translucent Shadows
+Description:Guide for setting up and using Colored Translucent Shadows
 Related: Engine/Rendering/Materials/MaterialProperties/BlendModes
 Related: Engine/Rendering/LightingandShadows/Lightmass
 Related: Engine/Rendering/Materials/Functions/Reference/Opacity
diff --git a/Engine/Documentation/Source/Engine/Rendering/Materials/HowTo/DetailTexturing/DetailTexturing.INT.udn b/Engine/Documentation/Source/Engine/Rendering/Materials/HowTo/DetailTexturing/DetailTexturing.INT.udn
--- a/Engine/Documentation/Source/Engine/Rendering/Materials/HowTo/DetailTexturing/DetailTexturing.INT.udn
+++ b/Engine/Documentation/Source/Engine/Rendering/Materials/HowTo/DetailTexturing/DetailTexturing.INT.udn
@@ -1,7 +1,7 @@
 Availability: Public
-Title:Material Editor - How To add Detail Textures
+Title:Adding Detail Textures
 Crumbs: %ROOT%, Engine, Engine/Rendering/Materials
-Description:A How To Guide for adding Detail Textures to your Materials
+Description:Guide for adding Detail Textures to your Materials
 Related: Engine/Rendering/Materials/ExpressionReference/Coordinates
 Related: Engine/Rendering/Materials/CustomizedUVs
 Related: Engine/Rendering/Materials/Functions/Reference/Texturing
diff --git a/Engine/Documentation/Source/Engine/Rendering/Materials/HowTo/Fresnel/Fresnel.INT.udn b/Engine/Documentation/Source/Engine/Rendering/Materials/HowTo/Fresnel/Fresnel.INT.udn
--- a/Engine/Documentation/Source/Engine/Rendering/Materials/HowTo/Fresnel/Fresnel.INT.udn
+++ b/Engine/Documentation/Source/Engine/Rendering/Materials/HowTo/Fresnel/Fresnel.INT.udn
@@ -1,7 +1,7 @@
 Availability: Public
-Title:Material - Using Fresnel in your Materials
+Title:Using Fresnel in your Materials
 Crumbs: %ROOT%, Engine, Engine/Rendering/Materials
-Description:A How To Guide for using the Fresnel Material node. 
+Description:Guide for using the Fresnel Material node. 
 Related: Engine/Rendering/Materials/Functions/Reference/VectorOps
 Related: Engine/Rendering/Materials/Editor/Interface
 Related: Engine/Rendering/Materials/ExpressionReference/Utility
diff --git a/Engine/Documentation/Source/Engine/Rendering/Materials/HowTo/Human_Skin/HumanSkin.INT.udn b/Engine/Documentation/Source/Engine/Rendering/Materials/HowTo/Human_Skin/HumanSkin.INT.udn
--- a/Engine/Documentation/Source/Engine/Rendering/Materials/HowTo/Human_Skin/HumanSkin.INT.udn
+++ b/Engine/Documentation/Source/Engine/Rendering/Materials/HowTo/Human_Skin/HumanSkin.INT.udn
@@ -1,7 +1,7 @@
 Availability: Public
-Title:Material Editor - How To set up Human Skin. 
+Title:Creating Human Skin
 Crumbs: %ROOT%, Engine, Engine/Rendering/Materials
-Description: A How To guide for setting up Human Skin 
+Description: Guide for setting up Human Skin 
 Related: Resources/Showcases/RealisticRendering
 Related: Engine/Rendering/Materials
 Related: Engine/Rendering/Materials/HowTo/Subsurface_Scattering
diff --git a/Engine/Documentation/Source/Engine/Rendering/Materials/HowTo/Instancing/Instancing.INT.udn b/Engine/Documentation/Source/Engine/Rendering/Materials/HowTo/Instancing/Instancing.INT.udn
--- a/Engine/Documentation/Source/Engine/Rendering/Materials/HowTo/Instancing/Instancing.INT.udn
+++ b/Engine/Documentation/Source/Engine/Rendering/Materials/HowTo/Instancing/Instancing.INT.udn
@@ -1,7 +1,7 @@
 Availability: Public
-Title:Material Editor - How to Create and use Material Instances. 
+Title:Creating and using Material Instances. 
 Crumbs: %ROOT%, Engine, Engine/Rendering/Materials, Engine/Rendering/Materials/HowTo
-Description:A How To Guide for setting up and using Material Instances. 
+Description:Guide for setting up and using Material Instances. 
 Related: Engine/Rendering/Materials/MaterialInstances
 Related: Engine/Rendering/Materials/Editor/Interface
 Related: Engine/Rendering/Materials/MaterialInputs
diff --git a/Engine/Documentation/Source/Engine/Rendering/Materials/HowTo/Main_Material_Node/Main_Material_Node.INT.udn b/Engine/Documentation/Source/Engine/Rendering/Materials/HowTo/Main_Material_Node/Main_Material_Node.INT.udn
--- a/Engine/Documentation/Source/Engine/Rendering/Materials/HowTo/Main_Material_Node/Main_Material_Node.INT.udn
+++ b/Engine/Documentation/Source/Engine/Rendering/Materials/HowTo/Main_Material_Node/Main_Material_Node.INT.udn
@@ -1,7 +1,7 @@
 Availability: Public
-Title:Material Editor - How To Use the Main Material Node
+Title:Using the Main Material Node
 Crumbs: %ROOT%, Engine, Engine/Rendering/Materials, Engine/Rendering/Materials/HowTo
-Description:A How To Guide for Setting Up the Main Material Node
+Description:Guide for setting up and using the Main Material Node
 Related: Engine/Rendering/Materials/Editor/Interface
 Related: Engine/Rendering/Materials/MaterialInputs
 Related: Engine/Rendering/Materials/Editor
diff --git a/Engine/Documentation/Source/Engine/Rendering/Materials/HowTo/Making_Functions/Making_Functions.INT.udn b/Engine/Documentation/Source/Engine/Rendering/Materials/HowTo/Making_Functions/Making_Functions.INT.udn
--- a/Engine/Documentation/Source/Engine/Rendering/Materials/HowTo/Making_Functions/Making_Functions.INT.udn
+++ b/Engine/Documentation/Source/Engine/Rendering/Materials/HowTo/Making_Functions/Making_Functions.INT.udn
@@ -1,7 +1,7 @@
 Availability: Public
-Title:Material Editor - How To Make Material Functions
+Title:Creating Material Functions
 Crumbs: %ROOT%, Engine, Engine/Rendering/Materials, Engine/Rendering/Materials/HowTo
-Description:A How To Guide for using Creating and using Material functions in your materials 
+Description:Guide for creating and using Material functions in your materials 
 Related: Engine/Rendering/Materials/Functions
 Related: Engine/Rendering/Materials/Functions/Creating
 Related: Engine/Rendering/Materials/Functions/Editing
diff --git a/Engine/Documentation/Source/Engine/Rendering/Materials/HowTo/Making_Parameters/Making_Parameters.INT.udn b/Engine/Documentation/Source/Engine/Rendering/Materials/HowTo/Making_Parameters/Making_Parameters.INT.udn
--- a/Engine/Documentation/Source/Engine/Rendering/Materials/HowTo/Making_Parameters/Making_Parameters.INT.udn
+++ b/Engine/Documentation/Source/Engine/Rendering/Materials/HowTo/Making_Parameters/Making_Parameters.INT.udn
@@ -1,7 +1,7 @@
 Availability: Public
-Title:Material Editor - Making Material Parameters
+Title:Making Material Parameters
 Crumbs: %ROOT%, Engine, Engine/Rendering/Materials, Engine/Rendering/Materials/HowTo
-Description:A How To Guide for Setting Up and using Material Parameters
+Description:Guide for setting up and using Material Parameters
 Related: Engine/Rendering/Materials/ExpressionReference/Parameters
 Related: Engine/Rendering/Materials/MaterialInstances
 Related: Engine/Rendering/Materials/ParameterCollections
diff --git a/Engine/Documentation/Source/Engine/Rendering/Materials/HowTo/Masking/Masking.INT.udn b/Engine/Documentation/Source/Engine/Rendering/Materials/HowTo/Masking/Masking.INT.udn
--- a/Engine/Documentation/Source/Engine/Rendering/Materials/HowTo/Masking/Masking.INT.udn
+++ b/Engine/Documentation/Source/Engine/Rendering/Materials/HowTo/Masking/Masking.INT.udn
@@ -1,7 +1,7 @@
 Availability: Public
-Title:Material Editor - How To use Texture Masking 
+Title:Using Texture Masks 
 Crumbs: %ROOT%, Engine, Engine/Rendering/Materials, Engine/Rendering/Materials/HowTo
-Description:A How To Guide for using Texture Masking in your Materials. 
+Description:Guide for using Texture Masking in your Materials. 
 Related: Engine/Content/Types/Textures/Properties
 Related: Engine/Content/Types/Textures/SupportAndSettings
 Related: Engine/Content/Types/Textures/Properties/Interface
diff --git a/Engine/Documentation/Source/Engine/Rendering/Materials/HowTo/MaterialHowTo.INT.udn b/Engine/Documentation/Source/Engine/Rendering/Materials/HowTo/MaterialHowTo.INT.udn
--- a/Engine/Documentation/Source/Engine/Rendering/Materials/HowTo/MaterialHowTo.INT.udn
+++ b/Engine/Documentation/Source/Engine/Rendering/Materials/HowTo/MaterialHowTo.INT.udn
@@ -1,39 +1,24 @@
 Availability: Public
-Title:Material - How To
+Title:Materials How-To's
 Crumbs: %ROOT%, Engine, Engine/Rendering/Materials
 Description:A How To Guide for using different aspects of the Material Editor in Unreal Engine 4. 
-Related: Engine/Rendering/Materials/Editor/Interface
-Related: Engine/Rendering/Materials/MaterialInputs
-Related: Engine/Rendering/Materials/Editor
-Related: Engine/Rendering/Materials/IntroductionToMaterials
-Related: Engine/Rendering/Materials/MaterialProperties
+prereq: Engine/Rendering/Materials/Editor/Interface
+prereq: Engine/Rendering/Materials/MaterialInputs
+prereq: Engine/Rendering/Materials/Editor
+prereq: Engine/Rendering/Materials/IntroductionToMaterials
+prereq: Engine/Rendering/Materials/MaterialProperties
 Version: 4.9
+parent:Engine/Rendering/Materials
+order:3
+type:landing
+topic-image:Materail_How_To.png
+tags:Materials 
 
-[VAR:TopicCompact]
-[OBJECT:TopicCompact]
-	[PARAM:image]
-		![%Engine/Rendering/Materials/HowTo:title%](%ROOT%/Engine/Rendering/Materials/HowTo/Materail_How_To.png)
-	[/PARAM]
-	[PARAM:icon]
-		![](%ROOT%/light_icon.png)(convert:false)
-	[/PARAM]
-	[PARAM:title]
-		%Engine/Rendering/Materials/HowTo:title%
-	[/PARAM]
-	[PARAM:description]
-		%Engine/Rendering/Materials/HowTo:description%
-	[/PARAM]
-	[PARAM:path]
-		[RELATIVE:Engine/Rendering/Materials/HowTo]
-	[/PARAM]
-[/OBJECT]
-[/VAR]
-
-
-
-The following How-To tutorials will demonstrate how you can utilize the full potential of the Material Editor in Unreal Engine 4 (UE4).
-From tips on How - To correctly set up Human Skin to How - To use Layered Materials, the How - To tutorials are a must read for anyone who wants to know more about UE4's Material Editor.
+![](materials_howto.png)
 
+The following step-by-step guides walk you through a wide range of topics designed to help you create and use Materials in your projects.
+<!--
+From tips on How - To correctly set up Human Skin to How - To use Layered Materials, the How - To tutorials are a must read for anyone who wants to know more about UE4's Material Editor.
 
 ## Creating a New Project
 
@@ -62,8 +47,21 @@ If you already have a project that also includes the Starter Content, please fee
 
 Listed below are several How-To guides covering a wide range of topics regarding Materials. 
 Read over each description to get an understanding of what the How-To will cover then click to open and begin that How-To.
+-->
+
+## Starting Out
+
+[DIR(output:"fancy" parent:"Engine/Rendering/Materials/HowTo" skill_level:"Beginner")]
+
+## Doing More with Materials
+
+[DIR(output:"fancy" parent:"Engine/Rendering/Materials/HowTo" skill_level:"Intermediate")]
 
+## Putting It All Together
 
+[DIR(output:"fancy" parent:"Engine/Rendering/Materials/HowTo" skill_level:"Advanced")]
+
+<!--
 |How - To... | Description |
 | ------- | ------------ |
 | **[Use the Main Material Node](Engine\Rendering\Materials\HowTo\Main_Material_Node)** | Setting up the Main Material node is key to get any Material to work. This How-To will show you all you need to know about setting up the Main Material Node. |
@@ -86,8 +84,4 @@ Read over each description to get an understanding of what the How-To will cover
 | **[Colored Translucent Shadows](Engine/Rendering/Materials/HowTo/ColoredTransluscentShadows)** |Getting shadows to look like they are casting colored shadows into the world can be achieved by using Colored Translucent Shadows. In this How-To, we will cover how you can create Materials that make use of Colored Translucent Shadows. |
 | **[Using a Subsurface Profile in your Materials](Engine/Rendering/Materials/HowTo/Subsurface_Profile)** |The ability to rendering realistic looking human skin is a must for any modern Video Game engine. In this How-To, we will cover what a Subsurface Profile is and how you can use them in your work. |
 | **[Using Subsurface Scattering in your Materials](Engine/Rendering/Materials/HowTo/Subsurface_Scattering)** |The term used to describe the lighting phenomenon where light scatters as it passes through a translucent / semi translucent surface. In this How-To, we will walk you through all you need to know about using Subsurface Scattering in your Materials. |
-
-
-
-
-
+-->
\ No newline at end of file
