Copying //UE4/Dev-Documentation to Samples-Main (//UE4/Samples-Main)

diff --git a/Engine/Documentation/Source/Engine/Rendering/Materials/HowTo/PixelNormalOffset/RefractionPixelNormalOffset.INT.udn b/Engine/Documentation/Source/Engine/Rendering/Materials/HowTo/PixelNormalOffset/RefractionPixelNormalOffset.INT.udn
new file mode 100644
--- /dev/null
+++ b/Engine/Documentation/Source/Engine/Rendering/Materials/HowTo/PixelNormalOffset/RefractionPixelNormalOffset.INT.udn
@@ -0,0 +1,92 @@
+Availability: Docs
+Title:Refraction Mode: Pixel Normal Offset
+Crumbs: %ROOT%, Engine, Engine/Rendering/Materials
+Description:How to use the Refraction mode for Pixel Normal Offset in your Materials.
+Related: Engine/REndering/Materials/HowTo/Refraction
+Version: 4.14
+SkillLevel: Intermediate
+
+
+In Unreal Engine 4, we now make use of a non-physical model of refraction, Pixel Normal Offset. The built-in code, using the physical model of refraction for Index of Refraction, is based on how
+light rays refract as they transfer between mediums, which can cause a lot of artifacts as the scene color is being read from off screen. This can be useful for small objects, like glass jars, where 
+you won't notice the effect as much, but is problematic for flat surfaces. Pixel Normal Offset enables refraction for these large flat surfaces, like water, where you do not want this constant offset
+that is reading from off screen so much. It uses the vertex normal as a reference and then computes the refraction offset from how different the per-pixel normal is from the vertex normal allowing these
+flat surfaces not to shift.
+
+### Physical vs Non-Physical Model of Refraction
+
+In the comparisons below, the Refraction Modes for the physical model of **Index of Refraction** and the non-physical model for **Pixel Normal Offset** demonstrate the differences in how the normal is
+read in the material. 
+
+[REGION:note]
+For additional information about the physical model of refraction and using it with your materials you can read about it on our page for [Using Refraction](https://docs.unrealengine.com/latest/INT/Engine/Rendering/Materials/HowTo/Refraction/).
+[/REGION]
+
+[REGION:imagetable]
+| [INCLUDE:#IORWithoutNormals] | [INCLUDE:#PNOWithoutNormals]  |
+| ![](RefractionSettingIOR.png)(w:400) | ![](RefractionSettingPNO.png)(w:400) |
+|---|---|
+| Physical Model of Refraction: Index of Refraction without Normal Map| Non-Physical Model of Refraction: Pixel Normal Offset without Normal Map |
+[/REGION]
+
+Here you will notice that the image is shifted when using the Index of Refraction mode compared to the Pixel Normal Offset mode where you do not read from off screen so much. Index of 
+Refraction will work without a Normal map plugged into the material, whereas with Pixel Normal Offset if there is no Normal map you will not get any refraction.
+
+
+[REGION:imagetable]
+| [INCLUDE:#IORWithNormals] | [INCLUDE:#PNOWithNormals] |
+| ![](RefractionSettingIOR.png)(w:400) | ![](RefractionSettingPNO.png)(w:400) |
+|---|---|
+| Physical Model with Normal Map | Non-Physical Model with Normal Map |
+[/REGION] 
+
+When you add a Normal map to your material the refraction and a value greater than 1 plugged into your Refraction input the normal will translate along the surface when using Pixel Normal Offset.
+However, you will notice that with Index of Refraction you will still have an offset that reads from off screen, which is not a desirable effect for these flat surfaces using refraction.
+
+## Final Results
+
+In this example, the refraction amount is adjusted between a value of 1.0, which is no refraction at all, to a value of 2.0 for some refraction along the surface without shifting the image while using
+Pixel Normal Offset. 
+
+[OBJECT:EmbeddedVideo]
+	[PARAMLITERAL:width]
+	640
+	[/PARAMLITERAL]
+	[PARAMLITERAL:height]
+	360
+	[/PARAMLITERAL]
+	[PARAMLITERAL:videoid]
+	H8FGf4iZdn0
+	[/PARAMLITERAL]
+[/OBJECT]
+
+
+
+
+<!---
+
+[EXCERPT:IORWithoutNormals]
+[REGION:lightbox]
+[![](IOR_WithoutNormals.png)(w:550)](IOR_WithoutNormals.png)
+[/REGION]
+[/EXCERPT:IORWithoutNormals]
+
+[EXCERPT:PNOWithoutNormals]
+[REGION:lightbox]
+[![](PNO_WithoutNormals.png)(w:550)](PNO_WithoutNormals.png)
+[/REGION]
+[/EXCERPT:PNOWithoutNormals]
+
+[EXCERPT:IORWithNormals]
+[REGION:lightbox]
+[![](IOR_WithNormals.png)(w:550)](IOR_WithNormals.png)
+[/REGION]
+[/EXCERPT:IORWithNormals]
+
+[EXCERPT:PNOWithNormals]
+[REGION:lightbox]
+[![](PNO_WithNormals.png)(w:550)](PNO_WithNormals.png)
+[/REGION]
+[/EXCERPT:PNOWithNormals]
+
+--->
diff --git a/Engine/Documentation/Source/Engine/Rendering/Materials/HowTo/Placing_Expressions_Functions/Placing_Expressions_Functions.INT.udn b/Engine/Documentation/Source/Engine/Rendering/Materials/HowTo/Placing_Expressions_Functions/Placing_Expressions_Functions.INT.udn
--- a/Engine/Documentation/Source/Engine/Rendering/Materials/HowTo/Placing_Expressions_Functions/Placing_Expressions_Functions.INT.udn
+++ b/Engine/Documentation/Source/Engine/Rendering/Materials/HowTo/Placing_Expressions_Functions/Placing_Expressions_Functions.INT.udn
@@ -1,7 +1,7 @@
 Availability: Public
-Title:Material Editor - How To Place Material Expression and Functions
+Title:Placing Material Expression and Functions
 Crumbs: %ROOT%, Engine, Engine/Rendering/Materials, Engine/Rendering/Materials/HowTo
-Description:A How To Guide for how to place Material Expression and Functions inside the Material graph.
+Description:Guide for placing Material Expression and Functions inside the Material graph.
 Related: Engine/Rendering/Materials/ExpressionReference/Functions
 Related: Engine/Rendering/Materials/ExpressionReference
 Related: Engine/UI/KeyBindings
diff --git a/Engine/Documentation/Source/Engine/Rendering/Materials/HowTo/Previewing_Applying/Previewing_Applying.INT.udn b/Engine/Documentation/Source/Engine/Rendering/Materials/HowTo/Previewing_Applying/Previewing_Applying.INT.udn
--- a/Engine/Documentation/Source/Engine/Rendering/Materials/HowTo/Previewing_Applying/Previewing_Applying.INT.udn
+++ b/Engine/Documentation/Source/Engine/Rendering/Materials/HowTo/Previewing_Applying/Previewing_Applying.INT.udn
@@ -1,7 +1,7 @@
 Availability: Public
-Title:Material Editor - How To Preview and Apply your Materials
+Title:Previewing and Applying your Materials
 Crumbs: %ROOT%, Engine, Engine/Rendering/Materials, Engine/Rendering/Materials/HowTo
-Description:A How To Guide for Previewing and Applying your Materials
+Description:Guide for previewing and applying your Materials
 Related: Engine/Rendering/Materials/Editor/Interface
 Related: Engine/Rendering/Materials/MaterialInputs
 Related: Engine/Rendering/Materials/Editor
diff --git a/Engine/Documentation/Source/Engine/Rendering/Materials/HowTo/Refraction/Refraction.INT.udn b/Engine/Documentation/Source/Engine/Rendering/Materials/HowTo/Refraction/Refraction.INT.udn
--- a/Engine/Documentation/Source/Engine/Rendering/Materials/HowTo/Refraction/Refraction.INT.udn
+++ b/Engine/Documentation/Source/Engine/Rendering/Materials/HowTo/Refraction/Refraction.INT.udn
@@ -1,7 +1,7 @@
 Availability: Public
-Title:Material Editor - How To use Refraction
+Title:Using Refraction
 Crumbs: %ROOT%, Engine, Engine/Rendering/Materials
-Description:A How To Guide for using Refraction in your Materials.
+Description:Guide for using Refraction in your Materials.
 Related: Resources/ContentExamples/MaterialNodes/1_14
 Related: Resources/ContentExamples/MaterialProperties/1_11
 Related: Engine/Rendering/Materials/MaterialInputs
diff --git a/Engine/Documentation/Source/Engine/Rendering/Materials/HowTo/ShinyMaterials/ShinyMaterials.INT.udn b/Engine/Documentation/Source/Engine/Rendering/Materials/HowTo/ShinyMaterials/ShinyMaterials.INT.udn
--- a/Engine/Documentation/Source/Engine/Rendering/Materials/HowTo/ShinyMaterials/ShinyMaterials.INT.udn
+++ b/Engine/Documentation/Source/Engine/Rendering/Materials/HowTo/ShinyMaterials/ShinyMaterials.INT.udn
@@ -1,7 +1,7 @@
 Availability: Public
-Title:Material Editor - How To make Shiny Materials
+Title:Creating Shiny Materials
 Crumbs: %ROOT%, Engine, Engine/Rendering/Materials
-Description: A How To Guide for making Shiny Materials
+Description: Guide for making Shiny Materials
 Related: Engine/Rendering/Materials/PhysicallyBased
 Related: Resources/ContentExamples/MaterialNodes/1_4
 Related: Resources/ContentExamples/MaterialNodes/1_2
diff --git a/Engine/Documentation/Source/Engine/Rendering/Materials/HowTo/Subsurface_Profile/Subsurface_Profile.INT.udn b/Engine/Documentation/Source/Engine/Rendering/Materials/HowTo/Subsurface_Profile/Subsurface_Profile.INT.udn
--- a/Engine/Documentation/Source/Engine/Rendering/Materials/HowTo/Subsurface_Profile/Subsurface_Profile.INT.udn
+++ b/Engine/Documentation/Source/Engine/Rendering/Materials/HowTo/Subsurface_Profile/Subsurface_Profile.INT.udn
@@ -1,5 +1,5 @@
 Availability:Public
-Title: Material Editor - Using a Subsurface Profile in your Materials 
+Title: Using a Subsurface Profile in your Materials 
 Crumbs: %ROOT%, Engine, Engine/Rendering/Materials
 Description: Setting up Materials to use Subsurface Profiles as they offer a better way to render skin and wax like surfaces.
 Related: Engine/Rendering/Materials/LightingModels/SubSurfaceProfile
@@ -23,9 +23,11 @@ Enabling your Material to use a Subsurface Profile can be done in the following
 1. First create a new Material by **Right-Clicking** in the **Content Browser** and then selecting Material from the **Create Basic Asset** list. 
 Once created, make sure that you name the Material. For this example, the Material will be named **MAT_Subsurface_Profile_00**. 
 When completed, your **Content Browser** should look like this. 
+	
 	![](SSP_Create_New_Material.png)
 
 1. Next, open up the Material by **Double-Clicking** on it in the **Content Browser** with the **Left Mouse Button**.
+	
 	[REGION:lightbox]
 	![](SSP_Open_New_Material.png)
 	[/REGION]
@@ -44,6 +46,7 @@ When completed, your **Content Browser** should look like this.
 	[/OBJECT]
 
 1. At this time, you can assign a **Subsurface Profile** by inputting one into the **Subsurface Profile** section in the **Details** panel.
+	
 	![](SSP_Apply_SSP.png)
 
 1. Now the Material is ready to be used with a Subsurface Profile. 
@@ -89,6 +92,7 @@ In the following steps, we will be creating a Material and a Material Instance t
 | **Subsurface_Amount**| 0.5|
 
 	When everything has been filled in, your shader graph should look something like this. 
+	
 	![](SSP_Named_Value_Nodes.png)
 
 1. With the Material Expression nodes now setup, it is time to connect all of them to the main Material node. 
@@ -210,6 +214,7 @@ When completed, your **Content Browser** should look something like this.
 	Once you have located the mesh, Drag and Drop it from the **Content Browser** into the level. 
 	Once it is in the level, set its location to **0,0,0** in the world. 
 	When completed, your level viewport should look like this. 
+	
 	![](Engine\Rendering\Materials\HowTo\Subsurface_Scattering\SSM_Adding_Meshes.png)
 
 1. Now that the Static Mesh has been placed in the level, we can now apply our Subsurface Material to it. 
diff --git a/Engine/Documentation/Source/Engine/Rendering/Materials/HowTo/Subsurface_Scattering/Subsurface_Scattering.INT.udn b/Engine/Documentation/Source/Engine/Rendering/Materials/HowTo/Subsurface_Scattering/Subsurface_Scattering.INT.udn
--- a/Engine/Documentation/Source/Engine/Rendering/Materials/HowTo/Subsurface_Scattering/Subsurface_Scattering.INT.udn
+++ b/Engine/Documentation/Source/Engine/Rendering/Materials/HowTo/Subsurface_Scattering/Subsurface_Scattering.INT.udn
@@ -1,7 +1,7 @@
 Availability:Public
-Title: Material Editor - How To Use Subsurface Scattering in Your Materials 
+Title: Using Subsurface Scattering in Your Materials 
 Crumbs: %ROOT%, Engine, Engine/Rendering/Materials
-Description: How to create Materials that can 
+Description: Guide for using the Sub-Surface Scattering shading model in your Materials.
 Related: Engine/Rendering/Materials/LightingModels/SubSurfaceProfile
 Related: Engine/Rendering/Materials/LightingModels/SubSurface
 Related: Resources/ContentExamples/MaterialNodes/1_13
@@ -57,6 +57,7 @@ Continuing with our Material from above that has Subsurface Scattering enabled,
 	* **Scalar Parameter** x 2
 
 	![](SSM_Adding_Nodes.png)
+
 	[region:tip]
 	The reason that we are using Parameter Material nodes instead of regular Material nodes is so that we can make a Material Instance out of this Material for easier tweaking inside the editor. 
 	[/region]
@@ -90,6 +91,7 @@ Continuing with our Material from above that has Subsurface Scattering enabled,
 	[/OBJECT]
 
 1. Now that all the nodes have been connected, make sure that you compile the Material by pressing the **Apply** button. Once compiled, you should have something that looks like this. 
+	
 	![](SSM_Finshed_Material.png)
 
 1. Once the Material has been compiled, you can now close the Material window. Then inside of the **Content Browser**, select the Material, **Right-Click** on it, and then select the **Create Material Instance** option from the menu.  
@@ -153,6 +155,7 @@ Continuing with our Material from above that has Subsurface Scattering enabled,
 	Once you have located the mesh, Drag and Drop it from the **Content Browser** into the level. 
 	Once it is in the level, set its location to **0,0,0** in the world. 
 	When completed, your level viewport should look like this. 
+	
 	![](SSM_Adding_Meshes.png)
 
 1. Now that the Static Mesh has been placed in the level, we can now apply our Subsurface Material to it. 
diff --git a/Engine/Documentation/Source/Engine/Rendering/Materials/HowTo/Transparency/Transparency.INT.udn b/Engine/Documentation/Source/Engine/Rendering/Materials/HowTo/Transparency/Transparency.INT.udn
--- a/Engine/Documentation/Source/Engine/Rendering/Materials/HowTo/Transparency/Transparency.INT.udn
+++ b/Engine/Documentation/Source/Engine/Rendering/Materials/HowTo/Transparency/Transparency.INT.udn
@@ -1,7 +1,7 @@
 Availability: Public
-Title:Material Editor - How To Use Transparency
+Title:Using Transparency
 Crumbs: %ROOT%, Engine, Engine/Rendering/Materials
-Description:A How To Guide for using Transparency in your Materials
+Description:Guide for using Transparency in your Materials
 Related: Engine/Rendering/Materials/MaterialProperties/BlendModes
 Related: Engine/Rendering/Materials/Functions/Reference/Opacity
 Related: Engine/Rendering/LightingAndShadows/LitTranslucency
diff --git a/Engine/Documentation/Source/Engine/Rendering/Materials/IntroductionToMaterials/IntroductionToMaterials.INT.udn b/Engine/Documentation/Source/Engine/Rendering/Materials/IntroductionToMaterials/IntroductionToMaterials.INT.udn
--- a/Engine/Documentation/Source/Engine/Rendering/Materials/IntroductionToMaterials/IntroductionToMaterials.INT.udn
+++ b/Engine/Documentation/Source/Engine/Rendering/Materials/IntroductionToMaterials/IntroductionToMaterials.INT.udn
@@ -4,6 +4,9 @@ Crumbs: %ROOT%, Engine, Engine/Rendering/Materials
 Description:An introductory document on UE4 Materials and how they work.
 Version: 4.9
 SkillLevel: Beginner
+topic-image:Engine/Rendering/Materials/EssentialImage.png
+order:1
+type:overview
 
 [VAR:playlistnav]
 [OBJECT:PlaylistNavigation]
@@ -23,40 +26,6 @@ SkillLevel: Beginner
 [/OBJECT]
 [/VAR]
 
-[VAR:topic]
-[OBJECT:Topic]
-	[PARAM:image]
-		![%Engine/Rendering/Materials/IntroductionToMaterials:title%](Engine/Rendering/Materials/EssentialImage.png)
-	[/PARAM]
-	[PARAM:title]
-		%Engine/Rendering/Materials/IntroductionToMaterials:title%
-	[/PARAM]
-	[PARAM:description]
-		%Engine/Rendering/Materials/IntroductionToMaterials:description%
-	[/PARAM]
-	[PARAM:path]
-		[RELATIVE:Engine/Rendering/Materials/IntroductionToMaterials]
-	[/PARAM]
-[/OBJECT]
-[/VAR]
-
-[VAR:TopicCompact]
-[OBJECT:TopicCompact]
-	[PARAM:image]
-		![%Engine/Rendering/Materials/IntroductionToMaterials:title%](Engine/Rendering/Materials/EssentialImage.png)
-	[/PARAM]
-	[PARAM:title]
-		%Engine/Rendering/Materials/IntroductionToMaterials:title%
-	[/PARAM]
-	[PARAM:description]
-		%Engine/Rendering/Materials/IntroductionToMaterials:description%
-	[/PARAM]
-	[PARAM:path]
-		[RELATIVE:Engine/Rendering/Materials/IntroductionToMaterials]
-	[/PARAM]
-[/OBJECT]
-[/VAR]
-
 %playlistnav%
 
 
diff --git a/Engine/Documentation/Source/Engine/Rendering/Materials/LayeredMaterials/LayeredMaterials.INT.udn b/Engine/Documentation/Source/Engine/Rendering/Materials/LayeredMaterials/LayeredMaterials.INT.udn
--- a/Engine/Documentation/Source/Engine/Rendering/Materials/LayeredMaterials/LayeredMaterials.INT.udn
+++ b/Engine/Documentation/Source/Engine/Rendering/Materials/LayeredMaterials/LayeredMaterials.INT.udn
@@ -3,6 +3,10 @@ Title:Layered Materials
 Crumbs: %ROOT%, Engine, Engine/Rendering/Materials
 Description:An introductory document on the Layered Materials technique in UE4.
 Version: 4.9
+parent:Engine/Rendering/Materials
+order:8
+topic-image:Engine/Rendering/Materials/LayeredMaterialsTopic.png
+type:overview
 
 [REGION:fullwidth]
 ![](LayeredMaterials.png)
diff --git a/Engine/Documentation/Source/Engine/Rendering/Materials/LightingModels/LightingModels.INT.udn b/Engine/Documentation/Source/Engine/Rendering/Materials/LightingModels/LightingModels.INT.udn
--- a/Engine/Documentation/Source/Engine/Rendering/Materials/LightingModels/LightingModels.INT.udn
+++ b/Engine/Documentation/Source/Engine/Rendering/Materials/LightingModels/LightingModels.INT.udn
@@ -3,17 +3,10 @@ Title:Material Shading Models
 Crumbs: %ROOT%, Engine, Engine/Rendering/Materials
 Description:Calculation models used by materials to determine how surfaces interact with light.
 Version: 4.9
+parent:Engine/Rendering/Materials
+order:7
+type:landing
+topic-image:shadingmodels_topic.png
 
 
-[OBJECT:TopicButtonList]
-	[PARAM:icon]
-		![](%ROOT%/light_icon.png)(convert:false)
-	[/PARAM]
-	[PARAM:title]
-		Existing Models
-	[/PARAM]
-	[PARAM:links]
-		* [](Engine/Rendering/Materials/LightingModels/SubSurface "%Engine/Rendering/Materials/LightingModels/SubSurface:description%")
-		* [](Engine/Rendering/Materials/LightingModels/SubSurfaceProfile "%Engine/Rendering/Materials/LightingModels/SubSurfaceProfile:description%")
-	[/PARAM]
-[/OBJECT]
+[DIR(output:"topiccompact" parent:"Engine/Rendering/Materials/LightingModels")]
diff --git a/Engine/Documentation/Source/Engine/Rendering/Materials/LightingModels/SubSurface/SSSLightingModel.INT.udn b/Engine/Documentation/Source/Engine/Rendering/Materials/LightingModels/SubSurface/SSSLightingModel.INT.udn
--- a/Engine/Documentation/Source/Engine/Rendering/Materials/LightingModels/SubSurface/SSSLightingModel.INT.udn
+++ b/Engine/Documentation/Source/Engine/Rendering/Materials/LightingModels/SubSurface/SSSLightingModel.INT.udn
@@ -3,6 +3,10 @@ Title:Subsurface Shading Model
 Crumbs: %ROOT%, Engine, Engine/Rendering/Materials
 Description:Description and technical details of the subsurface shading model available in materials.
 Version: 4.9
+parent:Engine/Rendering/Materials/LightingModels
+order:1
+topic-image:sss_topic.png
+type:overview
 
 Materials have a new subsurface shading model (**_MLM_Subsurface_**) which is intended for materials like wax or jade which appear opaque, but lighting scatters inside them allowing a portion of the light from the opposite side of the surface to show through.
 
@@ -38,3 +42,4 @@ Additionally there are some limitations with the current implementation:
 * SSS materials can still receive leaked lighting through fully opaque materials if they are close enough to the shadow caster.
 * Soft shadows from low opacity are limited in how soft they can be and have aliasing.
 
+
diff --git a/Engine/Documentation/Source/Engine/Rendering/Materials/LightingModels/SubSurfaceProfile/SSSProfileLightingModel.INT.udn b/Engine/Documentation/Source/Engine/Rendering/Materials/LightingModels/SubSurfaceProfile/SSSProfileLightingModel.INT.udn
--- a/Engine/Documentation/Source/Engine/Rendering/Materials/LightingModels/SubSurfaceProfile/SSSProfileLightingModel.INT.udn
+++ b/Engine/Documentation/Source/Engine/Rendering/Materials/LightingModels/SubSurfaceProfile/SSSProfileLightingModel.INT.udn
@@ -1,17 +1,20 @@
 Availability:Public
 Title:Subsurface Profile Shading Model
 Crumbs: %ROOT%, Engine, Engine/Rendering/Materials
-Description:Description and technical details of the subsurface profile shading model available in materials.
+Description:Description and technical details of the Subsurface Profile shading model available in Materials.
 Version: 4.9
+parent:Engine/Rendering/Materials/LightingModels
+order:1
+topic-image:sss_profile_topic.png
+type:overview
 
 [TOC(start:2)]
 
+Unreal Engine 4 (UE4) now offers a shading method specifically for rendering realistic skin or wax surfaces called **Subsurface Profile**.
+The Subsurface Profile shading method is similar to the Subsurface method, but with a key difference in how it renders: Subsurface Profile is based in **screen space**. This is because screen space rendering is more effective at displaying the subtle subsurface effects seen in human skin, where backscattering is a secondary effect only seen in few cases, such as ears.
 
 
-The ability to rendering realistic looking human skin is a must for any modern Video Game engine. To fulfill this need, Unreal Engine 4 (UE4) now offers a shading method specifically for skin or wax like surfaces called **Subsurface Profile**. 
-While the Subsurface Profile shading model has similar properties to the Subsurface shading model, its key difference is in how it renders. 
-SubsurfaceProfile uses a rendering approach that is based in Screen Space because this helps to better display the subtle subsurface effects seen in human skin where backscattering is a secondary effect only seen in few cases like ears for example.
-The following document will cover what a Subsurface Profile is and how you can use them in your work. 
+The following document will cover what Subsurface Profiles are and how you can use them in your work. 
 
 [OBJECT:ComparisonSlider]
 	[PARAM:before]
@@ -30,28 +33,27 @@ Special thanks to [REDACTED_PERSON] and his company [Infinite Realities] (http://i
 
 ## What is a Subsurface Profile
 
-The Subsurface scattering profile stores information about how the Subsurface scattering is supposed to be rendered. 
-It can be created, shared, and saved in the **Content Browser** just like any other Actor. 
-Subsurface scattering profiles work by storing artist driven data about the distance the light in the Subsurface should scatter, the color of the Subsurface, and the falloff color of the light once it has exited the object. 
-This data can then be applied to a Subsurface Material to effect what the Subsurface looks like. 
-Subsurface profile can also be tweaked interactively which means that you do not need to re-compile the material to see the results of edits. 
+Subsurface Scattering Profile data can be created, shared, and saved in the **Content Browser**. The data is intended to be authored by artists, and controls the distance the light in the Subsurface should scatter, the color of the Subsurface, and the falloff color of the light once it has exited the object. This data can then be applied to a Subsurface Material. Subsurface Profile data can also be tweaked interactively, meaning that you do not need to re-compile the material to see the results of edits.
+
+
+
 
 
 
 
 ## Enabling, Creating, and Using a Subsurface Profile 
 
-To use a Subsurface Profile in a Material, you must first enable the Material to use a Subsurface Profile by setting the **Shading Model** to **Subsurface Profile** in the **Details** panel of the Material. 
-You can override the Subsurface Profile that is use for this Material by inputting a different Subsurface Profile in the  **Subsurface Profile** input.
+To use a Subsurface Profile in a Material, you must first set the Material to use a Subsurface Profile by setting the **Shading Model** to **Subsurface Profile** in the **Details** panel of the Material. You can override the Subsurface Profile that is used with the **Subsurface Profile** input.
+
 
 [REGION:tip]
-The default settings for the Subsurface Profile are close to Caucasian skin. However please note that this is just one component to realistic looking skin. _Always make sure the base color of your texture fits to your Subsurface scattering profile._
+The default settings for the Subsurface Profile are close to Caucasian skin. Please note that this is just one component to realistic looking skin. _Always make sure the base color of your texture fits to your Subsurface scattering profile._
 [/REGION]
 
 ![Enable SubsurfaceProfile](1.png)
 
-The Subsurface Profile can also be overridden in a Material Instance. To do this, you first need to open the Material Instance that you wish to change. 
-Then in the **Details** section of the Material Instance, enable **Override Subsurface Profile** and then supply the Subsurface profile you want to use in the **Subsurface Profile** input. 
+
+The Subsurface Profile can also be overridden in a Material Instance. To do this, you first need to open the Material Instance that you wish to change. Then, in the **Details** section of the Material Instance, enable **Override Subsurface Profile** and then supply the Subsurface profile you want to use in the **Subsurface Profile** input. 
 
 ![](6.png)
 
@@ -60,54 +62,103 @@ To create a Subsurface Profile, first **Right-Click** inside of the **Content Br
 ![Create SubsurfaceProfile](2.png)
 
 [region:note]
-If no SubsurfaceProfile is specified, it uses the default SubsurfaceProfile which is Caucasian human skin
+If no SubsurfaceProfile is specified, it uses the default, which is Caucasian human skin.
 [/region] 
 
 
 You can edit Subsurface Profiles by **Double-Clicking** on them with the **Left Mouse Button** in the **Content Browser** to open it up. 
-Once open, you can adjust the individual properties of the Subsurface scattering profile by either inputting a number using the keyboard or by using the **Left Mouse Button** to **Click** on the color bar to expose the color picker. 
+Once open, you can adjust the individual properties of the Subsurface Scattering Profile by either inputting a number using the keyboard or by using the **Left Mouse Button** to **Click** on the color bar to expose the color picker. 
 
 
 ![Tweak SubsurfaceProfile](3.png)
 
 * ** Scatter Radius:** The distance in world space units to perform the scatter. 
 
-* ** Subsurface Color:** The Subsurface Color can be used as a weight for the subsurface effect. Black means there is no subsurface scattering. 
-White means all lighting is entering the material and gets scattered around. A non-grayscale value give more control over what
-color contributions are entering the surface resulting in a more complex looking shading. 
+* ** Subsurface Color:** The Subsurface Color can be used as a weight for the Subsurface effect. Black means there is no Subsurface scattering. White means all lighting is entering the material and gets scattered around. A non-grayscale value gives more control over what color contributions are entering the surface, resulting in a more complex looking shading. 
+
+
 
-* **Falloff Color:** The Falloff Color defines the material scattering color once the light has entered the material. 
-You should avoid using a vivid color here if you want to get a more complex shading variation over area you see the scattering.
+
+
+* **Falloff Color:** The Falloff Color defines the material scattering color once the light has entered the material. You should avoid using a vivid color here if you want to get a more complex shading variation over area where you see the scattering.
 
 [region:note]
- Remember that the whole computation is energy conserving so it is not possible to create more light through scattering (unlike the old/other Subsurface method).
+Remember that the whole computation is energy-conserving, so it is not possible to create light through scattering.
 [/region]
 
+## Full-Resolution Skin Shading
+
+UE4 supports full-resolution skin shading for the Subsurface Profile shading model. This provides high-fidelity lighting for surface details, such as pores and wrinkles.
+
+[OBJECT:ComparisonSlider]
+	[PARAM:before]
+	![Checkboard Rendered Skin Layout](FRSH_Checkboard.png)
+	[/PARAM]
+	[PARAM:after]
+	![Full Resolution Skin](FRSH_FullResolution.png)
+	[/PARAM]
+[/OBJECT]
+
+[OBJECT:ComparisonSlider]
+	[PARAM:before]
+	![Checkboard Rendered Skin Layout](FRSH_Checkboard_Zoom.png)
+	[/PARAM]
+	[PARAM:after]
+	![Full Resolution Skin](FRSH_FullResolution_Zoom.png)
+	[/PARAM]
+[/OBJECT]
+
+Previously, lighting on skin was represented using a checkerboard pattern, where half the pixels used only Diffuse lighting and the other half used Specular lighting. The lighting was recombined during a final Subsurface profile fullscreen pass. That approach gave good results for Subsurface lighting (which is low-frequency by nature), but it could result in lower-fidelity lighting for surface details. With the new approach, every pixel contains Diffuse and Specular lighting information, packed into an RGBA encoding. This allows us to reconstruct full-resolution lighting during the final Subsurface profile fullscreen pass, giving better results for surface details and more stable behavior with Temporal Antialiasing. 
+
+### Compatibility
+Full-resolution skin shading requires at least a 64-bit scene color format with a full alpha channel. The default FloatRGBA scene color format works fine, but 32-bit representations such as FloatRGB are not supported. If the scene color format is not compatible with full resolution skin, we fall back to checkerboard-based lighting. 
 
+This behaviour can be overridden using the **r.SSS.Checkerboard** console variable. The possible values for this are:
+* 0: Checkerboard disabled (full resolution)
+* 1: Checkerboard enabled (old behavior)
+* 2: Automatic (default) - Full resolution lighting will be used if the scene color pixelformat supports it
+
+### Limitations
+
+It’s worth noting that the full-resolution skin shading is an approximation. It works well in the vast majority of cases, but certain material features can be problematic due to the encoding method. In particular:
+
+* Metallic Materials
+* Emissive Materials
+
+These features will work, but you may notice differences in output compared to checkerboard due to the packed RGBA diffuse/specular encoding. It is possible to work around some issues when authoring materials by setting the **Opacity** to **0** in areas where skin shading is not desirable. Pixels with an Opacity of zero are treated as default lit for the purposes of shading. 
+
+[region:note]
+Masking non-opaque pixels in this way is also worthwhile for performance reasons, since these pixels are bypassed by the Subsurface Postprocess. 
+[/region]
+
+### Performance Considerations
+
+If your title has a 64-bit scene color format, full-resolution Subsurface lighting will typically be faster than checkerboard due to the reduced number of Texture fetches. However, if your title has a 32-bit scene color, the performance gain from the reduced Texture bandwidth will likely outweigh the benefits (although this is hardware dependent). 
 
 ## Material Input Channels
 
 The Screen Space Subsurface shading profile does not differ much from the Lit shading mode, with the main difference being that the Metallic input has been repurposed and is not available to be used.  
 
-**Base Color Input:** The Base Color input is used as usual for Diffuse lighting. There is no extra Subsurface Scattering color as the Screen Space subsurface scatter is not supposed to change color or brightness, it is just redistributing the lighting to nearby pixels. 
-So if a material should scatter in a specific color it needs to be expressed as part of the Base Color. 
+**Base Color Input:** The Base Color input is used as usual for Diffuse lighting. There is no extra Subsurface Scattering color as the Screen Space Subsurface scatter is not supposed to change color or brightness, it is just redistributing the lighting to nearby pixels. 
+So, if a material should scatter in a specific color, it needs to be expressed as part of the Base Color. 
 The Base Color should be the final color as if the material is viewed from a large distance where the scattering cannot be distinguished from diffuse lighting. 
 
 [REGION:note]
- Human skin can be seen as a thin pale layer with a more vibrant red color below its surface and this is how the default colors look like. The visible scatter distance in Human skin is about 1.2 cm and that is the default of the Subsurface Profile.
+Human skin is a thin layer that blocks a certain amount and color of light, and it covers vibrant, red-colored flesh below its surface. The visible scatter distance in light-colored human skin is about 1.2 cm.
+
 [/region]
 
 **Metallic Input:** The Metallic input channel is not available when using a Subsurface Profile because the GBuffer space for the Metallic input has been repurposed to accommodate the Subsurface Profile data.
 
 **Opacity Input:** The Opacity input channel is used to mask the Subsurface scattering contribution. 
-It works by using a value in the 0 to 1 range to allow for smooth transitions between different areas of Subsurface scattering intensity. 
-For example, when inputting a value of 0 into the Opacity input, Subsurface scattering will be disabled. 
-Inputting a value of 1 will show the full extent of the Subsurface Scattering. 
+It works by using a value in the 0 to 1 range to allow for smooth transitions between different areas of Subsurface scattering intensity, where 0 is no scattering and 1 is full scattering.
+
+
 In order to better control where the Subsurface scattering is stronger or weaker, it is best to use a mask texture.  
 Areas in the mask texture that have values closer to 1, or white, will have the strongest Subsurface scattering effect while areas that are closer to 0, or black, will have less of an effect. 
-Adjusting the Subsurface Color will help to compensate if areas get too dark and remember using brighter colors results in more subsurface scattering. 
+Adjusting the Subsurface Color will help to compensate if areas get too dark. Remember that using brighter colors results in more Subsurface scattering. 
 
-Here you can see how the mask can be used to render two surface types with one material. Note that the transition is soft and not limited to triangles borders.
+Here you can see how the mask can be used to render two surface types with one material. Note that the transition is soft and not limited to triangle borders.
 
 [REGION:lightbox]
 [![](4.png)(w:920 convert:false)](4.png)
@@ -117,29 +168,29 @@ Here you can see how the mask can be used to render two surface types with one m
 
 ## Technical Details
 
-At the moment, the Subsurface scattering profile shading model does not differ much from Lit (Lambert diffuse, GGX for specular, no Metallic).
- Most of the magic happens in a post process after all lighting has been computed. 
+At the moment, the Subsurface scattering profile shading model does not differ much from Lit (Lambert diffuse, GGX for specular, no Metallic). Most of the effect happens in a post process after all lighting has been computed. 
+
 
 [REGION:note]
-The Subsurface scattering profile is based on work from the work from [[REDACTED_PERSON]]([REDACTED_PROFILE_URL]). Make sure to check out his webpage for many useful tips on how to make your 3D images look more realistic.
+The Subsurface scattering profile is based on work from [[REDACTED_PERSON]]([REDACTED_PROFILE_URL]). Make sure to check out his webpage for many useful tips on how to make your 3D images look more realistic.
 [/REGION]
 
-We separate the non-specular (non-view dependent) lighting contributions to support a specular on top of subsurface material and down sample for better performance. 
-Similar to a Gaussian blur we filter the image with a 2 pass (assuming a separable kernel) post process. 
-The filter kernel depends on the subsurface scattering profile which is stored in the GBuffer (Max 255 active profiles per scene). 
+We separate the non-specular (non-view dependent) lighting contributions to support a specular on top of Subsurface Material and down sample for better performance. 
+Similar to a Gaussian blur, we filter the image with a two-pass (assuming a separable kernel) post process. 
+The filter kernel depends on the Subsurface scattering profile which is stored in the GBuffer (Max 255 active profiles per scene). 
 The kernel has colored weights and specific sample positions which can be scaled in the profile (defined in units/cm). 
-In a final step, we recombine the scattered light contribution with the full resolution image. To separate the view and non-view dependent lighting, we store a weighting term in the scene color alpha channel. 
-This approximation requires a 64 bit render target (see r.SceneColorFormat) and this approximation works for most cases
+In the final step, we recombine the scattered light contribution with the full resolution image. To separate the view-dependent and non-view-dependent lighting, we store a weighting value in the scene color alpha channel. 
+This approximation requires a 64 bit render target (see r.SceneColorFormat) and this approximation works for most cases.
 
 It successfully takes out the specular but you get a more desaturated non view dependent color for those specular pixels. This can be improved by using two 32-bit render targets for all lighting passes. That has the same memory bandwidth but on some hardware this might be slower. This might be something we want to change (added code complexity).
 
-Here we have an example where the Specular was removed before the blurring was applied. Notice how in the end image (image on the far right) the Specular reflection is crisp and smooth. This is the effect we wanted to achieve.
+Here we have an example where the Specular was removed before the blurring was applied. Notice how the Specular reflection is crisp and smooth in the end image (image on the far right). This is the effect we wanted to achieve.
 
 [REGION:lightbox]
 [![](Good_Combination.png)(w:920 convert:false)](Good_Combination.png)
 [/REGION]
 
-Here we have an example where the Specular was not removed before the blurring was applied. Notice how in the end image (image on the far right) the Specular reflection is dull and looks a bit stretched. This is not the correct way to have this effect render.  
+Here we have an example where the Specular was not removed before the blurring was applied. Notice how the Specular reflection is dull and looks a bit stretched in the end image (image on the far right). This is not the correct way to render this effect.
 
 [REGION:lightbox]
 [![](Bad_Combination.png)(w:920 convert:false)](Bad_Combination.png)
@@ -149,9 +200,21 @@ Here we have an example where the Specular was not removed before the blurring w
 
 There are some scaling and performance console commands that you can use to help you get a good trade off between high quality visuals and better performance. 
 
-**r.SSS.Scale**: Allows you to disable the post process passes or scale the effect for quick experiments. 
+**r.SSS.Scale**: Can be used to scale the effect for quick experiments. Setting this to **0** will disable the effect. Setting numbers higher than 0 will increase the effect which can be seen in the image sequence below.
 
-**r.SSS.SampleSet**: Allows you to use less samples so the effect will run faster. However this will mean that the effect will have a lower quality or rendering artifacts could show up.
+[INCLUDE:Engine\Rendering\Materials\LightingModels\SubSurfaceProfile\#SSRScale]
+
+
+**r.SSS.SampleSet**: Sets the number of samples used. Decreasing this will cause the effect to run faster. However, this will mean that the effect will have a lower quality, and rendering artifacts could show up.
+
+[OBJECT:ComparisonSlider]
+	[PARAM:before]
+	![r.SSS.SampleSet = 0](RSampleSet_Off.png)
+	[/PARAM]
+	[PARAM:after]
+	![r.SSS.SampleSet = 1](RSampleSet_On.png)
+	[/PARAM]
+[/OBJECT]
 
 The following image shows a bit more of the internals of the system. This view can be enabled using **ShowFlag.VisualizeSSS 1**.
 
@@ -160,12 +223,12 @@ The following image shows a bit more of the internals of the system. This view c
 [/REGION]
 
 While the Subsurface scattering profile shading model is a step forward when it comes to rendering skin, there are some limitations in what it can do.
- _Please note that as this system becomes more and more polished, this list could change. _
+_Please note that as this system becomes more and more polished, this list could change._
 
 * The feature does not work in the non-deferred (mobile) rendering mode.
 * Setting a large screen scatter radius will show banding artifacts in extreme lighting conditions.
-* Currently there is no backscatter of light. 
-* Currently a grey outline occurs when SSS Materials are occluded by non SSS Materials.
+* Currently, there is no backscatter of light. 
+* Currently, a grey outline appears when SSS Materials are occluded by non-SSS Materials.
 
 
 
@@ -178,5 +241,36 @@ Also a very special thanks to [[REDACTED_PERSON]]([REDACTED_PROFILE_URL]) for releas
 
 
 
+<!--
+[EXCERPT:SSRScale]
+[OBJECT:SliderSequence]
+	[PARAMLITERAL:max]
+	11
+	[/PARAMLITERAL]
+	[PARAMLITERAL:ext]
+	jpg
+	[/PARAMLITERAL]
+	[PARAMLITERAL:id]
+	SQ
+	[/PARAMLITERAL]
+	[PARAM:caption]
+	Dragging the slider will show what happens when **r.SSS.Scale** is set from 0 to a value of 10.
+	[/PARAM]
+	[PARAM:image_list]
+	![](SQ_1.png)
+	![](SQ_2.png)
+	![](SQ_3.png)
+	![](SQ_4.png)
+	![](SQ_5.png)
+	![](SQ_6.png)
+	![](SQ_7.png)
+	![](SQ_8.png)
+	![](SQ_9.png)
+	![](SQ_10.png)
+	![](SQ_11.png)
+	[/PARAM]
+[/OBJECT]
+[/EXCERPT:SSRScale]
+-->
 
 
diff --git a/Engine/Documentation/Source/Engine/Rendering/Materials/MaterialInputs/MaterialInputs.INT.udn b/Engine/Documentation/Source/Engine/Rendering/Materials/MaterialInputs/MaterialInputs.INT.udn
--- a/Engine/Documentation/Source/Engine/Rendering/Materials/MaterialInputs/MaterialInputs.INT.udn
+++ b/Engine/Documentation/Source/Engine/Rendering/Materials/MaterialInputs/MaterialInputs.INT.udn
@@ -3,6 +3,10 @@ Title:Material Inputs
 Crumbs: %ROOT%, Engine, Engine/Rendering/Materials, Engine/Rendering/Materials/IntroductionToMaterials
 Description:An in-depth look at the various inputs available on a Material and how they are used.
 Version: 4.9
+parent:Engine/Rendering/Materials
+order:5
+topic-image:Engine/Rendering/Materials/MaterialInputsTopic.png
+type:reference
 
 Here we take a look at the available inputs when making a Material. By feeding values (by way of constants, parameters, and textures) into these inputs, you can define just about any physical surface you can imagine.
 
diff --git a/Engine/Documentation/Source/Engine/Rendering/Materials/MaterialInstances/MaterialInstances.INT.udn b/Engine/Documentation/Source/Engine/Rendering/Materials/MaterialInstances/MaterialInstances.INT.udn
--- a/Engine/Documentation/Source/Engine/Rendering/Materials/MaterialInstances/MaterialInstances.INT.udn
+++ b/Engine/Documentation/Source/Engine/Rendering/Materials/MaterialInstances/MaterialInstances.INT.udn
@@ -3,6 +3,10 @@ Title:Instanced Materials
 Crumbs: %ROOT%, Engine, Engine/Rendering/Materials
 Description:Hierarchical material relationships that allow modification of properties to alter behavior and appearance of child materials without shader recompilation.
 Version: 4.9
+parent:Engine/Rendering/Materials
+order:6
+topic-image:Engine/Rendering/Materials/InstanceTopic.png
+type:overview
 
 
 ![](MICEditor.png)
diff --git a/Engine/Documentation/Source/Engine/Rendering/Materials/MaterialProperties/MaterialProperties.INT.udn b/Engine/Documentation/Source/Engine/Rendering/Materials/MaterialProperties/MaterialProperties.INT.udn
--- a/Engine/Documentation/Source/Engine/Rendering/Materials/MaterialProperties/MaterialProperties.INT.udn
+++ b/Engine/Documentation/Source/Engine/Rendering/Materials/MaterialProperties/MaterialProperties.INT.udn
@@ -3,6 +3,10 @@ Title:Material Properties
 Crumbs: %ROOT%, Engine, Engine/Rendering/Materials, Engine/Rendering/Materials/IntroductionToMaterials
 Description:An introductory document on UE4 Materials and how they work.
 Version: 4.9
+topic-image:Engine/Rendering/Materials/MatProps.png
+parent:Engine/Rendering/Materials
+order:4
+type:reference
 
 [TOC(start:2 end:2)]
 
@@ -77,6 +81,7 @@ The base material node has a number of properties which affect how the material
 | **Disable Depth Test** |  Allows the Material to disable depth tests, which is only meaningful with translucent blend modes.  Disabling depth tests will make rendering significantly slower since no occluded pixels can get Z-culled.<!--- In short, you'll be able to--->  |
 | **Use Translucency Vertex Fog** | When _true_, translucent materials are fogged. Defaults to _true_. |
 
+
 ## Translucency Self Shadowing
 
 ![](TransSelfShadowingProperties.png) ![](PerPixNonDirectional.png)(h:189)
@@ -259,9 +264,12 @@ This allows for control over the lighting mode used by Translucency within this
 [REGION:simpletable]
 | Mode | Description |
 | ---- | ----------- |
-| **TLM_VolumetricNonDirectional** | Lighting will be calculated for a volume, without directionality.  Use this on particle effects like smoke and dust.  This is the cheapest lighting method, however the material normal is not taken into account. |
-| **TLM_VolumetricDirectional** | Lighting will be calculated for a volume, with directionality so that the normal of the material is taken into account.  Note that the default particle tangent space is facing the camera, so enable bGenerateSphericalParticleNormals to get a more useful tangent space. |
-| **TLM_Surface** |Lighting will be calculated for a surface.  Use this on translucent surfaces like glass and water. |
+| **Volumetric NonDirectional** | Lighting will be calculated for a volume, without directionality.  Use this on particle effects like smoke and dust.  This is the cheapest lighting method, however, the material normal is not taken into account. |
+| **Volumetric Directional** | Lighting will be calculated for a volume, with directionality so that the normal of the material is taken into account.  Note that the default particle tangent space is facing the camera, so enable **bGenerateSphericalParticleNormals** to get a more useful tangent space. |
+| **Volumetric PerVertex NonDirectional** |Same as Volumetric NonDirectional, but lighting is only evaluated at vertices so the pixel shader cost is significantly less. Note that lighting still comes from a volume texture, so it is limited in range. Directional lights become unshadowed in the distance. |
+| **Volumetric PerVertex Directional** |Same as Volumetric Directional, but lighting is only evaluated at vertices so the pixel shader cost is significantly less. Note that lighting still comes from a volume texture, so it is limited in range. Directional lights become unshadowed in the distance. |
+| **Surface Translucency Volume** |Lighting will be calculated for a surface. The light in accumulated in a volume so the result is blurry, limited distance but the per pixel cost is very low. Use this on translucent surfaces like glass and water. Only Diffuse lighting is supported. |
+| **Surface Forward Shading** |Lighting will be calculated for a surface. Use this on translucent surfaces like glass and water. This is implemented with forward shading so Specular highlights from local lights are supported, however, many deferred-only features are not. This is the **most expensive** translucency lighting method as each light's contribution is computed per-pixel. |
 [/REGION]
 [/EXCERPT:TransLightingMode]
 
diff --git a/Engine/Documentation/Source/Engine/Rendering/Materials/MaterialsAndTextures.INT.udn b/Engine/Documentation/Source/Engine/Rendering/Materials/MaterialsAndTextures.INT.udn
--- a/Engine/Documentation/Source/Engine/Rendering/Materials/MaterialsAndTextures.INT.udn
+++ b/Engine/Documentation/Source/Engine/Rendering/Materials/MaterialsAndTextures.INT.udn
@@ -4,46 +4,12 @@ Crumbs: %ROOT%, Engine
 Description:Controlling the appearance of surfaces in the world using shaders.
 Navigation:topic
 Version:4.9
+parent:Engine/Rendering
+order:2
+type:landing
+tags:Materials
+topic-image:materials_topic.png
 
-[VAR:Topic]
-[OBJECT:Topic]
-	[PARAM:image]
-		![%Engine/Rendering/Materials:title%](%ROOT%/Engine/Rendering/Materials/Materials_topic.png)
-	[/PARAM]
-	[PARAM:icon]
-		![](%ROOT%/light_icon.png)(convert:false)
-	[/PARAM]
-	[PARAM:title]
-		%Engine/Rendering/Materials:title%
-	[/PARAM]
-	[PARAM:description]
-		%Engine/Rendering/Materials:description%
-	[/PARAM]
-	[PARAM:path]
-		[RELATIVE:Engine/Rendering/Materials]
-	[/PARAM]
-[/OBJECT]
-[/VAR]
-
-[VAR:TopicCompact]
-[OBJECT:TopicCompact]
-	[PARAM:image]
-		![%Engine/Rendering/Materials:title%](%ROOT%/Engine/Rendering/Materials/Materials_topic.png)
-	[/PARAM]
-	[PARAM:icon]
-		![](%ROOT%/light_icon.png)(convert:false)
-	[/PARAM]
-	[PARAM:title]
-		%Engine/Rendering/Materials:title%
-	[/PARAM]
-	[PARAM:description]
-		%Engine/Rendering/Materials:description%
-	[/PARAM]
-	[PARAM:path]
-		[RELATIVE:Engine/Rendering/Materials]
-	[/PARAM]
-[/OBJECT]
-[/VAR]
 
 
 [REGION:banner]
@@ -54,272 +20,16 @@ A **Material** is an asset that can be applied to a mesh to control the visual l
 
 In more technical terms, when light from the scene hits the surface, a Material is used to calculate how that light interacts with that surface. These calculations are done using incoming data that is input to the Material from a variety of images (textures) and math expressions, as well as from various property settings inherent to the Material itself.
 
-
-[REGION:topics third]
-%Engine/Rendering/Materials/IntroductionToMaterials:topiccompact%
-[OBJECT:TopicCompact]
-	[PARAM:image]
-		![%Engine/Rendering/Materials\Editor:title%](MatEditorUserGuideTopic.png)
-	[/PARAM]
-	[PARAM:title]
-		Material Editor Reference
-	[/PARAM]
-	[PARAM:description]
-		%Engine/Rendering/Materials\Editor:description%
-	[/PARAM]
-	[PARAM:path]
-		[RELATIVE:Engine/Rendering/Materials\Editor]
-	[/PARAM]  
-[/OBJECT]
-
-%Engine/Rendering/Materials/HowTo:topiccompact%
-
-[OBJECT:TopicCompact]
-	[PARAM:image]
-		![%Engine/Rendering/Materials\MaterialProperties:title%](MatProps.png)
-	[/PARAM]
-	[PARAM:title]
-		%Engine/Rendering/Materials\MaterialProperties:title%
-	[/PARAM]
-	[PARAM:description]
-		%Engine/Rendering/Materials\MaterialProperties:description%
-	[/PARAM]
-	[PARAM:path]
-		[RELATIVE:Engine/Rendering/Materials\MaterialProperties]
-	[/PARAM]  
-[/OBJECT]
-[OBJECT:TopicCompact]
-	[PARAM:image]
-		![%Engine/Rendering/Materials\MaterialInputs:title%](MaterialInputsTopic.png)
-	[/PARAM]
-	[PARAM:title]
-		%Engine/Rendering/Materials\MaterialInputs:title%
-	[/PARAM]
-	[PARAM:description]
-		%Engine/Rendering/Materials\MaterialInputs:description%
-	[/PARAM]
-	[PARAM:path]
-		[RELATIVE:Engine/Rendering/Materials\MaterialInputs]
-	[/PARAM]  
-[/OBJECT]
-[OBJECT:TopicCompact]
-	[PARAM:image]
-		![%Engine/Rendering/Materials\MaterialInstances:title%](InstanceTopic.png)
-	[/PARAM]
-	[PARAM:title]
-		%Engine/Rendering/Materials\MaterialInstances:title%
-	[/PARAM]
-	[PARAM:description]
-		%Engine/Rendering/Materials\MaterialInstances:description%
-	[/PARAM]
-	[PARAM:path]
-		[RELATIVE:Engine/Rendering/Materials\MaterialInstances]
-	[/PARAM]
-[/OBJECT]
-[OBJECT:TopicCompact]
-	[PARAM:image]
-		![%Engine/Rendering/Materials\LayeredMaterials:title%](LayeredMaterialsTopic.png)
-	[/PARAM]
-	[PARAM:title]
-		%Engine/Rendering/Materials\LayeredMaterials:title%
-	[/PARAM]c
-	[PARAM:description]
-		%Engine/Rendering/Materials\LayeredMaterials:description%
-	[/PARAM]
-	[PARAM:path]
-		[RELATIVE:Engine/Rendering/Materials\LayeredMaterials]
-	[/PARAM]
-[/OBJECT]
-[OBJECT:TopicCompact]
-	[PARAM:image]
-		![%Engine/Rendering/Materials/ExpressionReference:title%](Engine/Rendering/Materials/ExpressionReference/exp_ref_topic.png)
-	[/PARAM]
-	[PARAM:title]
-		%Engine/Rendering/Materials/ExpressionReference:title%
-	[/PARAM]
-	[PARAM:description]
-		%Engine/Rendering/Materials/ExpressionReference:description%
-	[/PARAM]
-	[PARAM:path]
-		[RELATIVE:Engine/Rendering/Materials/ExpressionReference]
-	[/PARAM]
-[/OBJECT]
-[OBJECT:TopicCompact]
-	[PARAM:image]
-		![%Engine/Rendering/Materials/ExpressionReference:title%](FunctionTopic.png)
-	[/PARAM]
-	[PARAM:title]
-		%Engine/Rendering/Materials/Functions:title%
-	[/PARAM]
-	[PARAM:description]
-		%Engine/Rendering/Materials/Functions:description%
-	[/PARAM]
-	[PARAM:path]
-		[RELATIVE:Engine/Rendering/Materials/Functions]
-	[/PARAM]
-[/OBJECT]
-[OBJECT:TopicCompact]
-	[PARAM:image]
-		![%Engine/Rendering/Materials/ParameterCollections:title%](ParamCollectTopic.png)
-	[/PARAM]
-	[PARAM:title]
-		%Engine/Rendering/Materials/ParameterCollections:title%
-	[/PARAM]
-	[PARAM:description]
-		%Engine/Rendering/Materials\ParameterCollections:description%
-	[/PARAM]
-	[PARAM:path]
-		[RELATIVE:Engine/Rendering/Materials\ParameterCollections]
-	[/PARAM]  
-[/OBJECT]
-[OBJECT:TopicCompact]
-	[PARAM:image]
-		![%Engine/Rendering/Materials/CustomizedUVs:title%](Engine/Rendering/Materials/CustomizedUVs/CUVs_topic.png)
-	[/PARAM]
-	[PARAM:icon]
-		![](%ROOT%/tech_icon.png)(convert:false)
-	[/PARAM]
-	[PARAM:title]
-		%Engine/Rendering/Materials/CustomizedUVs:title%
-	[/PARAM]
-	[PARAM:description]
-		%Engine/Rendering/Materials/CustomizedUVs:description%
-	[/PARAM]
-	[PARAM:path]
-		[RELATIVE:Engine/Rendering/Materials/CustomizedUVs]
-	[/PARAM]
-[/OBJECT]
-[/REGION]
+[DIR(output:"topiccompact" parent:"Engine/Rendering/Materials" org:"hierarchy" tags:"!PBR" end:"1")]
 
 ## Physically-Based Shading
 
 Unreal Engine 4 utilizes a _physically based shading model_. This means that rather than defining a Material using arbitrary properties (such as Diffuse Color and Specular Power), you instead use properties more easily relatable to the real world. These include Base Color, Metallic, Specular, and Roughness. For more information, please see the page below.
 
-[REGION:topics third]
-[OBJECT:TopicCompact]
-	[PARAM:image]
-		![%Engine/Rendering/Materials\PhysicallyBased:title%](Engine/Rendering/Materials/PhysicallyBased/physically_based_topic.png)
-	[/PARAM]
-	[PARAM:title]
-		%Engine/Rendering/Materials\PhysicallyBased:title%
-	[/PARAM]
-	[PARAM:description]
-		%Engine/Rendering/Materials\PhysicallyBased:description%
-	[/PARAM]
-	[PARAM:path]
-		[RELATIVE:Engine/Rendering/Materials\PhysicallyBased]
-	[/PARAM]
-[/OBJECT]
-[/REGION]
+[DIR(output:"topiccompact" parent:"Engine/Rendering/Materials" org:"hierarchy" tags:"PBR" end:"1")]
 
 _The above page is designed to provide in-depth details about our shift to this model, along with some of the technical considerations for it._
 
 ## Introductory Tutorials
 
-[REGION:topics third]
-[OBJECT:TopicCompact]
-	[PARAM:image]
-		![](Engine/Content/Types/Textures/Importing/ImportingTextures.png)
-	[/PARAM]
-	[PARAM:title]
-		Importing Textures
-	[/PARAM]
-	[PARAM:path]
-		[RELATIVE:Engine/Content/Types/Textures\Importing]
-	[/PARAM]
-[/OBJECT]
-[/REGION]
-
-
-
-<!---
-
-[REGION:topics third]
-[OBJECT:TopicCompact]
-	[PARAM:image]
-		![%Engine/Rendering/Materials/PhysicallyBased:title%](Engine/Rendering/Materials/PhysicallyBased/physically_based_topic.png)
-	[/PARAM]
-	[PARAM:title]
-		%Engine/Rendering/Materials/PhysicallyBased:title%
-	[/PARAM]
-	[PARAM:description]
-		%Engine/Rendering/Materials/PhysicallyBased:description%
-	[/PARAM]
-	[PARAM:path]
-		[RELATIVE:Engine/Rendering/Materials/PhysicallyBased]
-	[/PARAM]
-[/OBJECT]
-
-[OBJECT:TopicCompact]
-	[PARAM:image]
-		![%Engine/Rendering/Materials/ExpressionReference:title%](Engine/Rendering/Materials/ExpressionReference/exp_ref_topic.png)
-	[/PARAM]
-	[PARAM:title]
-		Quick Start: Physically-Based Materials
-	[/PARAM]
-	[PARAM:description]
-		%Engine/Rendering/Materials/ExpressionReference:description%
-	[/PARAM]
-	[PARAM:path]
-		[RELATIVE:Engine/Rendering/Materials/ExpressionReference]
-	[/PARAM]
-[/OBJECT]
-
-[OBJECT:TopicCompact]
-	[PARAM:image]
-		![%Engine/Rendering/Materials/ExpressionReference:title%](Engine/Rendering/Materials/Functions/matfunc_topic.png)
-	[/PARAM]
-	[PARAM:title]
-		%Engine/Rendering/Materials/Functions:title%
-	[/PARAM]
-	[PARAM:description]
-		%Engine/Rendering/Materials/Functions:description%
-	[/PARAM]
-	[PARAM:path]
-		[RELATIVE:Engine/Rendering/Materials/Functions]
-	[/PARAM]
-[/OBJECT]
-
-
-[/REGION]
-
-[OBJECT:TopicList]
-	[PARAM:icon]
-		![](%ROOT%/start_icon.png)
-	[/PARAM]
-	[PARAM:title]
-		Materials Essentials
-	[/PARAM]
-	[PARAM:description]
-	[/PARAM]
-	[PARAM:links]
-		* [Material Editor User Guide](Engine/Rendering/Materials/Editor) - Overview of the node-based Material Editor in UE4.
-		* [Material Instance Editor User Guide](Engine/Rendering/Materials/MaterialInstances/Editor) - Overview of the editor for Material Instances in UE4.
-		* [Instanced Materials](Engine/Rendering/Materials/MaterialInstances) - Overview of the material instancing system in UE4.
-		* Material Layers - How to use the power of Material Functions to apply multiple Materials to a single object on a pixel-by-pixel basis.
-
-	[/PARAM]
-[/OBJECT]
-
-[OBJECT:TopicList]
-	[PARAM:icon]
-		![](%ROOT%/texture_icon.png)
-	[/PARAM]
-	[PARAM:title]
-		Materials Content Creation
-	[/PARAM]
-	[PARAM:description]
-	[/PARAM]
-	[PARAM:links]
-		* [Creating Hair Using Alpha](Engine/Rendering/Materials/CreatingHairUsingAlpha) - Creating hair using materials and meshes in UE4.
-		* [World Position Offset](Engine/Rendering/Materials/WorldPositionOffset) - How to use the WorldPositionOffset material input.
-		* [Material Masks](Engine/Rendering/Materials/Masks) - Creating masks and gradients algorithmically using material expressions.
-	[/PARAM]
-[/OBJECT]
-
---->
-
-
-
-
-
+%Engine/Content/Types/Textures/Importing:topiccompact%
diff --git a/Engine/Documentation/Source/Engine/Rendering/Materials/ParameterCollections/ParemeterCollections.INT.udn b/Engine/Documentation/Source/Engine/Rendering/Materials/ParameterCollections/ParemeterCollections.INT.udn
--- a/Engine/Documentation/Source/Engine/Rendering/Materials/ParameterCollections/ParemeterCollections.INT.udn
+++ b/Engine/Documentation/Source/Engine/Rendering/Materials/ParameterCollections/ParemeterCollections.INT.udn
@@ -8,6 +8,10 @@ Related: Engine/Rendering/Materials/Editor
 Related: Engine/Rendering/Materials/HowTo/Instancing
 Related: Engine/Blueprints/UserGuide/Timelines
 Version: 4.9
+parent:Engine/Rendering/Materials
+order:11
+topic-image:Engine/Rendering/Materials/ParamCollectTopic.png
+type:overview
 
 [TOC(start:2)]
 
diff --git a/Engine/Documentation/Source/Engine/Rendering/Materials/PhysicallyBased/PhysicallyBasedShaders.INT.udn b/Engine/Documentation/Source/Engine/Rendering/Materials/PhysicallyBased/PhysicallyBasedShaders.INT.udn
--- a/Engine/Documentation/Source/Engine/Rendering/Materials/PhysicallyBased/PhysicallyBasedShaders.INT.udn
+++ b/Engine/Documentation/Source/Engine/Rendering/Materials/PhysicallyBased/PhysicallyBasedShaders.INT.udn
@@ -3,6 +3,12 @@ Title:Physically Based Materials
 Crumbs: %ROOT%, Engine
 Description:An overview of the primary Material inputs and how best to use them.
 Version: 4.9
+parent:Engine/Rendering/Materials
+order:13
+topic-image:physically_based_topic.png
+type:overview
+tags:Physically Based Shading
+tags:PBR
 
 [REGION:fullwidth]
 ![](PhysMatHeader.png)
diff --git a/Engine/Documentation/Source/Engine/Rendering/Overview/RenderingOverview.INT.udn b/Engine/Documentation/Source/Engine/Rendering/Overview/RenderingOverview.INT.udn
--- a/Engine/Documentation/Source/Engine/Rendering/Overview/RenderingOverview.INT.udn
+++ b/Engine/Documentation/Source/Engine/Rendering/Overview/RenderingOverview.INT.udn
@@ -3,6 +3,9 @@ Title: Rendering Overview
 Crumbs: %ROOT%, Engine
 Description:Overview of the main features of the rendering subsystem.
 Version: 4.9
+parent:Engine/Rendering
+order:1
+type:landing
 
 [VAR:Topic]
 [OBJECT:Topic]
diff --git a/Engine/Documentation/Source/Engine/Rendering/ParticleSystems/ParticleSystems.INT.udn b/Engine/Documentation/Source/Engine/Rendering/ParticleSystems/ParticleSystems.INT.udn
--- a/Engine/Documentation/Source/Engine/Rendering/ParticleSystems/ParticleSystems.INT.udn
+++ b/Engine/Documentation/Source/Engine/Rendering/ParticleSystems/ParticleSystems.INT.udn
@@ -4,6 +4,9 @@ Crumbs: %ROOT%, Engine, Engine/Rendering
 Description:Assets the contain any number of particle emitters used to create effects in the world.
 Version: 4.9
 SkillLevel: Beginner
+parent:Engine/Rendering
+order:4
+type:landing
 
 [VAR:Topic]
 [OBJECT:Topic]
diff --git a/Engine/Documentation/Source/Engine/Rendering/PostProcessEffects/PostProcessEffects.INT.udn b/Engine/Documentation/Source/Engine/Rendering/PostProcessEffects/PostProcessEffects.INT.udn
--- a/Engine/Documentation/Source/Engine/Rendering/PostProcessEffects/PostProcessEffects.INT.udn
+++ b/Engine/Documentation/Source/Engine/Rendering/PostProcessEffects/PostProcessEffects.INT.udn
@@ -4,6 +4,9 @@ Crumbs: %ROOT%, Engine
 Description:Effects applied to the whole rendered scene prior to being rendered.
 Navigation:topic
 Version: 4.9
+parent:Engine/Rendering
+order:5
+type:landing
 
 [VAR:Topic]
 [OBJECT:Topic]
diff --git a/Engine/Documentation/Source/Engine/Rendering/PostProcessEffects/Vignette/Vignette.INT.udn b/Engine/Documentation/Source/Engine/Rendering/PostProcessEffects/Vignette/Vignette.INT.udn
--- a/Engine/Documentation/Source/Engine/Rendering/PostProcessEffects/Vignette/Vignette.INT.udn
+++ b/Engine/Documentation/Source/Engine/Rendering/PostProcessEffects/Vignette/Vignette.INT.udn
@@ -22,7 +22,7 @@ Version: 4.9
 [EXCERPT:Intensity]
 [OBJECT:ComparisonSlider]
 	[PARAM:before]
-	![Vignetter Intensity - 0](vignette_off.png)(w:600)
+	![Vignette Intensity - 0](vignette_off.png)(w:600)
 	[/PARAM]
 	[PARAM:after]
 	![Vignette Intensity - 1](vignette_on.png)(w:600)
diff --git a/Engine/Documentation/Source/Engine/Rendering/RenderTargets/BlueprintRenderTargets/BestPractices/BlueprintRenderTargetsBestPractices.INT.udn b/Engine/Documentation/Source/Engine/Rendering/RenderTargets/BlueprintRenderTargets/BestPractices/BlueprintRenderTargetsBestPractices.INT.udn
--- a/Engine/Documentation/Source/Engine/Rendering/RenderTargets/BlueprintRenderTargets/BestPractices/BlueprintRenderTargetsBestPractices.INT.udn
+++ b/Engine/Documentation/Source/Engine/Rendering/RenderTargets/BlueprintRenderTargets/BestPractices/BlueprintRenderTargetsBestPractices.INT.udn
@@ -1,5 +1,5 @@
 availability:docs
-title:Blueprint & Render Target Best Practices
+title:Blueprint and Render Target Best Practices
 description: A collection of Best Practices to follow when using Blueprints and Render Targets.
 crumbs:
 version:4.13
diff --git a/Engine/Documentation/Source/Engine/Rendering/RenderTargets/BlueprintRenderTargets/BlueprintRenderTargets.INT.udn b/Engine/Documentation/Source/Engine/Rendering/RenderTargets/BlueprintRenderTargets/BlueprintRenderTargets.INT.udn
--- a/Engine/Documentation/Source/Engine/Rendering/RenderTargets/BlueprintRenderTargets/BlueprintRenderTargets.INT.udn
+++ b/Engine/Documentation/Source/Engine/Rendering/RenderTargets/BlueprintRenderTargets/BlueprintRenderTargets.INT.udn
@@ -1,93 +1,22 @@
-availability:docs
-title: Blueprint & Render Targets
-description: Using Render Targets with Blueprints in Unreal Engine 4
-crumbs:
+availability:Docs
+title: Blueprints and Render Targets
+description: Using Render Targets with Blueprints in Unreal Engine 4.
+crumbs:%ROOT%, Engine, Engine/Rendering, Engine/Rendering/RenderTargets
 version:4.13
 Type:landing
-SkillLevel:Intermediate
-parent:Engine
+parent:Engine/Rendering/RenderTargets
 order:1
 Navigation:topic
+tags: Render Targets
+tags: Blueprint
+topic-image:Engine/Rendering/RenderTargets/Render_Targets_Overview.png
 
 ![](Engine\Rendering\RenderTargets\BlueprintRenderTarget.png)
 
-Having the ability to use Blueprints to draw Materials into Render Targets and then use or save those results can enable a huge variety of game-specific rendering effects. On the following page, you will find several links to documentation pages that cover how to use Blueprints and Render Targets. 
-
-
-[REGION:topics third]
-[OBJECT:Topic]
-	[PARAM:image]
-		![Engine\Rendering\RenderTargets\BlueprintRenderTargets/Overview:title%](Engine/Rendering/RenderTargets/Render_Targets_Overview.png)
-	[/PARAM]
-	[PARAM:icon]
-		![](%ROOT%/blueprint_icon.png)(convert:false)
-	[/PARAM]
-	[PARAM:title]
-		%Engine/Rendering/RenderTargets/BlueprintRenderTargets/Overview:title%
-	[/PARAM]
-	[PARAM:description]
-		%Engine/Rendering/RenderTargets/BlueprintRenderTargets/Overview:description%
-	[/PARAM]
-	[PARAM:path]
-		[RELATIVE:Engine/Rendering/RenderTargets/BlueprintRenderTargets/Overview]
-	[/PARAM]
-[/OBJECT]
-
-[OBJECT:Topic]
-	[PARAM:image]
-		![Engine\Rendering\RenderTargets\BlueprintRenderTargets/Reference:title%](Engine/Rendering/RenderTargets/Render_Targets_Reference.png)
-	[/PARAM]
-	[PARAM:icon]
-		![](%ROOT%/blueprint_icon.png)(convert:false)
-	[/PARAM]
-	[PARAM:title]
-		%Engine/Rendering/RenderTargets/BlueprintRenderTargets/Reference:title%
-	[/PARAM]
-	[PARAM:description]
-		%Engine/Rendering/RenderTargets/BlueprintRenderTargets/Reference:description%
-	[/PARAM]
-	[PARAM:path]
-		[RELATIVE:Engine/Rendering/RenderTargets/BlueprintRenderTargets/Reference]
-	[/PARAM]
-[/OBJECT]
-
-[OBJECT:Topic]
-	[PARAM:image]
-		![Engine\Rendering\RenderTargets\BlueprintRenderTargets/BestPractices:title%](Engine/Rendering/RenderTargets/Render_Targets_Best_Practies.png)
-	[/PARAM]
-	[PARAM:icon]
-		![](%ROOT%/blueprint_icon.png)(convert:false)
-	[/PARAM]
-	[PARAM:title]
-		%Engine/Rendering/RenderTargets/BlueprintRenderTargets/BestPractices:title%
-	[/PARAM]
-	[PARAM:description]
-		%Engine/Rendering/RenderTargets/BlueprintRenderTargets/BestPractices:description%
-	[/PARAM]
-	[PARAM:path]
-		[RELATIVE:Engine/Rendering/RenderTargets/BlueprintRenderTargets/BestPractices]
-	[/PARAM]
-[/OBJECT]
-
-[OBJECT:Topic]
-	[PARAM:image]
-		![Engine\Rendering\RenderTargets\BlueprintRenderTargets/BestPractices:title%](Engine/Rendering/RenderTargets/Render_Targets_Best_Practies.png)
-	[/PARAM]
-	[PARAM:icon]
-		![](%ROOT%/blueprint_icon.png)(convert:false)
-	[/PARAM]
-	[PARAM:title]
-		%Engine/Rendering/RenderTargets/BlueprintRenderTargets/HowTo:title%
-	[/PARAM]
-	[PARAM:description]
-		%Engine/Rendering/RenderTargets/BlueprintRenderTargets/HowTo:description%
-	[/PARAM]
-	[PARAM:path]
-		[RELATIVE:Engine/Rendering/RenderTargets/BlueprintRenderTargets/HowTo]
-	[/PARAM]
-[/OBJECT]
-
-[/REGION]
+Using Blueprints to draw Materials into Render Targets, and then using or saving those results, can enable a variety of game-specific rendering effects. 
+On the following pages, you will find several links to documentation that covers how to use Blueprints and Render Targets. 
+
+[DIR(output:"topic" parent:"Engine/Rendering/RenderTargets/BlueprintRenderTargets")]
 
 
 
diff --git a/Engine/Documentation/Source/Engine/Rendering/RenderTargets/BlueprintRenderTargets/HowTo/BlueprintRenderTargetsHowTos.INT.udn b/Engine/Documentation/Source/Engine/Rendering/RenderTargets/BlueprintRenderTargets/HowTo/BlueprintRenderTargetsHowTos.INT.udn
--- a/Engine/Documentation/Source/Engine/Rendering/RenderTargets/BlueprintRenderTargets/HowTo/BlueprintRenderTargetsHowTos.INT.udn
+++ b/Engine/Documentation/Source/Engine/Rendering/RenderTargets/BlueprintRenderTargets/HowTo/BlueprintRenderTargetsHowTos.INT.udn
@@ -1,56 +1,22 @@
-availability:docs
-title: Blueprint and Render Targets - How Tos
-description: A How To Guide for using Blueprints to manipulate render targets. 
-crumbs:
+availability:Docs
+title: Blueprint and Render Targets How To's
+description: How to guides for using Blueprints to manipulate render targets. 
+crumbs:%ROOT%, Engine, Engine\Rendering\RenderTargets
 version:4.13
-Type:topic
-SkillLevel:Intermediate
-parent:Engine
-order:2
-Navigation:topic
+Type:how-to
+parent:Engine/Rendering/RenderTargets/BlueprintRenderTargets
+Navigation:landing
+tags: Render Targets
+tags: Blueprint
+topic-image:Engine/Rendering/RenderTargets/Render_Targets_Best_Practies.png
 
-The **Blueprint and Render Targets - How To** pages provide step by step examples on how to use Blueprints to work with render targets. Whether you want to use Blueprints and render targets to create Textures for use in your project's Materials, or you want to create a surface that deforms when the player shoots at it, these pages will guide you on how to accomplish this. Reference links are also included to provi
+The following **Blueprint and Render Targets - How To** pages provide step-by-step examples on how to use Blueprints and render targets together in your Unreal Engine 4 (UE4) projects. 
+<!--Reference links are also included to provide you with additional information.-->
 
 ## Intermediate
-[REGION:wrapper transparent]
-[](FANCY:Engine\Rendering\RenderTargets\BlueprintRenderTargets\HowTo\CreatingTextures "%Engine\Rendering\RenderTargets\BlueprintRenderTargets\HowTo\CreatingTextures:description%")
-
-[/REGION]
 
+[DIR(output:"fancy" parent:"Engine/Rendering/RenderTargets/BlueprintRenderTargets/HowTo" skill_level:"Intermediate")]
 
 ## Advanced
-[REGION:wrapper transparent]
-[](FANCY:Engine\Rendering\RenderTargets\BlueprintRenderTargets\HowTo\HeightFieldPainter "%Engine\Rendering\RenderTargets\BlueprintRenderTargets\HowTo\HeightFieldPainter:description%")
-[/REGION]
-
-
-
-
-
-
-
-
-
-
-
-
-
-
-
-
-
-
-
-
-
-
-
-
-
-
-
-
-
-
-
 
+[DIR(output:"fancy" parent:"Engine/Rendering/RenderTargets/BlueprintRenderTargets/HowTo" skill_level:"Advanced")]
\ No newline at end of file
diff --git a/Engine/Documentation/Source/Engine/Rendering/RenderTargets/BlueprintRenderTargets/HowTo/CreatingTextures/1/CTHT_1.INT.udn b/Engine/Documentation/Source/Engine/Rendering/RenderTargets/BlueprintRenderTargets/HowTo/CreatingTextures/1/CTHT_1.INT.udn
new file mode 100644
--- /dev/null
+++ b/Engine/Documentation/Source/Engine/Rendering/RenderTargets/BlueprintRenderTargets/HowTo/CreatingTextures/1/CTHT_1.INT.udn
@@ -0,0 +1,69 @@
+availability:docs
+title:1. Project Setup
+description:Going over the Project Setup for creating Textures using Blueprints and Render Targets.
+crumbs:%ROOT%, Engine, Engine\Rendering\RenderTargets
+version:4.13
+type:multi-step
+SkillLevel:Intermediate
+parent:Engine/Rendering/RenderTargets/BlueprintRenderTargets/HowTo/CreatingTextures
+Navigation:multi-step
+tags: Render Targets
+order:1
+checkpoint:CTUB
+
+[NAV]
+
+
+In this section of the Creating Textures using Blueprints and Render Targets, we will go over creating a new level for our project to use.
+
+## Create a New Empty Level
+
+Creating a new **Empty Level** to place your Texture creation Blueprint in will make things easier to manage, ensuring that nothing else is affecting 
+your Texture creation process. To create a new Empty Level inside of UE4, you will need to do the following:
+
+1. From the **Main Toolbar**, go to **File** > **New Level**.
+
+	![](Engine\Rendering\RenderTargets\BlueprintRenderTargets\HowTo\CreatingTextures\New_Level_01.png)
+
+1. From the New Level window that is displayed, click on the **Empty Level** option.
+
+	![](Engine\Rendering\RenderTargets\BlueprintRenderTargets\HowTo\CreatingTextures\New_Level_02.png)
+
+1. After clicking on the **Empty Level** option, UE4 will load a new level that has nothing inside of it. Once that has happened, save the level by clicking on the **Save** button on the Main Toolbar, naming it `BlueprintTextureCreation`.
+
+	![](Engine\Rendering\RenderTargets\BlueprintRenderTargets\HowTo\CreatingTextures\New_Level_03.png)
+
+With the level now created and saved in the next section we will go over creating and setting up the required assets. 
+
+[NAV]
+
+
+
+
+
+
+
+
+
+
+
+
+
+
+
+
+
+
+
+
+
+
+
+
+
+
+
+
+
+
+
diff --git a/Engine/Documentation/Source/Engine/Rendering/RenderTargets/BlueprintRenderTargets/HowTo/CreatingTextures/2/CTHT_2.INT.udn b/Engine/Documentation/Source/Engine/Rendering/RenderTargets/BlueprintRenderTargets/HowTo/CreatingTextures/2/CTHT_2.INT.udn
new file mode 100644
--- /dev/null
+++ b/Engine/Documentation/Source/Engine/Rendering/RenderTargets/BlueprintRenderTargets/HowTo/CreatingTextures/2/CTHT_2.INT.udn
@@ -0,0 +1,61 @@
+availability:docs
+title:2. Required Asset Creation & Setup
+description:Going over creating and setting up the required assets.
+crumbs:%ROOT%, Engine, Engine\Rendering\RenderTargets
+version:4.13
+Type:multi-step
+SkillLevel:Intermediate
+parent:Engine/Rendering/RenderTargets/BlueprintRenderTargets/HowTo/CreatingTextures
+Navigation:multi-step
+tags: Render Targets
+Order:2
+checkpoint:CTUB
+
+[NAV]
+
+In this section we will go over how to create and setup the required assets so you can createTextures using Blueprints and Render Targets.
+
+## Render Target Creation & Setup
+To create a **Render Target** inside of the UE4 Content Browser, you will need to do the following:
+
+1. Open up the UE4 Editor and then  right-click inside of the **Content Browser**. From the menu that is displayed, locate the **Create Advanced Asset** section.
+
+	![](Engine\Rendering\RenderTargets\BlueprintRenderTargets\HowTo\CreatingTextures\Create_RTT_01.png)
+
+1. From the Create Advanced Assets menu, mouseover the **Materials & Textures** section and from the asset list that is displayed, select the **Render Target** option.
+
+	![](Engine\Rendering\RenderTargets\BlueprintRenderTargets\HowTo\CreatingTextures\Create_RTT_02.png)
+
+1. Name the newly created render target `RT_00`. Your Content Browser should now look like the following image:
+
+	![](Engine\Rendering\RenderTargets\BlueprintRenderTargets\HowTo\CreatingTextures\Create_RTT_03.png)
+	
+## Material Creation & Setup
+To create and set up a Material to use a render target, you will need to do the following inside of the UE4 Content Browser. 
+
+1. Right-click in the Content Browser and from the **Create Basic Asset** section, select the **Material** option to create a new Material.  
+
+	![](Engine\Rendering\RenderTargets\BlueprintRenderTargets\HowTo\CreatingTextures\Create_RTT_Material_01.png)
+
+1. Name the Material `RT_Material_00` and then double-click on the material to open it up in the **Material Editor**.
+
+	![](Engine\Rendering\RenderTargets\BlueprintRenderTargets\HowTo\CreatingTextures\Create_RTT_Material_02.png)
+
+1. With `RT_Material_00` open in the Material Editor, add the following **Material Expression** nodes to the **Material Graph**.
+
+	* **Texture Coordinate**
+	* **Constant**
+	* **Append**
+	* **Noise**
+
+	![](Engine\Rendering\RenderTargets\BlueprintRenderTargets\HowTo\CreatingTextures\Create_RTT_Material_03.png)
+
+1. Once the required Material Expression nodes have been added to the Material Graph, connect the nodes so that they match the image below, pressing the **Apply** button when completed to compile the Material.
+
+	![](Engine\Rendering\RenderTargets\BlueprintRenderTargets\HowTo\CreatingTextures\Create_RTT_Material_05.png)
+
+
+With the required assets now setup, in the next section, we will take a look at how to set everything up to work together.
+
+
+[NAV]
\ No newline at end of file
diff --git a/Engine/Documentation/Source/Engine/Rendering/RenderTargets/BlueprintRenderTargets/HowTo/CreatingTextures/3/CTHT_3.INT.udn b/Engine/Documentation/Source/Engine/Rendering/RenderTargets/BlueprintRenderTargets/HowTo/CreatingTextures/3/CTHT_3.INT.udn
new file mode 100644
--- /dev/null
+++ b/Engine/Documentation/Source/Engine/Rendering/RenderTargets/BlueprintRenderTargets/HowTo/CreatingTextures/3/CTHT_3.INT.udn
@@ -0,0 +1,70 @@
+availability:docs
+title:3. Blueprint Setup
+description:Going over setting up the Blueprint to work with Render Targets.
+crumbs:%ROOT%, Engine, Engine\Rendering\RenderTargets
+version:4.13
+Type:multi-step
+SkillLevel:Intermediate
+parent:Engine/Rendering/RenderTargets/BlueprintRenderTargets/HowTo/CreatingTextures
+Navigation:multi-step
+tags: Render Targets
+Order:3
+checkpoint:CTUB
+
+[NAV]
+
+In this section, we will take a look at creating a Blueprint that will work with the assets that were created in the previous steps.
+
+## Blueprint Creation & Setup
+To create a Blueprint that will work with both render targets and Materials, you will need to do the following.
+
+1. Right-click in the **Content Browser** and from the menu that is displayed in the **Create Basic Asset** section, select the **Blueprint** option to create a new Blueprint.  
+
+	![](Engine\Rendering\RenderTargets\BlueprintRenderTargets\HowTo\CreatingTextures\Create_Blueprint_01.png)
+
+1. From the **Pick Parent Class** window, select the **Actor** class and then name the Blueprint `BP_Texture_Creator`.
+
+	![](Engine\Rendering\RenderTargets\BlueprintRenderTargets\HowTo\CreatingTextures\Create_Blueprint_02.png)
+
+1. Double-click on the `BP_Texture_Creator` Blueprint to open it up and once opened, click on the **Event Graph** tab to open up the Event Graph.
+
+	![](Engine\Rendering\RenderTargets\BlueprintRenderTargets\HowTo\CreatingTextures\Create_Blueprint_03.png)
+
+1. Add the following Blueprint nodes to the Event Graph:
+	* **Custom Event**
+	* **Draw Material to Render Target** 
+
+	![](Engine\Rendering\RenderTargets\BlueprintRenderTargets\HowTo\CreatingTextures\Create_Blueprint_04.png)
+
+1. Select the Custom Event node and press **F2** so that you can rename the Custom Event node to `Bake`. Then, connect it's output to the input of the **Draw Material to Render Target** node.
+
+	![](Engine\Rendering\RenderTargets\BlueprintRenderTargets\HowTo\CreatingTextures\Create_Blueprint_05.png)
+
+1. Next, on the **Draw Material to Render Target** node, load the render target that was created previously into the **Texture Render Target** slot by clicking on the small white triangle next to **Select Asset**, and then from the list, find the render target `RT_00` and select it by clicking on it. Once that is completed, do the same for the **Material** section, selecting `RT_Material_00`.
+
+	![](Engine\Rendering\RenderTargets\BlueprintRenderTargets\HowTo\CreatingTextures\Create_Blueprint_06.png)
+
+1. Now, click on the **Construction Script** tab and drag off the Construction Script node to display the **Executable actions** menu.
+	
+	![](Engine\Rendering\RenderTargets\BlueprintRenderTargets\HowTo\CreatingTextures\Create_Blueprint_09.png)
+
+	[region:warning]
+	Please note that DrawMaterialToRenderTarget will not work in Construction scripts outside of the UE4 editor which means if you are trying to do this at runtime it will not work. 
+	[/region]
+
+1. From this menu locate the **Bake** event that was created by inputting **Bake** into the search box and once found, click on it to add it to the graph and then connect the output of the Construction Script to the input on the Bake node. 
+
+	![](Engine\Rendering\RenderTargets\BlueprintRenderTargets\HowTo\CreatingTextures\Create_Blueprint_07.png)
+
+1. Finally, press the **Compile** and **Save** buttons to compile and save your Blueprint.
+
+	![](Engine\Rendering\RenderTargets\BlueprintRenderTargets\HowTo\CreatingTextures\Create_Blueprint_08.png)
+
+
+Now that the Blueprint has been setup, in the next section, we will take a look at how you capture and save a Texture that displays what was created in the Material.
+
+
+
+
+
+[NAV]
\ No newline at end of file
diff --git a/Engine/Documentation/Source/Engine/Rendering/RenderTargets/BlueprintRenderTargets/HowTo/CreatingTextures/4/CTHT_4.INT.udn b/Engine/Documentation/Source/Engine/Rendering/RenderTargets/BlueprintRenderTargets/HowTo/CreatingTextures/4/CTHT_4.INT.udn
new file mode 100644
--- /dev/null
+++ b/Engine/Documentation/Source/Engine/Rendering/RenderTargets/BlueprintRenderTargets/HowTo/CreatingTextures/4/CTHT_4.INT.udn
@@ -0,0 +1,61 @@
+availability:docs
+title:4. Texture Creation & End Result
+description:Taking a look at how to create and save a Texture.
+crumbs:%ROOT%, Engine, Engine\Rendering\RenderTargets
+version:4.13
+Type:multi-step
+SkillLevel:Intermediate
+parent:Engine/Rendering/RenderTargets/BlueprintRenderTargets/HowTo/CreatingTextures
+Navigation:multi-step
+tags: Render Targets
+Order:4
+checkpoint:CTUB
+
+[NAV]
+
+In this section, we will take a look at how to bring all parts of this how to together so that we can create and save Texture assets based on the Material that was created.
+
+## Texture Creation
+Now that all of the required parts have been created and set up, we will take a look at how to create and save the Textures that are created.
+
+1. Locate the `BP_Texture_Creator` Blueprint in the Content Browser, dragging it into the level.
+
+	![](Engine\Rendering\RenderTargets\BlueprintRenderTargets\HowTo\CreatingTextures\Create_Textures_01.png)
+
+1. Locate the `RT_Material_00` Material in the Content Browser and double-click on it to open it up. When opened, select the **Noise** Material Expression node.
+
+	![](Engine\Rendering\RenderTargets\BlueprintRenderTargets\HowTo\CreatingTextures\Create_Textures_02.png)
+
+1. With the **Noise** Material Expression selected,  go to the **Details** panel and input the following values into their respective parameters. When completed, press the **Apply** button to compile the Material.
+
+	![](Engine\Rendering\RenderTargets\BlueprintRenderTargets\HowTo\CreatingTextures\Create_Textures_03.png)
+
+	* **Scale:** 2.0 
+	* **Levels:** 8.0
+	* **Output Min:** -0.01
+	* **Level Scale:** 3.0
+
+1. Select the `BP_Texture_Creator` Blueprint that was placed in the level and move it slightly in any direction to update the `RT_00` render target that is in the Content Browser. When you do this, you should see the `RT_00` render target update and display exactly what was in the `RT_Material_00`. 
+
+	![](Engine\Rendering\RenderTargets\BlueprintRenderTargets\HowTo\CreatingTextures\Create_Textures_04.png)
+
+	[region:note]
+	Anytime you want to see or update the changes that were made in the Material, you will need to [slightly] move the Blueprint. 
+	[/region]
+
+1. Inside of the Content Browser, right-click on the `RT_00` render target and in the **Render Targets Actions** section of the menu, select the **Create Static Texture** option to create a new Texture asset that will display the output of `RT_Material_00`.
+
+	![](Engine\Rendering\RenderTargets\BlueprintRenderTargets\HowTo\CreatingTextures\Create_Textures_05.png)
+
+## End Result
+
+When completed, you should have a new 2D Texture that will display the contents of `RT_Material_00` that is now ready to be used in any other Material of your liking. 
+
+![](Engine\Rendering\RenderTargets\BlueprintRenderTargets\HowTo\CreatingTextures\Create_Textures_06.png)
+
+Remember that in order to see your changes you will need to make sure that you compile your Material and [slightly] move your Blueprint so that your changes will be propagated to your render target.
+
+ 
+
+
+[NAV]
\ No newline at end of file
diff --git a/Engine/Documentation/Source/Engine/Rendering/RenderTargets/BlueprintRenderTargets/HowTo/CreatingTextures/CreatingTexturesHowTo.INT.udn b/Engine/Documentation/Source/Engine/Rendering/RenderTargets/BlueprintRenderTargets/HowTo/CreatingTextures/CreatingTexturesHowTo.INT.udn
--- a/Engine/Documentation/Source/Engine/Rendering/RenderTargets/BlueprintRenderTargets/HowTo/CreatingTextures/CreatingTexturesHowTo.INT.udn
+++ b/Engine/Documentation/Source/Engine/Rendering/RenderTargets/BlueprintRenderTargets/HowTo/CreatingTextures/CreatingTexturesHowTo.INT.udn
@@ -1,156 +1,23 @@
 availability:docs
-title:Creating Textures using Blueprints and Render Targets
+title:Creating Textures Using Blueprints and Render Targets
 description:A demonstration of how to use Blueprints and Render Targets to create Textures in UE4.
 crumbs:%ROOT%, Engine, Engine\Rendering\RenderTargets
 version:4.13
-Type:topic
+Type:how-to
 SkillLevel:Intermediate
-parent:Engine
-order:2
-Navigation:topic
-
-[TOC(start:2)]
+parent:Engine/Rendering/RenderTargets/BlueprintRenderTargets/HowTo
+order:1
+Navigation:landing
+tags: Render Targets
+Related:Gameplay/Framework/GameMode/
+Related:Engine/Blueprints/
+Related:Engine/Rendering/Materials/
+Related:Engine/Blueprints/UserGuide/UserConstructionScript/
+checkpoint:CTUB
 
 ![](Creating_Textures_Header_00.png)
 
-When creating Textures for your Unreal Engine 4 (UE4) project, you will generally need to use an additional 2D painting program like Adobe Photoshop or GIMP. While authoring content in this manner can give you amazing results, there are some instances where having the ability to create content like this inside of UE4 would be faster and easier. In the following How To, we will take a look at how you can use Blueprints and Render Targets to create Textures for your UE4 projects.
-
-## Required Asset Creation & Setup
-[region:note]
-Please note that the following is not specific to any UE4 project type and can be accomplished in any project that uses UE4 version 4.13 or later. 
-[/region]
-
-Before we can create and use Textures, we need to set up the following three assets inside of the UE4 Content Browser:
-* **Render Target**
-* **Material**
-* **Blueprint** (Based on Actor)
-
-### Create a new Empty Level
-Creating a new **Empty Level** to place your Texture creation Blueprint in will make things easier to manage and also ensure that nothing else is affecting your Texture creation process. To create a new Empty Level inside of UE4, you will need to do the following:
-
-1. From the **Main Toolbar** go to **File** > **New Level**.
-
-	![](New_Level_01.png)
-
-1. From the New Level window that is displayed, click on the **Empty Level** option.
-
-	![](New_Level_02.png)
-
-1. After clicking on the Empty Level option, UE4 will load a new level that has nothing inside of it. Once that has happened, save the level by clicking on the **Save** button on the Main Toolbar and giving it a name of **BlueprintTextureCreation**.
-
-	![](New_Level_03.png)
-
-### Render Target Creation & Setup
-To create a **Render Target** inside of the UE4 Content Browser, you will need to do the following:
-
-1. Open up the UE4 Editor and then **right-click** Inside of the **Content Browser**. From the menu that is displayed, select the **Create Advanced Assets** section.
-
-	![](Create_RTT_01.png)
-
-1. From the Create Advanced Assets menu, move your mouse to the **Materials & Textures** section and select the **Render Target** option.
-
-	![](Create_RTT_02.png)
-
-1. Name the newly created Render Target, **RT_00**. Your Content Browser should look like the following image:
-
-	![](Create_RTT_03.png)
-	
-### Material Creation & Setup
-To create and set up a Material to use a Render Target, you will need to do the following inside of the UE4 Content Browser:
-
-1. Right-click in the Content Browser and from the **Create Basic Asset** section, select the **Material** option to create a new Material.  
-
-	![](Create_RTT_Material_01.png)
-
-1. Name the Material **RT_Material_00** and then **double-click** on it to open it up in the Material Editor.
-
-	![](Create_RTT_Material_02.png)
-
-1. With the RT_Material_00 open in the Material Editor, add the following Material Expression nodes to the Material Graph:
-
-	* **Texture Coordinate**
-	* **Constant**
-	* **Append**
-	* **Noise**
-
-	![](Create_RTT_Material_03.png)
-
-1. Once the required Material Expression nodes have been added to the Material Graph, connect the nodes so that they match the image below, then press the **Apply** button to compile the Material.
-
-	![](Create_RTT_Material_05.png)
-
-### Blueprint Creation & Setup
-To create and set up a Blueprint that will work with both Render Targets and Materials, you will need to do the following:
-
-1. Right-click in the Content Browser. From the menu that is displayed, under the **Create Basic Asset** section, select the **Blueprint** option to create a new Blueprint.  
-
-	![](Create_Blueprint_01.png)
-
-1. From the **Pick Parent Class** window, select the **Actor** class, and then name the Blueprint **BP_Texture_Creator**.
-
-	![](Create_Blueprint_02.png)
-
-1. Double-click on the BP_Texture_Creator Blueprint to open it up. Once it is opened, click on the **Event Graph** tab to open up the Event Graph.
-
-	![](Create_Blueprint_03.png)
-
-1. Add the following Blueprint nodes to the Event Graph:
-	* **Custom Event**
-	* **Draw Material to Render Target** 
-
-	![](Create_Blueprint_04.png)
-
-1. Select the Custom Event node and press **F2** so you can rename the Custom Event node to **Bake**. Then connect its output to the input of the Draw Material to Render Target node.
-
-	![](Create_Blueprint_05.png)
-
-1. Next, on the Draw Material to Render Target node, load the Render Target that was created previously into the **Texture Render Target** slot by clicking on the small white triangle next to the text **Select Asset** and then find the Render Target, RT_00, in the list and select it by clicking on it. Once that is completed, do the same for the **Material** section, selecting the **RT_Material_00**.
-
-	![](Create_Blueprint_06.png)
-
-1. Now, click on the **Construction Script** tab and drag off the Construction Script node to display the **Executable actions** menu. From this menu, locate the **Bake** event that was created in the Event Graph and click on it to add it. 
-
-	![](Create_Blueprint_07.png)
-
-1. Finally, press the **Compile** and **Save** buttons to compile and save your Blueprint.
-
-	![](Create_Blueprint_08.png)
-
-## Texture Creation
-Now that all of the required parts have been created and set up, we will take a look at how to create and save our Textures.
-
-1. Locate the **BP_Texture_Creator** Blueprint in the Content Browser and drag it into the level.
-
-	![](Create_Textures_01.png)
-
-1. Locate the **RT_Material_00** Material in the Content Browser and double-click on it to open it up. When opened, select the **Noise** Material Expression node.
-
-	![](Create_Textures_02.png)
-
-1. With the **Noise** Material Expression selected, go to the **Details** panel and input the following values into the following parameters. When this is complete, press the **Apply** button to compile the Material.
-
-	![](Create_Textures_03.png)
-
-	* **Scale:** 2.0 
-	* **Levels:** 8.0
-	* **Output Min:** -0.01
-	* **Level Scale:** 3.0
-
-1. Select the BP_Texture_Creator Blueprint that was placed in the level and move it slightly in any direction to update the RT_00 Render Target asset. When you do this, you should see RT_00 update and display exactly what was in RT_Material_00. 
-
-	![](Create_Textures_04.png)
-
-	[region:note]
-	Anytime you want to see or update the changes that were made in the Material, you will need to move the Blueprint in order for the changes to propagate. 
-	[/region]
-
-1. Inside of the Content Browser, right-click on the RT_00 Render Target. In the **Render Targets Actions** section of the menu, select the **Create Static Texture** option to create a new Texture asset that will display the output of RT_Material_00.
-
-	![](Create_Textures_05.png)
-
-1. When completed, you should have a new 2D Texture that will display the contents of RT_Material_00, and is ready for use in any other Material.
-
-	![](Create_Textures_06.png)
+When creating Textures for your Unreal Engine 4 (UE4) project, you will generally need to use an external 2D painting program like Adobe Photoshop or GIMP. While authoring content in this manner can give you amazing results, there are some instances where having the ability to create content like this inside of UE4 would be helpful. In the following How - To, we will take a look at how you can use Blueprints and render targets to create Textures 
+for your UE4 projects by first creating a new empty level in a UE4 project and then creating and setting up the render targets for use.
 
-## Conclusion
-Having the ability to create new Textures inside of UE4 opens up an entirely new world of content creation possibilities. Remember that, in order to see your changes, you will need to make sure that you compile your Material and slightly move your Blueprint so that your changes will be propagated to your Render Target.
\ No newline at end of file
+[DIR(output:"checkpoint" parent:"Engine/Rendering/RenderTargets/BlueprintRenderTargets/HowTo/CreatingTextures")]
diff --git a/Engine/Documentation/Source/Engine/Rendering/RenderTargets/BlueprintRenderTargets/HowTo/FluidSurface/1/FSBP1.INT.udn b/Engine/Documentation/Source/Engine/Rendering/RenderTargets/BlueprintRenderTargets/HowTo/FluidSurface/1/FSBP1.INT.udn
new file mode 100644
--- /dev/null
+++ b/Engine/Documentation/Source/Engine/Rendering/RenderTargets/BlueprintRenderTargets/HowTo/FluidSurface/1/FSBP1.INT.udn
@@ -0,0 +1,29 @@
+availability:docs
+title:1. Project Setup
+description:Going over the Project Setup for creating a Fluid Surface Blueprint.
+crumbs:%ROOT%, Engine, Engine\Rendering\RenderTargets
+version:4.13
+type:multi-step
+SkillLevel:Advanced
+parent:Engine/Rendering/RenderTargets/BlueprintRenderTargets/HowTo/FluidSurface
+Navigation:multi-step
+tags: Render Targets
+order:1
+checkpoint:FSB
+
+[NAV]
+
+In this section, we will go over how to set your project up so that you can create a fluid surface Blueprint.
+
+
+## Project Setup
+The project setup for the Fluid surface is the same as for the Height Field Painter. If you have not done so, read the **Project Setup** section of the [](Engine\Rendering\RenderTargets\BlueprintRenderTargets\HowTo\HeightFieldPainter) How To guide. 
+When completed you should now have a new Game Mode and HUD. When you run the project, you should see something similar to the image below. 
+
+![](Engine\Rendering\RenderTargets\BlueprintRenderTargets\HowTo\HeightFieldPainter\HFP_PS_07.png)
+
+
+ In the next section we will take a look at creating all of the required assets to create a Fluid Surface Blueprint.
+
+
+[NAV]
\ No newline at end of file
