Copying //UE4/Dev-Documentation to Samples-Main (//UE4/Samples-Main)

diff --git a/Engine/Documentation/Source/Engine/Rendering/RenderTargets/BlueprintRenderTargets/HowTo/FluidSurface/2/FSBP2.INT.udn b/Engine/Documentation/Source/Engine/Rendering/RenderTargets/BlueprintRenderTargets/HowTo/FluidSurface/2/FSBP2.INT.udn
new file mode 100644
--- /dev/null
+++ b/Engine/Documentation/Source/Engine/Rendering/RenderTargets/BlueprintRenderTargets/HowTo/FluidSurface/2/FSBP2.INT.udn
@@ -0,0 +1,2133 @@
+availability:docs
+title:2. Required Asset Creation & Setup
+description:Going over creating and setting up the required assets for the Fluid Surface Blueprint.
+crumbs:%ROOT%, Engine, Engine\Rendering\RenderTargets
+version:4.13
+type:multi-step
+SkillLevel:Advanced
+parent:Engine/Rendering/RenderTargets/BlueprintRenderTargets/HowTo/FluidSurface
+Navigation:multi-step
+tags: Render Targets
+order:2
+checkpoint:FSB
+
+[NAV]
+
+For the fluid surface to function correctly, we are going to need to create some new assets. We'll also use some assets created for the Height Field Painter. In the following section, we will go over how each of these assets needs to be setup and any special properties that you will need to use. 
+
+## Render Target Creation & Setup
+For the fluid surface to function correctly, we will need to create four new Render Targets by right-clicking in the **Content Browser** and from the menu that is displayed, go to **Create Advanced Asset** > **Materials & Textures** to select the **Render Target** option.
+
+![](Engine\Rendering\RenderTargets\BlueprintRenderTargets\HowTo\HeightFieldPainter\HFP_RT_Creation_00.png)
+
+Once you have created four new Render Targets, give them the following names:
+
+* `Height0`
+* `Height1`
+* `Height2`
+* `HeightfieldNormal`
+
+![](Engine\Rendering\RenderTargets\BlueprintRenderTargets\HowTo\FluidSurface\FS_Render_Targets.png)
+
+## Material Creation & Setup
+For the fluid simulation to work, we are going to need to create two new Materials; one Material for the water, and one Material for the heightmap simulation. To create these Materials, right-click in the **Content Browser** and from the **Create Basic Assets** list, select the Material option.
+
+[region:note]
+Do not forget to Compile and Save your Blueprint so that it can be used.
+[/region]
+
+* **WaterMaterial** is what makes the surface look and act like water. Before you start to build WaterMaterial, make sure the following options have been set in the **Details** panel for the WaterMaterial.
+	* **Blend Mode:** Translucent
+	* **Lighting Mode:** Surface Translucency Volume
+	* **Used with Static Lighting:** True
+	* **Uses Distortion:** True
+	* **Refraction Mode:** Pixel Normal Offset
+
+	[REGION:fullwidth]
+	[INCLUDE:Engine\Rendering\RenderTargets\BlueprintRenderTargets\HowTo\FluidSurface\2\#WaterMaterial]
+	[/REGION]
+
+* The **HeightSim** Material is used to propagate the various waves that happen when you shoot at the surface of the fluid simulation. Before you start to build the HeightSim Material, make sure the following options have been set in the **Details** panel for the HeightSim Material.
+	* **Shading Model:** Unlit
+	* **Allow Negative Emissive Color:** True
+
+	[REGION:fullwidth]
+	[INCLUDE:Engine\Rendering\RenderTargets\BlueprintRenderTargets\HowTo\FluidSurface\2\#HeightSim]
+	[/REGION]
+
+* The **ComputeNormal** Material is used to create the normals that are needed for the water surface. Before you start to build the ComputeNormal Material, make sure the following options have been set in the **Details** panel for the ComputeNormal Material.
+	* **Shading Model:** Unlit
+	* **Allow Negative Emissive Color:** True
+
+	[REGION:fullwidth]
+	[INCLUDE:Engine\Rendering\RenderTargets\BlueprintRenderTargets\HowTo\FluidSurface\2\#CompNormal]
+	[/REGION]
+
+* The **MAT_ForceSplat** is used to determine where the impact happened and how big the impact was. Before you start to build the MAT_ForceSplat Material, make sure the following options have been set in the **Details** panel for the MAT_ForceSplat Material. 
+	
+	* **Blend Mode:** Additive
+	* **Shading Model:** Unlit
+	* **Allow Negative Emissive Color:** Checked
+
+	[REGION:fullwidth]
+	[INCLUDE:Engine\Rendering\RenderTargets\BlueprintRenderTargets\HowTo\FluidSurface\2\#ForceSplat]
+	[/REGION]
+
+	[region:note]
+	The MAT_ForceSplat Material is the same Material that was used in the [Height Field Painter](Engine\Rendering\RenderTargets\BlueprintRenderTargets\HowTo\HeightFieldPainter). If you have already created this Material, feel free to use that one. 
+	[/region]
+
+
+With the necessary assets now setup, in the next section, we will take a look at how each piece fits togther to create the Fluid Surface Blueprint.
+[NAV]
+
+<!--
+[EXCERPT:HeightSim]	
+[OBJECT:Snippet]
+	[PARAMLITERAL:id]
+		HeightSim
+	[/PARAMLITERAL]
+	[PARAMLITERAL:type]
+		material
+	[/PARAMLITERAL]
+	[PARAM:image]
+		![](Engine\Rendering\RenderTargets\BlueprintRenderTargets\HowTo\FluidSurface\FS_HeighSim_Material.png)
+	[/PARAM]
+	[PARAM:snippet]
+			BEGIN OBJECT
+				Begin Object Class=MaterialGraphNode_Root Name="MaterialGraphNode_Root_1"
+			   Material=PreviewMaterial'/Engine/Transient.HeightSim'
+			   NodePosX=-304
+			   NodePosY=-304
+			   NodeGuid=77F3776842896D24D55857819143B4CB
+			   CustomProperties Pin (PinId=D51F09A14DBE08D90ABDB4812F0CAD43,PinName="Base Color",PinType.PinCategory="materialinput",PinType.PinSubCategoryObject=None,PinType.PinSubCategoryMemberReference=(),PinType.bIsArray=False,PinType.bIsReference=False,PinType.bIsConst=False,PinType.bIsWeakPointer=False,PersistentGuid=00000000000000000000000000000000,bHidden=False,bNotConnectable=False,bDefaultValueIsReadOnly=False,bDefaultValueIsIgnored=False,bAdvancedView=False,)
+			   CustomProperties Pin (PinId=86495F364FA1A8619675F3A24780B7A4,PinName="Metallic",PinType.PinCategory="materialinput",PinType.PinSubCategoryObject=None,PinType.PinSubCategoryMemberReference=(),PinType.bIsArray=False,PinType.bIsReference=False,PinType.bIsConst=False,PinType.bIsWeakPointer=False,PersistentGuid=00000000000000000000000000000000,bHidden=False,bNotConnectable=False,bDefaultValueIsReadOnly=False,bDefaultValueIsIgnored=False,bAdvancedView=False,)
+			   CustomProperties Pin (PinId=E87E0F604B34A2865A807EA0803996F7,PinName="Specular",PinType.PinCategory="materialinput",PinType.PinSubCategoryObject=None,PinType.PinSubCategoryMemberReference=(),PinType.bIsArray=False,PinType.bIsReference=False,PinType.bIsConst=False,PinType.bIsWeakPointer=False,PersistentGuid=00000000000000000000000000000000,bHidden=False,bNotConnectable=False,bDefaultValueIsReadOnly=False,bDefaultValueIsIgnored=False,bAdvancedView=False,)
+			   CustomProperties Pin (PinId=14BF8A754E42F38ACA102AB81014255C,PinName="Roughness",PinType.PinCategory="materialinput",PinType.PinSubCategoryObject=None,PinType.PinSubCategoryMemberReference=(),PinType.bIsArray=False,PinType.bIsReference=False,PinType.bIsConst=False,PinType.bIsWeakPointer=False,PersistentGuid=00000000000000000000000000000000,bHidden=False,bNotConnectable=False,bDefaultValueIsReadOnly=False,bDefaultValueIsIgnored=False,bAdvancedView=False,)
+			   CustomProperties Pin (PinId=E16A4C4C4DB979FBBCB33EB6B0A0E991,PinName="Emissive Color",PinType.PinCategory="materialinput",PinType.PinSubCategoryObject=None,PinType.PinSubCategoryMemberReference=(),PinType.bIsArray=False,PinType.bIsReference=False,PinType.bIsConst=False,PinType.bIsWeakPointer=False,LinkedTo=(MaterialGraphNode_12 1F730C404BA9D857C3E8509FA498510F,),PersistentGuid=00000000000000000000000000000000,bHidden=False,bNotConnectable=False,bDefaultValueIsReadOnly=False,bDefaultValueIsIgnored=False,bAdvancedView=False,)
+			   CustomProperties Pin (PinId=5F138F3646FBA6048298FFA2180B6526,PinName="Opacity",PinType.PinCategory="materialinput",PinType.PinSubCategoryObject=None,PinType.PinSubCategoryMemberReference=(),PinType.bIsArray=False,PinType.bIsReference=False,PinType.bIsConst=False,PinType.bIsWeakPointer=False,PersistentGuid=00000000000000000000000000000000,bHidden=False,bNotConnectable=False,bDefaultValueIsReadOnly=False,bDefaultValueIsIgnored=False,bAdvancedView=False,)
+			   CustomProperties Pin (PinId=48FDE6BE449FCDC5C127CCABC30F65BA,PinName="Opacity Mask",PinType.PinCategory="materialinput",PinType.PinSubCategoryObject=None,PinType.PinSubCategoryMemberReference=(),PinType.bIsArray=False,PinType.bIsReference=False,PinType.bIsConst=False,PinType.bIsWeakPointer=False,PersistentGuid=00000000000000000000000000000000,bHidden=False,bNotConnectable=False,bDefaultValueIsReadOnly=False,bDefaultValueIsIgnored=False,bAdvancedView=False,)
+			   CustomProperties Pin (PinId=20F8F83F4934339D93197D9977B4CF56,PinName="Normal",PinType.PinCategory="materialinput",PinType.PinSubCategoryObject=None,PinType.PinSubCategoryMemberReference=(),PinType.bIsArray=False,PinType.bIsReference=False,PinType.bIsConst=False,PinType.bIsWeakPointer=False,PersistentGuid=00000000000000000000000000000000,bHidden=False,bNotConnectable=False,bDefaultValueIsReadOnly=False,bDefaultValueIsIgnored=False,bAdvancedView=False,)
+			   CustomProperties Pin (PinId=C2E6D7CF495259C59E3CD38FB3A82FC2,PinName="World Position Offset",PinType.PinCategory="materialinput",PinType.PinSubCategoryObject=None,PinType.PinSubCategoryMemberReference=(),PinType.bIsArray=False,PinType.bIsReference=False,PinType.bIsConst=False,PinType.bIsWeakPointer=False,PersistentGuid=00000000000000000000000000000000,bHidden=False,bNotConnectable=False,bDefaultValueIsReadOnly=False,bDefaultValueIsIgnored=False,bAdvancedView=False,)
+			   CustomProperties Pin (PinId=48FD5C704AAA8F3B5AE2C180EC28937F,PinName="World Displacement",PinType.PinCategory="materialinput",PinType.PinSubCategoryObject=None,PinType.PinSubCategoryMemberReference=(),PinType.bIsArray=False,PinType.bIsReference=False,PinType.bIsConst=False,PinType.bIsWeakPointer=False,PersistentGuid=00000000000000000000000000000000,bHidden=False,bNotConnectable=False,bDefaultValueIsReadOnly=False,bDefaultValueIsIgnored=False,bAdvancedView=False,)
+			   CustomProperties Pin (PinId=F7093B304E35B45DDBAE6091E54DE7CE,PinName="Tessellation Multiplier",PinType.PinCategory="materialinput",PinType.PinSubCategoryObject=None,PinType.PinSubCategoryMemberReference=(),PinType.bIsArray=False,PinType.bIsReference=False,PinType.bIsConst=False,PinType.bIsWeakPointer=False,PersistentGuid=00000000000000000000000000000000,bHidden=False,bNotConnectable=False,bDefaultValueIsReadOnly=False,bDefaultValueIsIgnored=False,bAdvancedView=False,)
+			   CustomProperties Pin (PinId=B51386BE4E4661D2C244C494BEC2DBAD,PinName="Subsurface Color",PinType.PinCategory="materialinput",PinType.PinSubCategoryObject=None,PinType.PinSubCategoryMemberReference=(),PinType.bIsArray=False,PinType.bIsReference=False,PinType.bIsConst=False,PinType.bIsWeakPointer=False,PersistentGuid=00000000000000000000000000000000,bHidden=False,bNotConnectable=False,bDefaultValueIsReadOnly=False,bDefaultValueIsIgnored=False,bAdvancedView=False,)
+			   CustomProperties Pin (PinId=BE018B43484F3B1EC0AB69B33CDD9432,PinName="Custom Data 0",PinType.PinCategory="materialinput",PinType.PinSubCategoryObject=None,PinType.PinSubCategoryMemberReference=(),PinType.bIsArray=False,PinType.bIsReference=False,PinType.bIsConst=False,PinType.bIsWeakPointer=False,PersistentGuid=00000000000000000000000000000000,bHidden=False,bNotConnectable=False,bDefaultValueIsReadOnly=False,bDefaultValueIsIgnored=False,bAdvancedView=False,)
+			   CustomProperties Pin (PinId=3565A9404AB8428EDC65289AA47E7E5C,PinName="Custom Data 1",PinType.PinCategory="materialinput",PinType.PinSubCategoryObject=None,PinType.PinSubCategoryMemberReference=(),PinType.bIsArray=False,PinType.bIsReference=False,PinType.bIsConst=False,PinType.bIsWeakPointer=False,PersistentGuid=00000000000000000000000000000000,bHidden=False,bNotConnectable=False,bDefaultValueIsReadOnly=False,bDefaultValueIsIgnored=False,bAdvancedView=False,)
+			   CustomProperties Pin (PinId=F1FC09194ECCD9CDB8B93090FC97B0B4,PinName="Ambient Occlusion",PinType.PinCategory="materialinput",PinType.PinSubCategoryObject=None,PinType.PinSubCategoryMemberReference=(),PinType.bIsArray=False,PinType.bIsReference=False,PinType.bIsConst=False,PinType.bIsWeakPointer=False,PersistentGuid=00000000000000000000000000000000,bHidden=False,bNotConnectable=False,bDefaultValueIsReadOnly=False,bDefaultValueIsIgnored=False,bAdvancedView=False,)
+			   CustomProperties Pin (PinId=BA7BE54941EC239CD9C5F2884423002F,PinName="Refraction",PinType.PinCategory="materialinput",PinType.PinSubCategoryObject=None,PinType.PinSubCategoryMemberReference=(),PinType.bIsArray=False,PinType.bIsReference=False,PinType.bIsConst=False,PinType.bIsWeakPointer=False,PersistentGuid=00000000000000000000000000000000,bHidden=False,bNotConnectable=False,bDefaultValueIsReadOnly=False,bDefaultValueIsIgnored=False,bAdvancedView=False,)
+			   CustomProperties Pin (PinId=0AEF8D3947B11D6361213FBD9F4CA86C,PinName="Customized UV0",PinType.PinCategory="materialinput",PinType.PinSubCategoryObject=None,PinType.PinSubCategoryMemberReference=(),PinType.bIsArray=False,PinType.bIsReference=False,PinType.bIsConst=False,PinType.bIsWeakPointer=False,PersistentGuid=00000000000000000000000000000000,bHidden=False,bNotConnectable=False,bDefaultValueIsReadOnly=False,bDefaultValueIsIgnored=False,bAdvancedView=False,)
+			   CustomProperties Pin (PinId=B58A5A1848C3D80F93F8DA8DF3399E1A,PinName="Customized UV1",PinType.PinCategory="materialinput",PinType.PinSubCategoryObject=None,PinType.PinSubCategoryMemberReference=(),PinType.bIsArray=False,PinType.bIsReference=False,PinType.bIsConst=False,PinType.bIsWeakPointer=False,PersistentGuid=00000000000000000000000000000000,bHidden=False,bNotConnectable=False,bDefaultValueIsReadOnly=False,bDefaultValueIsIgnored=False,bAdvancedView=False,)
+			   CustomProperties Pin (PinId=B1D2951146D7141BB9CAD3BA0FD83B5C,PinName="Customized UV2",PinType.PinCategory="materialinput",PinType.PinSubCategoryObject=None,PinType.PinSubCategoryMemberReference=(),PinType.bIsArray=False,PinType.bIsReference=False,PinType.bIsConst=False,PinType.bIsWeakPointer=False,PersistentGuid=00000000000000000000000000000000,bHidden=False,bNotConnectable=False,bDefaultValueIsReadOnly=False,bDefaultValueIsIgnored=False,bAdvancedView=False,)
+			   CustomProperties Pin (PinId=B910AB3345C9EEEB2B318AAF5C75F16E,PinName="Customized UV3",PinType.PinCategory="materialinput",PinType.PinSubCategoryObject=None,PinType.PinSubCategoryMemberReference=(),PinType.bIsArray=False,PinType.bIsReference=False,PinType.bIsConst=False,PinType.bIsWeakPointer=False,PersistentGuid=00000000000000000000000000000000,bHidden=False,bNotConnectable=False,bDefaultValueIsReadOnly=False,bDefaultValueIsIgnored=False,bAdvancedView=False,)
+			   CustomProperties Pin (PinId=7431501A4C71BF6B69C1799EE875798C,PinName="Customized UV4",PinType.PinCategory="materialinput",PinType.PinSubCategoryObject=None,PinType.PinSubCategoryMemberReference=(),PinType.bIsArray=False,PinType.bIsReference=False,PinType.bIsConst=False,PinType.bIsWeakPointer=False,PersistentGuid=00000000000000000000000000000000,bHidden=False,bNotConnectable=False,bDefaultValueIsReadOnly=False,bDefaultValueIsIgnored=False,bAdvancedView=False,)
+			   CustomProperties Pin (PinId=8BD8DC324185CEE003B57498510DE6D7,PinName="Customized UV5",PinType.PinCategory="materialinput",PinType.PinSubCategoryObject=None,PinType.PinSubCategoryMemberReference=(),PinType.bIsArray=False,PinType.bIsReference=False,PinType.bIsConst=False,PinType.bIsWeakPointer=False,PersistentGuid=00000000000000000000000000000000,bHidden=False,bNotConnectable=False,bDefaultValueIsReadOnly=False,bDefaultValueIsIgnored=False,bAdvancedView=False,)
+			   CustomProperties Pin (PinId=99632802424775A45C05599CAB9DD168,PinName="Customized UV6",PinType.PinCategory="materialinput",PinType.PinSubCategoryObject=None,PinType.PinSubCategoryMemberReference=(),PinType.bIsArray=False,PinType.bIsReference=False,PinType.bIsConst=False,PinType.bIsWeakPointer=False,PersistentGuid=00000000000000000000000000000000,bHidden=False,bNotConnectable=False,bDefaultValueIsReadOnly=False,bDefaultValueIsIgnored=False,bAdvancedView=False,)
+			   CustomProperties Pin (PinId=566D5E344FFF6E9CEF61C3BB2F584BB8,PinName="Customized UV7",PinType.PinCategory="materialinput",PinType.PinSubCategoryObject=None,PinType.PinSubCategoryMemberReference=(),PinType.bIsArray=False,PinType.bIsReference=False,PinType.bIsConst=False,PinType.bIsWeakPointer=False,PersistentGuid=00000000000000000000000000000000,bHidden=False,bNotConnectable=False,bDefaultValueIsReadOnly=False,bDefaultValueIsIgnored=False,bAdvancedView=False,)
+			   CustomProperties Pin (PinId=63B3CF6945275D470B8F56888CD0CAF4,PinName="Pixel Depth Offset",PinType.PinCategory="materialinput",PinType.PinSubCategoryObject=None,PinType.PinSubCategoryMemberReference=(),PinType.bIsArray=False,PinType.bIsReference=False,PinType.bIsConst=False,PinType.bIsWeakPointer=False,PersistentGuid=00000000000000000000000000000000,bHidden=False,bNotConnectable=False,bDefaultValueIsReadOnly=False,bDefaultValueIsIgnored=False,bAdvancedView=False,)
+			   CustomProperties Pin (PinId=A783538041DD38C2E5CB96980060D4F3,PinName="Material Attributes",PinType.PinCategory="materialinput",PinType.PinSubCategoryObject=None,PinType.PinSubCategoryMemberReference=(),PinType.bIsArray=False,PinType.bIsReference=False,PinType.bIsConst=False,PinType.bIsWeakPointer=False,PersistentGuid=00000000000000000000000000000000,bHidden=False,bNotConnectable=False,bDefaultValueIsReadOnly=False,bDefaultValueIsIgnored=False,bAdvancedView=False,)
+			End Object
+			Begin Object Class=MaterialGraphNode Name="MaterialGraphNode_12"
+			   Begin Object Class=MaterialExpressionMultiply Name="MaterialExpressionMultiply_7"
+			   End Object
+			   Begin Object Name="MaterialExpressionMultiply_7"
+				  A=(Expression=MaterialExpressionSubtract'MaterialGraphNode_47.MaterialExpressionSubtract_1')
+				  B=(Expression=MaterialExpressionScalarParameter'MaterialGraphNode_19.MaterialExpressionScalarParameter_22')
+				  MaterialExpressionEditorX=-448
+				  MaterialExpressionEditorY=-208
+				  MaterialExpressionGuid=9AD2E13E4FD4F85F15A01A9ED2719AF0
+				  Material=PreviewMaterial'/Engine/Transient.HeightSim'
+			   End Object
+			   MaterialExpression=MaterialExpressionMultiply'MaterialExpressionMultiply_7'
+			   NodePosX=-448
+			   NodePosY=-208
+			   NodeGuid=3DD0E174484FA23BB94884AFA6190E0B
+			   CustomProperties Pin (PinId=AE4E493242D5E3BA3C1D83A2BF5D2295,PinName="A",PinType.PinCategory="optional",PinType.PinSubCategoryObject=None,PinType.PinSubCategoryMemberReference=(),PinType.bIsArray=False,PinType.bIsReference=False,PinType.bIsConst=False,PinType.bIsWeakPointer=False,LinkedTo=(MaterialGraphNode_47 297B10D944D15582C2242E9387BD8FA9,),PersistentGuid=00000000000000000000000000000000,bHidden=False,bNotConnectable=False,bDefaultValueIsReadOnly=False,bDefaultValueIsIgnored=False,bAdvancedView=False,)
+			   CustomProperties Pin (PinId=FE74BD24456C31D437C95A8FBF3CB635,PinName="B",PinType.PinCategory="optional",PinType.PinSubCategoryObject=None,PinType.PinSubCategoryMemberReference=(),PinType.bIsArray=False,PinType.bIsReference=False,PinType.bIsConst=False,PinType.bIsWeakPointer=False,LinkedTo=(MaterialGraphNode_19 AE27F1A44FAB642957FFCEB73F56A6DE,),PersistentGuid=00000000000000000000000000000000,bHidden=False,bNotConnectable=False,bDefaultValueIsReadOnly=False,bDefaultValueIsIgnored=False,bAdvancedView=False,)
+			   CustomProperties Pin (PinId=1F730C404BA9D857C3E8509FA498510F,PinName="Output",PinFriendlyName=" ",Direction="EGPD_Output",PinType.PinSubCategoryObject=None,PinType.PinSubCategoryMemberReference=(),PinType.bIsArray=False,PinType.bIsReference=False,PinType.bIsConst=False,PinType.bIsWeakPointer=False,LinkedTo=(MaterialGraphNode_Root_1 E16A4C4C4DB979FBBCB33EB6B0A0E991,),PersistentGuid=00000000000000000000000000000000,bHidden=False,bNotConnectable=False,bDefaultValueIsReadOnly=False,bDefaultValueIsIgnored=False,bAdvancedView=False,)
+			End Object
+			Begin Object Class=MaterialGraphNode Name="MaterialGraphNode_14"
+			   Begin Object Class=MaterialExpressionMultiply Name="MaterialExpressionMultiply_8"
+			   End Object
+			   Begin Object Name="MaterialExpressionMultiply_8"
+				  A=(Expression=MaterialExpressionScalarParameter'MaterialGraphNode_18.MaterialExpressionScalarParameter_21')
+				  B=(Expression=MaterialExpressionConstant2Vector'MaterialGraphNode_32.MaterialExpressionConstant2Vector_0')
+				  MaterialExpressionEditorX=-2144
+				  MaterialExpressionEditorY=320
+				  MaterialExpressionGuid=9BB2CF3945F3869BD091C7AFFEBBD2A1
+				  Material=PreviewMaterial'/Engine/Transient.HeightSim'
+			   End Object
+			   MaterialExpression=MaterialExpressionMultiply'MaterialExpressionMultiply_8'
+			   NodePosX=-2144
+			   NodePosY=320
+			   NodeGuid=CCE7000C4721D1FBA1C4F6A06053754D
+			   CustomProperties Pin (PinId=CBCD523B4B9AFBD0B35332A0D4FD0D2F,PinName="A",PinType.PinCategory="optional",PinType.PinSubCategoryObject=None,PinType.PinSubCategoryMemberReference=(),PinType.bIsArray=False,PinType.bIsReference=False,PinType.bIsConst=False,PinType.bIsWeakPointer=False,LinkedTo=(MaterialGraphNode_18 8149D11640C4A85F21F77F826F1AA7C6,),PersistentGuid=00000000000000000000000000000000,bHidden=False,bNotConnectable=False,bDefaultValueIsReadOnly=False,bDefaultValueIsIgnored=False,bAdvancedView=False,)
+			   CustomProperties Pin (PinId=B243F91847B4DFA8816E3F8485A06DE3,PinName="B",PinType.PinCategory="optional",PinType.PinSubCategoryObject=None,PinType.PinSubCategoryMemberReference=(),PinType.bIsArray=False,PinType.bIsReference=False,PinType.bIsConst=False,PinType.bIsWeakPointer=False,LinkedTo=(MaterialGraphNode_32 0C3200D34E5DE557B3B7D38DA4AA3369,),PersistentGuid=00000000000000000000000000000000,bHidden=False,bNotConnectable=False,bDefaultValueIsReadOnly=False,bDefaultValueIsIgnored=False,bAdvancedView=False,)
+			   CustomProperties Pin (PinId=5E9617AB4A810D7A05D2D2A1AECCAC21,PinName="Output",PinFriendlyName=" ",Direction="EGPD_Output",PinType.PinSubCategoryObject=None,PinType.PinSubCategoryMemberReference=(),PinType.bIsArray=False,PinType.bIsReference=False,PinType.bIsConst=False,PinType.bIsWeakPointer=False,LinkedTo=(MaterialGraphNode_30 7ECB60B641D4CC0CC8A0D1ADFAC26BB2,),PersistentGuid=00000000000000000000000000000000,bHidden=False,bNotConnectable=False,bDefaultValueIsReadOnly=False,bDefaultValueIsIgnored=False,bAdvancedView=False,)
+			End Object
+			Begin Object Class=MaterialGraphNode Name="MaterialGraphNode_15"
+			   Begin Object Class=MaterialExpressionMultiply Name="MaterialExpressionMultiply_9"
+			   End Object
+			   Begin Object Name="MaterialExpressionMultiply_9"
+				  A=(Expression=MaterialExpressionScalarParameter'MaterialGraphNode_18.MaterialExpressionScalarParameter_21')
+				  B=(Expression=MaterialExpressionConstant2Vector'MaterialGraphNode_37.MaterialExpressionConstant2Vector_1')
+				  MaterialExpressionEditorX=-2096
+				  MaterialExpressionEditorY=80
+				  MaterialExpressionGuid=9BB2CF3945F3869BD091C7AFFEBBD2A1
+				  Material=PreviewMaterial'/Engine/Transient.HeightSim'
+			   End Object
+			   MaterialExpression=MaterialExpressionMultiply'MaterialExpressionMultiply_9'
+			   NodePosX=-2096
+			   NodePosY=80
+			   NodeGuid=3BD2CF1C470FCE727444868AA6420FF7
+			   CustomProperties Pin (PinId=D94C9E9C47A9CE4458FF71ABB26FD13B,PinName="A",PinType.PinCategory="optional",PinType.PinSubCategoryObject=None,PinType.PinSubCategoryMemberReference=(),PinType.bIsArray=False,PinType.bIsReference=False,PinType.bIsConst=False,PinType.bIsWeakPointer=False,LinkedTo=(MaterialGraphNode_18 8149D11640C4A85F21F77F826F1AA7C6,),PersistentGuid=00000000000000000000000000000000,bHidden=False,bNotConnectable=False,bDefaultValueIsReadOnly=False,bDefaultValueIsIgnored=False,bAdvancedView=False,)
+			   CustomProperties Pin (PinId=59E2DAA140B7A021018B8980E6423224,PinName="B",PinType.PinCategory="optional",PinType.PinSubCategoryObject=None,PinType.PinSubCategoryMemberReference=(),PinType.bIsArray=False,PinType.bIsReference=False,PinType.bIsConst=False,PinType.bIsWeakPointer=False,LinkedTo=(MaterialGraphNode_37 9AD9001E4BB978F2279A2E8107F639E1,),PersistentGuid=00000000000000000000000000000000,bHidden=False,bNotConnectable=False,bDefaultValueIsReadOnly=False,bDefaultValueIsIgnored=False,bAdvancedView=False,)
+			   CustomProperties Pin (PinId=9FC9244144840A564AEA2BA8E1D244B7,PinName="Output",PinFriendlyName=" ",Direction="EGPD_Output",PinType.PinSubCategoryObject=None,PinType.PinSubCategoryMemberReference=(),PinType.bIsArray=False,PinType.bIsReference=False,PinType.bIsConst=False,PinType.bIsWeakPointer=False,LinkedTo=(MaterialGraphNode_35 1576FA7A426EF5DB948BBCA5C72B5288,),PersistentGuid=00000000000000000000000000000000,bHidden=False,bNotConnectable=False,bDefaultValueIsReadOnly=False,bDefaultValueIsIgnored=False,bAdvancedView=False,)
+			End Object
+			Begin Object Class=MaterialGraphNode Name="MaterialGraphNode_16"
+			   Begin Object Class=MaterialExpressionMultiply Name="MaterialExpressionMultiply_10"
+			   End Object
+			   Begin Object Name="MaterialExpressionMultiply_10"
+				  A=(Expression=MaterialExpressionScalarParameter'MaterialGraphNode_18.MaterialExpressionScalarParameter_21')
+				  B=(Expression=MaterialExpressionConstant2Vector'MaterialGraphNode_41.MaterialExpressionConstant2Vector_2')
+				  MaterialExpressionEditorX=-2096
+				  MaterialExpressionEditorY=-192
+				  MaterialExpressionGuid=9BB2CF3945F3869BD091C7AFFEBBD2A1
+				  Material=PreviewMaterial'/Engine/Transient.HeightSim'
+			   End Object
+			   MaterialExpression=MaterialExpressionMultiply'MaterialExpressionMultiply_10'
+			   NodePosX=-2096
+			   NodePosY=-192
+			   NodeGuid=237306A949A7AF6C97C7CE9C5F7A3748
+			   CustomProperties Pin (PinId=01F9DA1E4D2696EEA6B442987815B787,PinName="A",PinType.PinCategory="optional",PinType.PinSubCategoryObject=None,PinType.PinSubCategoryMemberReference=(),PinType.bIsArray=False,PinType.bIsReference=False,PinType.bIsConst=False,PinType.bIsWeakPointer=False,LinkedTo=(MaterialGraphNode_18 8149D11640C4A85F21F77F826F1AA7C6,),PersistentGuid=00000000000000000000000000000000,bHidden=False,bNotConnectable=False,bDefaultValueIsReadOnly=False,bDefaultValueIsIgnored=False,bAdvancedView=False,)
+			   CustomProperties Pin (PinId=8C0F76D241F89C6DCA6B3E9FD3BA69CD,PinName="B",PinType.PinCategory="optional",PinType.PinSubCategoryObject=None,PinType.PinSubCategoryMemberReference=(),PinType.bIsArray=False,PinType.bIsReference=False,PinType.bIsConst=False,PinType.bIsWeakPointer=False,LinkedTo=(MaterialGraphNode_41 B72141BD46005B6680C210816BEEB6B1,),PersistentGuid=00000000000000000000000000000000,bHidden=False,bNotConnectable=False,bDefaultValueIsReadOnly=False,bDefaultValueIsIgnored=False,bAdvancedView=False,)
+			   CustomProperties Pin (PinId=DEB3772B4736834DF22CAFA5F034F702,PinName="Output",PinFriendlyName=" ",Direction="EGPD_Output",PinType.PinSubCategoryObject=None,PinType.PinSubCategoryMemberReference=(),PinType.bIsArray=False,PinType.bIsReference=False,PinType.bIsConst=False,PinType.bIsWeakPointer=False,LinkedTo=(MaterialGraphNode_39 40DFB4A047C1871E144694A9EF25DBE5,),PersistentGuid=00000000000000000000000000000000,bHidden=False,bNotConnectable=False,bDefaultValueIsReadOnly=False,bDefaultValueIsIgnored=False,bAdvancedView=False,)
+			End Object
+			Begin Object Class=MaterialGraphNode Name="MaterialGraphNode_17"
+			   Begin Object Class=MaterialExpressionMultiply Name="MaterialExpressionMultiply_11"
+			   End Object
+			   Begin Object Name="MaterialExpressionMultiply_11"
+				  A=(Expression=MaterialExpressionScalarParameter'MaterialGraphNode_18.MaterialExpressionScalarParameter_21')
+				  B=(Expression=MaterialExpressionConstant2Vector'MaterialGraphNode_45.MaterialExpressionConstant2Vector_3')
+				  MaterialExpressionEditorX=-2112
+				  MaterialExpressionEditorY=-464
+				  MaterialExpressionGuid=9BB2CF3945F3869BD091C7AFFEBBD2A1
+				  Material=PreviewMaterial'/Engine/Transient.HeightSim'
+			   End Object
+			   MaterialExpression=MaterialExpressionMultiply'MaterialExpressionMultiply_11'
+			   NodePosX=-2112
+			   NodePosY=-464
+			   NodeGuid=8CAC45BE4B5B8200A674DDBA09731B49
+			   CustomProperties Pin (PinId=B9E82F5D4815925C2C4C7BAC9E32C12A,PinName="A",PinType.PinCategory="optional",PinType.PinSubCategoryObject=None,PinType.PinSubCategoryMemberReference=(),PinType.bIsArray=False,PinType.bIsReference=False,PinType.bIsConst=False,PinType.bIsWeakPointer=False,LinkedTo=(MaterialGraphNode_18 8149D11640C4A85F21F77F826F1AA7C6,),PersistentGuid=00000000000000000000000000000000,bHidden=False,bNotConnectable=False,bDefaultValueIsReadOnly=False,bDefaultValueIsIgnored=False,bAdvancedView=False,)
+			   CustomProperties Pin (PinId=9E785A3944293393B058D6A68927A4CF,PinName="B",PinType.PinCategory="optional",PinType.PinSubCategoryObject=None,PinType.PinSubCategoryMemberReference=(),PinType.bIsArray=False,PinType.bIsReference=False,PinType.bIsConst=False,PinType.bIsWeakPointer=False,LinkedTo=(MaterialGraphNode_45 150FAB2F48D1DD64E87EC4AFA4C3D9EC,),PersistentGuid=00000000000000000000000000000000,bHidden=False,bNotConnectable=False,bDefaultValueIsReadOnly=False,bDefaultValueIsIgnored=False,bAdvancedView=False,)
+			   CustomProperties Pin (PinId=573CDBF54E19966A3470138C02D5DBF8,PinName="Output",PinFriendlyName=" ",Direction="EGPD_Output",PinType.PinSubCategoryObject=None,PinType.PinSubCategoryMemberReference=(),PinType.bIsArray=False,PinType.bIsReference=False,PinType.bIsConst=False,PinType.bIsWeakPointer=False,LinkedTo=(MaterialGraphNode_43 E936B7EF4C08A3F0F256FF95563C7509,),PersistentGuid=00000000000000000000000000000000,bHidden=False,bNotConnectable=False,bDefaultValueIsReadOnly=False,bDefaultValueIsIgnored=False,bAdvancedView=False,)
+			End Object
+			Begin Object Class=MaterialGraphNode Name="MaterialGraphNode_18"
+			   Begin Object Class=MaterialExpressionScalarParameter Name="MaterialExpressionScalarParameter_21"
+			   End Object
+			   Begin Object Name="MaterialExpressionScalarParameter_21"
+				  DefaultValue=0.003906
+				  ParameterName="TexelSize"
+				  ExpressionGUID=AD3899C04C5810DD5A9DFCA8BFFCEE67
+				  MaterialExpressionEditorX=-2480
+				  MaterialExpressionEditorY=-208
+				  MaterialExpressionGuid=DB480CAC4000EF2A7EB61BA94FE23AF6
+				  Material=PreviewMaterial'/Engine/Transient.HeightSim'
+			   End Object
+			   MaterialExpression=MaterialExpressionScalarParameter'MaterialExpressionScalarParameter_21'
+			   NodePosX=-2480
+			   NodePosY=-208
+			   bCanRenameNode=True
+			   NodeGuid=E69FFD0442F14069684A2FBFCA7E656E
+			   CustomProperties Pin (PinId=8149D11640C4A85F21F77F826F1AA7C6,PinName="Output",PinFriendlyName=" ",Direction="EGPD_Output",PinType.PinSubCategoryObject=None,PinType.PinSubCategoryMemberReference=(),PinType.bIsArray=False,PinType.bIsReference=False,PinType.bIsConst=False,PinType.bIsWeakPointer=False,LinkedTo=(MaterialGraphNode_14 CBCD523B4B9AFBD0B35332A0D4FD0D2F,MaterialGraphNode_15 D94C9E9C47A9CE4458FF71ABB26FD13B,MaterialGraphNode_16 01F9DA1E4D2696EEA6B442987815B787,MaterialGraphNode_17 B9E82F5D4815925C2C4C7BAC9E32C12A,),PersistentGuid=00000000000000000000000000000000,bHidden=False,bNotConnectable=False,bDefaultValueIsReadOnly=False,bDefaultValueIsIgnored=False,bAdvancedView=False,)
+			End Object
+			Begin Object Class=MaterialGraphNode Name="MaterialGraphNode_19"
+			   Begin Object Class=MaterialExpressionScalarParameter Name="MaterialExpressionScalarParameter_22"
+			   End Object
+			   Begin Object Name="MaterialExpressionScalarParameter_22"
+				  DefaultValue=0.980000
+				  ParameterName="Dampening"
+				  ExpressionGUID=4A7EBA704E603F9148931A8772BE8719
+				  MaterialExpressionEditorX=-592
+				  MaterialExpressionEditorY=-80
+				  MaterialExpressionGuid=7A931EA244566C9B0EFCBCBDBD92ECF0
+				  Material=PreviewMaterial'/Engine/Transient.HeightSim'
+			   End Object
+			   MaterialExpression=MaterialExpressionScalarParameter'MaterialExpressionScalarParameter_22'
+			   NodePosX=-592
+			   NodePosY=-80
+			   bCanRenameNode=True
+			   NodeGuid=77BD648F476FC57ECF61E1BDBE2718FA
+			   CustomProperties Pin (PinId=AE27F1A44FAB642957FFCEB73F56A6DE,PinName="Output",PinFriendlyName=" ",Direction="EGPD_Output",PinType.PinSubCategoryObject=None,PinType.PinSubCategoryMemberReference=(),PinType.bIsArray=False,PinType.bIsReference=False,PinType.bIsConst=False,PinType.bIsWeakPointer=False,LinkedTo=(MaterialGraphNode_12 FE74BD24456C31D437C95A8FBF3CB635,),PersistentGuid=00000000000000000000000000000000,bHidden=False,bNotConnectable=False,bDefaultValueIsReadOnly=False,bDefaultValueIsIgnored=False,bAdvancedView=False,)
+			End Object
+			Begin Object Class=MaterialGraphNode Name="MaterialGraphNode_20"
+			   Begin Object Class=MaterialExpressionTextureSampleParameter2D Name="MaterialExpressionTextureSampleParameter2D_14"
+			   End Object
+			   Begin Object Name="MaterialExpressionTextureSampleParameter2D_14"
+				  ParameterName="PreviousHeight2"
+				  ExpressionGUID=88FAFFA443D25A267841F9B2A8F0723C
+				  MaterialExpressionEditorX=-768
+				  MaterialExpressionEditorY=-96
+				  MaterialExpressionGuid=8F3166474EF9A6F6A6825185B78360BE
+				  Material=PreviewMaterial'/Engine/Transient.HeightSim'
+			   End Object
+			   MaterialExpression=MaterialExpressionTextureSampleParameter2D'MaterialExpressionTextureSampleParameter2D_14'
+			   NodePosX=-768
+			   NodePosY=-96
+			   bCanRenameNode=True
+			   NodeGuid=288F1E58435C2B310739908EC8347906
+			   CustomProperties Pin (PinId=221DBE154ABAC46D9CA1BF8DEA05E457,PinName="UVs",PinType.PinCategory="optional",PinType.PinSubCategoryObject=None,PinType.PinSubCategoryMemberReference=(),PinType.bIsArray=False,PinType.bIsReference=False,PinType.bIsConst=False,PinType.bIsWeakPointer=False,PersistentGuid=00000000000000000000000000000000,bHidden=False,bNotConnectable=False,bDefaultValueIsReadOnly=False,bDefaultValueIsIgnored=False,bAdvancedView=False,)
+			   CustomProperties Pin (PinId=9E7C587F42626A9947903B92D74E95FF,PinName="Tex",PinType.PinCategory="optional",PinType.PinSubCategoryObject=None,PinType.PinSubCategoryMemberReference=(),PinType.bIsArray=False,PinType.bIsReference=False,PinType.bIsConst=False,PinType.bIsWeakPointer=False,PersistentGuid=00000000000000000000000000000000,bHidden=False,bNotConnectable=False,bDefaultValueIsReadOnly=False,bDefaultValueIsIgnored=False,bAdvancedView=False,)
+			   CustomProperties Pin (PinId=E664394C40456701DAA179A68A2F5BB6,PinName="Output",PinFriendlyName=" ",Direction="EGPD_Output",PinType.PinCategory="mask",PinType.PinSubCategoryObject=None,PinType.PinSubCategoryMemberReference=(),PinType.bIsArray=False,PinType.bIsReference=False,PinType.bIsConst=False,PinType.bIsWeakPointer=False,PersistentGuid=00000000000000000000000000000000,bHidden=False,bNotConnectable=False,bDefaultValueIsReadOnly=False,bDefaultValueIsIgnored=False,bAdvancedView=False,)
+			   CustomProperties Pin (PinId=9E1D6AE64F3486451D911EA6B84ECC06,PinName="Output2",PinFriendlyName=" ",Direction="EGPD_Output",PinType.PinCategory="mask",PinType.PinSubCategory="red",PinType.PinSubCategoryObject=None,PinType.PinSubCategoryMemberReference=(),PinType.bIsArray=False,PinType.bIsReference=False,PinType.bIsConst=False,PinType.bIsWeakPointer=False,LinkedTo=(MaterialGraphNode_47 C91C84A64517D7F47FA7B39B116852A6,),PersistentGuid=00000000000000000000000000000000,bHidden=False,bNotConnectable=False,bDefaultValueIsReadOnly=False,bDefaultValueIsIgnored=False,bAdvancedView=False,)
+			   CustomProperties Pin (PinId=FD9ABED04022E27DD1CCF1B6E112386D,PinName="Output3",PinFriendlyName=" ",Direction="EGPD_Output",PinType.PinCategory="mask",PinType.PinSubCategory="green",PinType.PinSubCategoryObject=None,PinType.PinSubCategoryMemberReference=(),PinType.bIsArray=False,PinType.bIsReference=False,PinType.bIsConst=False,PinType.bIsWeakPointer=False,PersistentGuid=00000000000000000000000000000000,bHidden=False,bNotConnectable=False,bDefaultValueIsReadOnly=False,bDefaultValueIsIgnored=False,bAdvancedView=False,)
+			   CustomProperties Pin (PinId=8F4C368244F90DE0F25A8F854FFCBC20,PinName="Output4",PinFriendlyName=" ",Direction="EGPD_Output",PinType.PinCategory="mask",PinType.PinSubCategory="blue",PinType.PinSubCategoryObject=None,PinType.PinSubCategoryMemberReference=(),PinType.bIsArray=False,PinType.bIsReference=False,PinType.bIsConst=False,PinType.bIsWeakPointer=False,PersistentGuid=00000000000000000000000000000000,bHidden=False,bNotConnectable=False,bDefaultValueIsReadOnly=False,bDefaultValueIsIgnored=False,bAdvancedView=False,)
+			   CustomProperties Pin (PinId=816D3ACF4964A0F10F9B69BEB8488B49,PinName="Output5",PinFriendlyName=" ",Direction="EGPD_Output",PinType.PinCategory="mask",PinType.PinSubCategory="alpha",PinType.PinSubCategoryObject=None,PinType.PinSubCategoryMemberReference=(),PinType.bIsArray=False,PinType.bIsReference=False,PinType.bIsConst=False,PinType.bIsWeakPointer=False,PersistentGuid=00000000000000000000000000000000,bHidden=False,bNotConnectable=False,bDefaultValueIsReadOnly=False,bDefaultValueIsIgnored=False,bAdvancedView=False,)
+			End Object
+			Begin Object Class=MaterialGraphNode Name="MaterialGraphNode_21"
+			   Begin Object Class=MaterialExpressionMultiply Name="MaterialExpressionMultiply_12"
+			   End Object
+			   Begin Object Name="MaterialExpressionMultiply_12"
+				  A=(Expression=MaterialExpressionAdd'MaterialGraphNode_24.MaterialExpressionAdd_0')
+				  ConstB=0.500000
+				  MaterialExpressionEditorX=-768
+				  MaterialExpressionEditorY=-208
+				  MaterialExpressionGuid=A2679E7D45AB6CC6E4DC589AFCAFAEB6
+				  Material=PreviewMaterial'/Engine/Transient.HeightSim'
+			   End Object
+			   MaterialExpression=MaterialExpressionMultiply'MaterialExpressionMultiply_12'
+			   NodePosX=-768
+			   NodePosY=-208
+			   NodeGuid=E802E202403CF3D0232BEE9F75DCF341
+			   CustomProperties Pin (PinId=BE3AE6364E3BA620AB73C19920383471,PinName="A",PinType.PinCategory="optional",PinType.PinSubCategoryObject=None,PinType.PinSubCategoryMemberReference=(),PinType.bIsArray=False,PinType.bIsReference=False,PinType.bIsConst=False,PinType.bIsWeakPointer=False,LinkedTo=(MaterialGraphNode_24 69A2ACEF445864C15E03A3A43F0413F8,),PersistentGuid=00000000000000000000000000000000,bHidden=False,bNotConnectable=False,bDefaultValueIsReadOnly=False,bDefaultValueIsIgnored=False,bAdvancedView=False,)
+			   CustomProperties Pin (PinId=60B995604BE1D9F8EB8BA4B75BA41FFE,PinName="B",PinType.PinCategory="optional",PinType.PinSubCategoryObject=None,PinType.PinSubCategoryMemberReference=(),PinType.bIsArray=False,PinType.bIsReference=False,PinType.bIsConst=False,PinType.bIsWeakPointer=False,PersistentGuid=00000000000000000000000000000000,bHidden=False,bNotConnectable=False,bDefaultValueIsReadOnly=False,bDefaultValueIsIgnored=False,bAdvancedView=False,)
+			   CustomProperties Pin (PinId=F7B1B2994F6F5A2E84B9EFA1485EFD4A,PinName="Output",PinFriendlyName=" ",Direction="EGPD_Output",PinType.PinSubCategoryObject=None,PinType.PinSubCategoryMemberReference=(),PinType.bIsArray=False,PinType.bIsReference=False,PinType.bIsConst=False,PinType.bIsWeakPointer=False,LinkedTo=(MaterialGraphNode_47 836C63F044D0CE95EF676390B28B3921,),PersistentGuid=00000000000000000000000000000000,bHidden=False,bNotConnectable=False,bDefaultValueIsReadOnly=False,bDefaultValueIsIgnored=False,bAdvancedView=False,)
+			End Object
+			Begin Object Class=MaterialGraphNode Name="MaterialGraphNode_22"
+			   Begin Object Class=MaterialExpressionMultiply Name="MaterialExpressionMultiply_13"
+			   End Object
+			   Begin Object Name="MaterialExpressionMultiply_13"
+				  A=(Expression=MaterialExpressionSubtract'MaterialGraphNode_23.MaterialExpressionSubtract_0')
+				  B=(Expression=MaterialExpressionScalarParameter'MaterialGraphNode_25.MaterialExpressionScalarParameter_23')
+				  MaterialExpressionEditorX=-944
+				  MaterialExpressionEditorY=16
+				  MaterialExpressionGuid=D6DBD2BB45CB73C7108F1089D6B0610D
+				  Material=PreviewMaterial'/Engine/Transient.HeightSim'
+			   End Object
+			   MaterialExpression=MaterialExpressionMultiply'MaterialExpressionMultiply_13'
+			   NodePosX=-944
+			   NodePosY=16
+			   NodeGuid=6880C81B4FFE553D8B857493162CAEB5
+			   CustomProperties Pin (PinId=0C7E3CBD4F48315F0B53FBA17EF9A98A,PinName="A",PinType.PinCategory="optional",PinType.PinSubCategoryObject=None,PinType.PinSubCategoryMemberReference=(),PinType.bIsArray=False,PinType.bIsReference=False,PinType.bIsConst=False,PinType.bIsWeakPointer=False,LinkedTo=(MaterialGraphNode_23 265235A346C581EDAE72BCA7D9BC355F,),PersistentGuid=00000000000000000000000000000000,bHidden=False,bNotConnectable=False,bDefaultValueIsReadOnly=False,bDefaultValueIsIgnored=False,bAdvancedView=False,)
+			   CustomProperties Pin (PinId=7CB7D4D94D2CB5B5F59ED998423153E7,PinName="B",PinType.PinCategory="optional",PinType.PinSubCategoryObject=None,PinType.PinSubCategoryMemberReference=(),PinType.bIsArray=False,PinType.bIsReference=False,PinType.bIsConst=False,PinType.bIsWeakPointer=False,LinkedTo=(MaterialGraphNode_25 F1D57D00491B73D6EF7215919AE4BBDB,),PersistentGuid=00000000000000000000000000000000,bHidden=False,bNotConnectable=False,bDefaultValueIsReadOnly=False,bDefaultValueIsIgnored=False,bAdvancedView=False,)
+			   CustomProperties Pin (PinId=E4B6CBFB47131CBEDDFF5E9313B2638F,PinName="Output",PinFriendlyName=" ",Direction="EGPD_Output",PinType.PinSubCategoryObject=None,PinType.PinSubCategoryMemberReference=(),PinType.bIsArray=False,PinType.bIsReference=False,PinType.bIsConst=False,PinType.bIsWeakPointer=False,LinkedTo=(MaterialGraphNode_24 35E274EA4B48E22D0E57D1B4FA4FBEDA,),PersistentGuid=00000000000000000000000000000000,bHidden=False,bNotConnectable=False,bDefaultValueIsReadOnly=False,bDefaultValueIsIgnored=False,bAdvancedView=False,)
+			End Object
+			Begin Object Class=MaterialGraphNode Name="MaterialGraphNode_23"
+			   Begin Object Class=MaterialExpressionSubtract Name="MaterialExpressionSubtract_0"
+			   End Object
+			   Begin Object Name="MaterialExpressionSubtract_0"
+				  A=(Expression=MaterialExpressionAdd'MaterialGraphNode_28.MaterialExpressionAdd_1')
+				  B=(Expression=MaterialExpressionMultiply'MaterialGraphNode_26.MaterialExpressionMultiply_14')
+				  MaterialExpressionEditorX=-1040
+				  MaterialExpressionEditorY=-96
+				  MaterialExpressionGuid=A82F279A40DCED04284668A1AE39C7C9
+				  Material=PreviewMaterial'/Engine/Transient.HeightSim'
+			   End Object
+			   MaterialExpression=MaterialExpressionSubtract'MaterialExpressionSubtract_0'
+			   NodePosX=-1040
+			   NodePosY=-96
+			   NodeGuid=1787DCE54BE5C0FAB05856A31EC0FE35
+			   CustomProperties Pin (PinId=16843B8C4667FAEE24ECE0B46B47FBC7,PinName="A",PinType.PinCategory="optional",PinType.PinSubCategoryObject=None,PinType.PinSubCategoryMemberReference=(),PinType.bIsArray=False,PinType.bIsReference=False,PinType.bIsConst=False,PinType.bIsWeakPointer=False,LinkedTo=(MaterialGraphNode_28 1E3560994BCFA9D85FD37A9C60BF5D19,),PersistentGuid=00000000000000000000000000000000,bHidden=False,bNotConnectable=False,bDefaultValueIsReadOnly=False,bDefaultValueIsIgnored=False,bAdvancedView=False,)
+			   CustomProperties Pin (PinId=4EF0449A4C49834C2530BB917DC6C003,PinName="B",PinType.PinCategory="optional",PinType.PinSubCategoryObject=None,PinType.PinSubCategoryMemberReference=(),PinType.bIsArray=False,PinType.bIsReference=False,PinType.bIsConst=False,PinType.bIsWeakPointer=False,LinkedTo=(MaterialGraphNode_26 7513292D4C68BB06A266D0A870388A6B,),PersistentGuid=00000000000000000000000000000000,bHidden=False,bNotConnectable=False,bDefaultValueIsReadOnly=False,bDefaultValueIsIgnored=False,bAdvancedView=False,)
+			   CustomProperties Pin (PinId=265235A346C581EDAE72BCA7D9BC355F,PinName="Output",PinFriendlyName=" ",Direction="EGPD_Output",PinType.PinSubCategoryObject=None,PinType.PinSubCategoryMemberReference=(),PinType.bIsArray=False,PinType.bIsReference=False,PinType.bIsConst=False,PinType.bIsWeakPointer=False,LinkedTo=(MaterialGraphNode_22 0C7E3CBD4F48315F0B53FBA17EF9A98A,),PersistentGuid=00000000000000000000000000000000,bHidden=False,bNotConnectable=False,bDefaultValueIsReadOnly=False,bDefaultValueIsIgnored=False,bAdvancedView=False,)
+			End Object
+			Begin Object Class=MaterialGraphNode Name="MaterialGraphNode_24"
+			   Begin Object Class=MaterialExpressionAdd Name="MaterialExpressionAdd_0"
+			   End Object
+			   Begin Object Name="MaterialExpressionAdd_0"
+				  A=(Expression=MaterialExpressionMultiply'MaterialGraphNode_26.MaterialExpressionMultiply_14')
+				  B=(Expression=MaterialExpressionMultiply'MaterialGraphNode_22.MaterialExpressionMultiply_13')
+				  MaterialExpressionEditorX=-912
+				  MaterialExpressionEditorY=-208
+				  MaterialExpressionGuid=CC3B3D8B469F3812691DE583C9F3E8B4
+				  Material=PreviewMaterial'/Engine/Transient.HeightSim'
+			   End Object
+			   MaterialExpression=MaterialExpressionAdd'MaterialExpressionAdd_0'
+			   NodePosX=-912
+			   NodePosY=-208
+			   NodeGuid=8343E5CA43D441FA485CFE9B49811B6B
+			   CustomProperties Pin (PinId=CCF86AC143B8FD0BBDD5B6A50E86034D,PinName="A",PinType.PinCategory="optional",PinType.PinSubCategoryObject=None,PinType.PinSubCategoryMemberReference=(),PinType.bIsArray=False,PinType.bIsReference=False,PinType.bIsConst=False,PinType.bIsWeakPointer=False,LinkedTo=(MaterialGraphNode_26 7513292D4C68BB06A266D0A870388A6B,),PersistentGuid=00000000000000000000000000000000,bHidden=False,bNotConnectable=False,bDefaultValueIsReadOnly=False,bDefaultValueIsIgnored=False,bAdvancedView=False,)
+			   CustomProperties Pin (PinId=35E274EA4B48E22D0E57D1B4FA4FBEDA,PinName="B",PinType.PinCategory="optional",PinType.PinSubCategoryObject=None,PinType.PinSubCategoryMemberReference=(),PinType.bIsArray=False,PinType.bIsReference=False,PinType.bIsConst=False,PinType.bIsWeakPointer=False,LinkedTo=(MaterialGraphNode_22 E4B6CBFB47131CBEDDFF5E9313B2638F,),PersistentGuid=00000000000000000000000000000000,bHidden=False,bNotConnectable=False,bDefaultValueIsReadOnly=False,bDefaultValueIsIgnored=False,bAdvancedView=False,)
+			   CustomProperties Pin (PinId=69A2ACEF445864C15E03A3A43F0413F8,PinName="Output",PinFriendlyName=" ",Direction="EGPD_Output",PinType.PinSubCategoryObject=None,PinType.PinSubCategoryMemberReference=(),PinType.bIsArray=False,PinType.bIsReference=False,PinType.bIsConst=False,PinType.bIsWeakPointer=False,LinkedTo=(MaterialGraphNode_21 BE3AE6364E3BA620AB73C19920383471,),PersistentGuid=00000000000000000000000000000000,bHidden=False,bNotConnectable=False,bDefaultValueIsReadOnly=False,bDefaultValueIsIgnored=False,bAdvancedView=False,)
+			End Object
+			Begin Object Class=MaterialGraphNode Name="MaterialGraphNode_25"
+			   Begin Object Class=MaterialExpressionScalarParameter Name="MaterialExpressionScalarParameter_23"
+			   End Object
+			   Begin Object Name="MaterialExpressionScalarParameter_23"
+				  DefaultValue=1.000000
+				  ParameterName="TravelSpeed"
+				  ExpressionGUID=1D1CE882495E06C46C539A95598F74F6
+				  MaterialExpressionEditorX=-1392
+				  MaterialExpressionEditorY=16
+				  MaterialExpressionGuid=A0585B19495AB35BF3A7BEBFA1E8A87E
+				  Material=PreviewMaterial'/Engine/Transient.HeightSim'
+			   End Object
+			   MaterialExpression=MaterialExpressionScalarParameter'MaterialExpressionScalarParameter_23'
+			   NodePosX=-1392
+			   NodePosY=16
+			   bCanRenameNode=True
+			   NodeGuid=E92B67E0436E14C8CE74E792055C6DD2
+			   CustomProperties Pin (PinId=F1D57D00491B73D6EF7215919AE4BBDB,PinName="Output",PinFriendlyName=" ",Direction="EGPD_Output",PinType.PinSubCategoryObject=None,PinType.PinSubCategoryMemberReference=(),PinType.bIsArray=False,PinType.bIsReference=False,PinType.bIsConst=False,PinType.bIsWeakPointer=False,LinkedTo=(MaterialGraphNode_22 7CB7D4D94D2CB5B5F59ED998423153E7,),PersistentGuid=00000000000000000000000000000000,bHidden=False,bNotConnectable=False,bDefaultValueIsReadOnly=False,bDefaultValueIsIgnored=False,bAdvancedView=False,)
+			End Object
+			Begin Object Class=MaterialGraphNode Name="MaterialGraphNode_26"
+			   Begin Object Class=MaterialExpressionMultiply Name="MaterialExpressionMultiply_14"
+			   End Object
+			   Begin Object Name="MaterialExpressionMultiply_14"
+				  A=(Expression=MaterialExpressionTextureSampleParameter2D'MaterialGraphNode_27.MaterialExpressionTextureSampleParameter2D_15',Mask=1,MaskR=1,MaskG=1,MaskB=1)
+				  ConstB=4.000000
+				  MaterialExpressionEditorX=-1216
+				  MaterialExpressionEditorY=-208
+				  MaterialExpressionGuid=C681CC4E4F9AB62B72D8729C5345D77A
+				  Material=PreviewMaterial'/Engine/Transient.HeightSim'
+			   End Object
+			   MaterialExpression=MaterialExpressionMultiply'MaterialExpressionMultiply_14'
+			   NodePosX=-1216
+			   NodePosY=-208
+			   NodeGuid=B9CAFCCA49AD91062E385E9EFDEDACFF
+			   CustomProperties Pin (PinId=BD1F41E34440547A47544BB1D6A35635,PinName="A",PinType.PinCategory="optional",PinType.PinSubCategoryObject=None,PinType.PinSubCategoryMemberReference=(),PinType.bIsArray=False,PinType.bIsReference=False,PinType.bIsConst=False,PinType.bIsWeakPointer=False,LinkedTo=(MaterialGraphNode_27 0B3C30D849C447FCC213A4A26F876FA5,),PersistentGuid=00000000000000000000000000000000,bHidden=False,bNotConnectable=False,bDefaultValueIsReadOnly=False,bDefaultValueIsIgnored=False,bAdvancedView=False,)
+			   CustomProperties Pin (PinId=7A2278ED45F2084883C008B1935D7802,PinName="B",PinType.PinCategory="optional",PinType.PinSubCategoryObject=None,PinType.PinSubCategoryMemberReference=(),PinType.bIsArray=False,PinType.bIsReference=False,PinType.bIsConst=False,PinType.bIsWeakPointer=False,PersistentGuid=00000000000000000000000000000000,bHidden=False,bNotConnectable=False,bDefaultValueIsReadOnly=False,bDefaultValueIsIgnored=False,bAdvancedView=False,)
+			   CustomProperties Pin (PinId=7513292D4C68BB06A266D0A870388A6B,PinName="Output",PinFriendlyName=" ",Direction="EGPD_Output",PinType.PinSubCategoryObject=None,PinType.PinSubCategoryMemberReference=(),PinType.bIsArray=False,PinType.bIsReference=False,PinType.bIsConst=False,PinType.bIsWeakPointer=False,LinkedTo=(MaterialGraphNode_23 4EF0449A4C49834C2530BB917DC6C003,MaterialGraphNode_24 CCF86AC143B8FD0BBDD5B6A50E86034D,),PersistentGuid=00000000000000000000000000000000,bHidden=False,bNotConnectable=False,bDefaultValueIsReadOnly=False,bDefaultValueIsIgnored=False,bAdvancedView=False,)
+			End Object
+			Begin Object Class=MaterialGraphNode Name="MaterialGraphNode_27"
+			   Begin Object Class=MaterialExpressionTextureSampleParameter2D Name="MaterialExpressionTextureSampleParameter2D_15"
+			   End Object
+			   Begin Object Name="MaterialExpressionTextureSampleParameter2D_15"
+				  ParameterName="PreviousHeight1"
+				  ExpressionGUID=73F0A9CA495E267230A23E8D0A879DE6
+				  MaterialExpressionEditorX=-1440
+				  MaterialExpressionEditorY=-336
+				  MaterialExpressionGuid=8F3166474EF9A6F6A6825185B78360BE
+				  Material=PreviewMaterial'/Engine/Transient.HeightSim'
+			   End Object
+			   MaterialExpression=MaterialExpressionTextureSampleParameter2D'MaterialExpressionTextureSampleParameter2D_15'
+			   NodePosX=-1440
+			   NodePosY=-336
+			   bCanRenameNode=True
+			   NodeGuid=C1914F76428E09EF4F98D19C560FCA38
+			   CustomProperties Pin (PinId=64C501D8480CAD37BC501B877B50F6F8,PinName="UVs",PinType.PinCategory="optional",PinType.PinSubCategoryObject=None,PinType.PinSubCategoryMemberReference=(),PinType.bIsArray=False,PinType.bIsReference=False,PinType.bIsConst=False,PinType.bIsWeakPointer=False,PersistentGuid=00000000000000000000000000000000,bHidden=False,bNotConnectable=False,bDefaultValueIsReadOnly=False,bDefaultValueIsIgnored=False,bAdvancedView=False,)
+			   CustomProperties Pin (PinId=78B401EA42A5132AE5A973A9D27DFD2E,PinName="Tex",PinType.PinCategory="optional",PinType.PinSubCategoryObject=None,PinType.PinSubCategoryMemberReference=(),PinType.bIsArray=False,PinType.bIsReference=False,PinType.bIsConst=False,PinType.bIsWeakPointer=False,PersistentGuid=00000000000000000000000000000000,bHidden=False,bNotConnectable=False,bDefaultValueIsReadOnly=False,bDefaultValueIsIgnored=False,bAdvancedView=False,)
+			   CustomProperties Pin (PinId=0B3C30D849C447FCC213A4A26F876FA5,PinName="Output",PinFriendlyName=" ",Direction="EGPD_Output",PinType.PinCategory="mask",PinType.PinSubCategoryObject=None,PinType.PinSubCategoryMemberReference=(),PinType.bIsArray=False,PinType.bIsReference=False,PinType.bIsConst=False,PinType.bIsWeakPointer=False,LinkedTo=(MaterialGraphNode_26 BD1F41E34440547A47544BB1D6A35635,),PersistentGuid=00000000000000000000000000000000,bHidden=False,bNotConnectable=False,bDefaultValueIsReadOnly=False,bDefaultValueIsIgnored=False,bAdvancedView=False,)
+			   CustomProperties Pin (PinId=A05C583D40CB1D7BBBA12EA3AD9028D1,PinName="Output2",PinFriendlyName=" ",Direction="EGPD_Output",PinType.PinCategory="mask",PinType.PinSubCategory="red",PinType.PinSubCategoryObject=None,PinType.PinSubCategoryMemberReference=(),PinType.bIsArray=False,PinType.bIsReference=False,PinType.bIsConst=False,PinType.bIsWeakPointer=False,PersistentGuid=00000000000000000000000000000000,bHidden=False,bNotConnectable=False,bDefaultValueIsReadOnly=False,bDefaultValueIsIgnored=False,bAdvancedView=False,)
+			   CustomProperties Pin (PinId=E71951E14965A34C65E815A64E428C7C,PinName="Output3",PinFriendlyName=" ",Direction="EGPD_Output",PinType.PinCategory="mask",PinType.PinSubCategory="green",PinType.PinSubCategoryObject=None,PinType.PinSubCategoryMemberReference=(),PinType.bIsArray=False,PinType.bIsReference=False,PinType.bIsConst=False,PinType.bIsWeakPointer=False,PersistentGuid=00000000000000000000000000000000,bHidden=False,bNotConnectable=False,bDefaultValueIsReadOnly=False,bDefaultValueIsIgnored=False,bAdvancedView=False,)
+			   CustomProperties Pin (PinId=C366BA3A42E0EF4DFEF15BBC9375F3A5,PinName="Output4",PinFriendlyName=" ",Direction="EGPD_Output",PinType.PinCategory="mask",PinType.PinSubCategory="blue",PinType.PinSubCategoryObject=None,PinType.PinSubCategoryMemberReference=(),PinType.bIsArray=False,PinType.bIsReference=False,PinType.bIsConst=False,PinType.bIsWeakPointer=False,PersistentGuid=00000000000000000000000000000000,bHidden=False,bNotConnectable=False,bDefaultValueIsReadOnly=False,bDefaultValueIsIgnored=False,bAdvancedView=False,)
+			   CustomProperties Pin (PinId=AC911C394AFD48CF1E4040829D008C67,PinName="Output5",PinFriendlyName=" ",Direction="EGPD_Output",PinType.PinCategory="mask",PinType.PinSubCategory="alpha",PinType.PinSubCategoryObject=None,PinType.PinSubCategoryMemberReference=(),PinType.bIsArray=False,PinType.bIsReference=False,PinType.bIsConst=False,PinType.bIsWeakPointer=False,PersistentGuid=00000000000000000000000000000000,bHidden=False,bNotConnectable=False,bDefaultValueIsReadOnly=False,bDefaultValueIsIgnored=False,bAdvancedView=False,)
+			End Object
+			Begin Object Class=MaterialGraphNode Name="MaterialGraphNode_28"
+			   Begin Object Class=MaterialExpressionAdd Name="MaterialExpressionAdd_1"
+			   End Object
+			   Begin Object Name="MaterialExpressionAdd_1"
+				  A=(Expression=MaterialExpressionAdd'MaterialGraphNode_42.MaterialExpressionAdd_6')
+				  B=(Expression=MaterialExpressionAdd'MaterialGraphNode_33.MaterialExpressionAdd_3')
+				  MaterialExpressionEditorX=-1424
+				  MaterialExpressionEditorY=-96
+				  MaterialExpressionGuid=741CF8D24976692383DDCAA7D1F1FF97
+				  Material=PreviewMaterial'/Engine/Transient.HeightSim'
+			   End Object
+			   MaterialExpression=MaterialExpressionAdd'MaterialExpressionAdd_1'
+			   NodePosX=-1424
+			   NodePosY=-96
+			   NodeGuid=64C020D84E02DDF2C1F54788B2A270AD
+			   CustomProperties Pin (PinId=BB7B4EAB4EDB7DDEB9548E9EBE9E2B2D,PinName="A",PinType.PinCategory="optional",PinType.PinSubCategoryObject=None,PinType.PinSubCategoryMemberReference=(),PinType.bIsArray=False,PinType.bIsReference=False,PinType.bIsConst=False,PinType.bIsWeakPointer=False,LinkedTo=(MaterialGraphNode_42 EA5FD970428A72F86BD8A0B7F397B30E,),PersistentGuid=00000000000000000000000000000000,bHidden=False,bNotConnectable=False,bDefaultValueIsReadOnly=False,bDefaultValueIsIgnored=False,bAdvancedView=False,)
+			   CustomProperties Pin (PinId=E997EAB34224500E4A6C609C4EAA2351,PinName="B",PinType.PinCategory="optional",PinType.PinSubCategoryObject=None,PinType.PinSubCategoryMemberReference=(),PinType.bIsArray=False,PinType.bIsReference=False,PinType.bIsConst=False,PinType.bIsWeakPointer=False,LinkedTo=(MaterialGraphNode_33 B9CD95A141124727348AD6A0230FF4A6,),PersistentGuid=00000000000000000000000000000000,bHidden=False,bNotConnectable=False,bDefaultValueIsReadOnly=False,bDefaultValueIsIgnored=False,bAdvancedView=False,)
+			   CustomProperties Pin (PinId=1E3560994BCFA9D85FD37A9C60BF5D19,PinName="Output",PinFriendlyName=" ",Direction="EGPD_Output",PinType.PinSubCategoryObject=None,PinType.PinSubCategoryMemberReference=(),PinType.bIsArray=False,PinType.bIsReference=False,PinType.bIsConst=False,PinType.bIsWeakPointer=False,LinkedTo=(MaterialGraphNode_23 16843B8C4667FAEE24ECE0B46B47FBC7,),PersistentGuid=00000000000000000000000000000000,bHidden=False,bNotConnectable=False,bDefaultValueIsReadOnly=False,bDefaultValueIsIgnored=False,bAdvancedView=False,)
+			End Object
+			Begin Object Class=MaterialGraphNode Name="MaterialGraphNode_29"
+			   Begin Object Class=MaterialExpressionTextureSampleParameter2D Name="MaterialExpressionTextureSampleParameter2D_16"
+			   End Object
+			   Begin Object Name="MaterialExpressionTextureSampleParameter2D_16"
+				  ParameterName="PreviousHeight1"
+				  ExpressionGUID=6C6969ED4547E73070E640B22EFBF548
+				  Coordinates=(Expression=MaterialExpressionAdd'MaterialGraphNode_30.MaterialExpressionAdd_2')
+				  MaterialExpressionEditorX=-1824
+				  MaterialExpressionEditorY=224
+				  MaterialExpressionGuid=8F3166474EF9A6F6A6825185B78360BE
+				  Material=PreviewMaterial'/Engine/Transient.HeightSim'
+			   End Object
+			   MaterialExpression=MaterialExpressionTextureSampleParameter2D'MaterialExpressionTextureSampleParameter2D_16'
+			   NodePosX=-1824
+			   NodePosY=224
+			   bCanRenameNode=True
+			   NodeGuid=78F2D9D14261ABA345A1E29E392F5508
+			   CustomProperties Pin (PinId=0AE0FEAE4FDE26A7DEC5C4A20B818D92,PinName="UVs",PinType.PinCategory="optional",PinType.PinSubCategoryObject=None,PinType.PinSubCategoryMemberReference=(),PinType.bIsArray=False,PinType.bIsReference=False,PinType.bIsConst=False,PinType.bIsWeakPointer=False,LinkedTo=(MaterialGraphNode_30 56B9252F43B37515BEF81695D982BCA8,),PersistentGuid=00000000000000000000000000000000,bHidden=False,bNotConnectable=False,bDefaultValueIsReadOnly=False,bDefaultValueIsIgnored=False,bAdvancedView=False,)
+			   CustomProperties Pin (PinId=0398A20947677E6011AE2E8C7AA47194,PinName="Tex",PinType.PinCategory="optional",PinType.PinSubCategoryObject=None,PinType.PinSubCategoryMemberReference=(),PinType.bIsArray=False,PinType.bIsReference=False,PinType.bIsConst=False,PinType.bIsWeakPointer=False,PersistentGuid=00000000000000000000000000000000,bHidden=False,bNotConnectable=False,bDefaultValueIsReadOnly=False,bDefaultValueIsIgnored=False,bAdvancedView=False,)
+			   CustomProperties Pin (PinId=59D14CBD4994017371C83FA1A08FC32E,PinName="Output",PinFriendlyName=" ",Direction="EGPD_Output",PinType.PinCategory="mask",PinType.PinSubCategoryObject=None,PinType.PinSubCategoryMemberReference=(),PinType.bIsArray=False,PinType.bIsReference=False,PinType.bIsConst=False,PinType.bIsWeakPointer=False,PersistentGuid=00000000000000000000000000000000,bHidden=False,bNotConnectable=False,bDefaultValueIsReadOnly=False,bDefaultValueIsIgnored=False,bAdvancedView=False,)
+			   CustomProperties Pin (PinId=9E3F893A40C9090B3CB0E195E9BC1C51,PinName="Output2",PinFriendlyName=" ",Direction="EGPD_Output",PinType.PinCategory="mask",PinType.PinSubCategory="red",PinType.PinSubCategoryObject=None,PinType.PinSubCategoryMemberReference=(),PinType.bIsArray=False,PinType.bIsReference=False,PinType.bIsConst=False,PinType.bIsWeakPointer=False,LinkedTo=(MaterialGraphNode_33 326B9F214D7E4D47BC38DCB06A7EBF7C,),PersistentGuid=00000000000000000000000000000000,bHidden=False,bNotConnectable=False,bDefaultValueIsReadOnly=False,bDefaultValueIsIgnored=False,bAdvancedView=False,)
+			   CustomProperties Pin (PinId=20339579495FB653DF4E8AA21B7A93CA,PinName="Output3",PinFriendlyName=" ",Direction="EGPD_Output",PinType.PinCategory="mask",PinType.PinSubCategory="green",PinType.PinSubCategoryObject=None,PinType.PinSubCategoryMemberReference=(),PinType.bIsArray=False,PinType.bIsReference=False,PinType.bIsConst=False,PinType.bIsWeakPointer=False,PersistentGuid=00000000000000000000000000000000,bHidden=False,bNotConnectable=False,bDefaultValueIsReadOnly=False,bDefaultValueIsIgnored=False,bAdvancedView=False,)
+			   CustomProperties Pin (PinId=0E0A5A0F47422DB551A031B1E0838883,PinName="Output4",PinFriendlyName=" ",Direction="EGPD_Output",PinType.PinCategory="mask",PinType.PinSubCategory="blue",PinType.PinSubCategoryObject=None,PinType.PinSubCategoryMemberReference=(),PinType.bIsArray=False,PinType.bIsReference=False,PinType.bIsConst=False,PinType.bIsWeakPointer=False,PersistentGuid=00000000000000000000000000000000,bHidden=False,bNotConnectable=False,bDefaultValueIsReadOnly=False,bDefaultValueIsIgnored=False,bAdvancedView=False,)
+			   CustomProperties Pin (PinId=9D8DEA3E4FE14EC1C9423A9B6536EDE2,PinName="Output5",PinFriendlyName=" ",Direction="EGPD_Output",PinType.PinCategory="mask",PinType.PinSubCategory="alpha",PinType.PinSubCategoryObject=None,PinType.PinSubCategoryMemberReference=(),PinType.bIsArray=False,PinType.bIsReference=False,PinType.bIsConst=False,PinType.bIsWeakPointer=False,PersistentGuid=00000000000000000000000000000000,bHidden=False,bNotConnectable=False,bDefaultValueIsReadOnly=False,bDefaultValueIsIgnored=False,bAdvancedView=False,)
+			End Object
+			Begin Object Class=MaterialGraphNode Name="MaterialGraphNode_30"
+			   Begin Object Class=MaterialExpressionAdd Name="MaterialExpressionAdd_2"
+			   End Object
+			   Begin Object Name="MaterialExpressionAdd_2"
+				  A=(Expression=MaterialExpressionTextureCoordinate'MaterialGraphNode_31.MaterialExpressionTextureCoordinate_0')
+				  B=(Expression=MaterialExpressionMultiply'MaterialGraphNode_14.MaterialExpressionMultiply_8')
+				  MaterialExpressionEditorX=-1968
+				  MaterialExpressionEditorY=240
+				  MaterialExpressionGuid=EF5DFD284086D3D0FDC39A86FEE5E825
+				  Material=PreviewMaterial'/Engine/Transient.HeightSim'
+			   End Object
+			   MaterialExpression=MaterialExpressionAdd'MaterialExpressionAdd_2'
+			   NodePosX=-1968
+			   NodePosY=240
+			   NodeGuid=0B5A854A4A2CF7B6823D898CF2ADDE68
+			   CustomProperties Pin (PinId=0521EFEE4ABBF16458247EB77A1AC5D6,PinName="A",PinType.PinCategory="optional",PinType.PinSubCategoryObject=None,PinType.PinSubCategoryMemberReference=(),PinType.bIsArray=False,PinType.bIsReference=False,PinType.bIsConst=False,PinType.bIsWeakPointer=False,LinkedTo=(MaterialGraphNode_31 EC3E7A034F50A4B55DF6E58EEA47F536,),PersistentGuid=00000000000000000000000000000000,bHidden=False,bNotConnectable=False,bDefaultValueIsReadOnly=False,bDefaultValueIsIgnored=False,bAdvancedView=False,)
+			   CustomProperties Pin (PinId=7ECB60B641D4CC0CC8A0D1ADFAC26BB2,PinName="B",PinType.PinCategory="optional",PinType.PinSubCategoryObject=None,PinType.PinSubCategoryMemberReference=(),PinType.bIsArray=False,PinType.bIsReference=False,PinType.bIsConst=False,PinType.bIsWeakPointer=False,LinkedTo=(MaterialGraphNode_14 5E9617AB4A810D7A05D2D2A1AECCAC21,),PersistentGuid=00000000000000000000000000000000,bHidden=False,bNotConnectable=False,bDefaultValueIsReadOnly=False,bDefaultValueIsIgnored=False,bAdvancedView=False,)
+			   CustomProperties Pin (PinId=56B9252F43B37515BEF81695D982BCA8,PinName="Output",PinFriendlyName=" ",Direction="EGPD_Output",PinType.PinSubCategoryObject=None,PinType.PinSubCategoryMemberReference=(),PinType.bIsArray=False,PinType.bIsReference=False,PinType.bIsConst=False,PinType.bIsWeakPointer=False,LinkedTo=(MaterialGraphNode_29 0AE0FEAE4FDE26A7DEC5C4A20B818D92,),PersistentGuid=00000000000000000000000000000000,bHidden=False,bNotConnectable=False,bDefaultValueIsReadOnly=False,bDefaultValueIsIgnored=False,bAdvancedView=False,)
+			End Object
+			Begin Object Class=MaterialGraphNode Name="MaterialGraphNode_31"
+			   Begin Object Class=MaterialExpressionTextureCoordinate Name="MaterialExpressionTextureCoordinate_0"
+			   End Object
+			   Begin Object Name="MaterialExpressionTextureCoordinate_0"
+				  MaterialExpressionEditorX=-2144
+				  MaterialExpressionEditorY=192
+				  MaterialExpressionGuid=F1C1E2C24651DC43FD6AD5B0AF633576
+				  Material=PreviewMaterial'/Engine/Transient.HeightSim'
+			   End Object
+			   MaterialExpression=MaterialExpressionTextureCoordinate'MaterialExpressionTextureCoordinate_0'
+			   NodePosX=-2144
+			   NodePosY=192
+			   NodeGuid=248C745F415DFA3163CCF2B8F80A6BEC
+			   CustomProperties Pin (PinId=EC3E7A034F50A4B55DF6E58EEA47F536,PinName="Output",PinFriendlyName=" ",Direction="EGPD_Output",PinType.PinSubCategoryObject=None,PinType.PinSubCategoryMemberReference=(),PinType.bIsArray=False,PinType.bIsReference=False,PinType.bIsConst=False,PinType.bIsWeakPointer=False,LinkedTo=(MaterialGraphNode_30 0521EFEE4ABBF16458247EB77A1AC5D6,),PersistentGuid=00000000000000000000000000000000,bHidden=False,bNotConnectable=False,bDefaultValueIsReadOnly=False,bDefaultValueIsIgnored=False,bAdvancedView=False,)
+			End Object
+			Begin Object Class=MaterialGraphNode Name="MaterialGraphNode_32"
+			   Begin Object Class=MaterialExpressionConstant2Vector Name="MaterialExpressionConstant2Vector_0"
+			   End Object
+			   Begin Object Name="MaterialExpressionConstant2Vector_0"
+				  G=1.000000
+				  MaterialExpressionEditorX=-2256
+				  MaterialExpressionEditorY=368
+				  MaterialExpressionGuid=F611381D424E60EA8824DA9E1A50EFEB
+				  Material=PreviewMaterial'/Engine/Transient.HeightSim'
+			   End Object
+			   MaterialExpression=MaterialExpressionConstant2Vector'MaterialExpressionConstant2Vector_0'
+			   NodePosX=-2256
+			   NodePosY=368
+			   NodeGuid=CEC2A9C549CCB56ACE2EACADFC542750
+			   CustomProperties Pin (PinId=0C3200D34E5DE557B3B7D38DA4AA3369,PinName="Output",PinFriendlyName=" ",Direction="EGPD_Output",PinType.PinSubCategoryObject=None,PinType.PinSubCategoryMemberReference=(),PinType.bIsArray=False,PinType.bIsReference=False,PinType.bIsConst=False,PinType.bIsWeakPointer=False,LinkedTo=(MaterialGraphNode_14 B243F91847B4DFA8816E3F8485A06DE3,),PersistentGuid=00000000000000000000000000000000,bHidden=False,bNotConnectable=False,bDefaultValueIsReadOnly=False,bDefaultValueIsIgnored=False,bAdvancedView=False,)
+			End Object
+			Begin Object Class=MaterialGraphNode Name="MaterialGraphNode_33"
+			   Begin Object Class=MaterialExpressionAdd Name="MaterialExpressionAdd_3"
+			   End Object
+			   Begin Object Name="MaterialExpressionAdd_3"
+				  A=(Expression=MaterialExpressionTextureSampleParameter2D'MaterialGraphNode_34.MaterialExpressionTextureSampleParameter2D_17',OutputIndex=1,Mask=1,MaskR=1)
+				  B=(Expression=MaterialExpressionTextureSampleParameter2D'MaterialGraphNode_29.MaterialExpressionTextureSampleParameter2D_16',OutputIndex=1,Mask=1,MaskR=1)
+				  MaterialExpressionEditorX=-1632
+				  MaterialExpressionEditorY=128
+				  MaterialExpressionGuid=7C09BC9848A1D28AA8F90791329F6B40
+				  Material=PreviewMaterial'/Engine/Transient.HeightSim'
+			   End Object
+			   MaterialExpression=MaterialExpressionAdd'MaterialExpressionAdd_3'
+			   NodePosX=-1632
+			   NodePosY=128
+			   NodeGuid=159C9DC74AFF909DCEE74796840971D5
+			   CustomProperties Pin (PinId=3403324C41E1131264B0E59B4223E10C,PinName="A",PinType.PinCategory="optional",PinType.PinSubCategoryObject=None,PinType.PinSubCategoryMemberReference=(),PinType.bIsArray=False,PinType.bIsReference=False,PinType.bIsConst=False,PinType.bIsWeakPointer=False,LinkedTo=(MaterialGraphNode_34 D613CC90425867D1D465F7A30BA7F588,),PersistentGuid=00000000000000000000000000000000,bHidden=False,bNotConnectable=False,bDefaultValueIsReadOnly=False,bDefaultValueIsIgnored=False,bAdvancedView=False,)
+			   CustomProperties Pin (PinId=326B9F214D7E4D47BC38DCB06A7EBF7C,PinName="B",PinType.PinCategory="optional",PinType.PinSubCategoryObject=None,PinType.PinSubCategoryMemberReference=(),PinType.bIsArray=False,PinType.bIsReference=False,PinType.bIsConst=False,PinType.bIsWeakPointer=False,LinkedTo=(MaterialGraphNode_29 9E3F893A40C9090B3CB0E195E9BC1C51,),PersistentGuid=00000000000000000000000000000000,bHidden=False,bNotConnectable=False,bDefaultValueIsReadOnly=False,bDefaultValueIsIgnored=False,bAdvancedView=False,)
+			   CustomProperties Pin (PinId=B9CD95A141124727348AD6A0230FF4A6,PinName="Output",PinFriendlyName=" ",Direction="EGPD_Output",PinType.PinSubCategoryObject=None,PinType.PinSubCategoryMemberReference=(),PinType.bIsArray=False,PinType.bIsReference=False,PinType.bIsConst=False,PinType.bIsWeakPointer=False,LinkedTo=(MaterialGraphNode_28 E997EAB34224500E4A6C609C4EAA2351,),PersistentGuid=00000000000000000000000000000000,bHidden=False,bNotConnectable=False,bDefaultValueIsReadOnly=False,bDefaultValueIsIgnored=False,bAdvancedView=False,)
+			End Object
+			Begin Object Class=MaterialGraphNode Name="MaterialGraphNode_34"
+			   Begin Object Class=MaterialExpressionTextureSampleParameter2D Name="MaterialExpressionTextureSampleParameter2D_17"
+			   End Object
+			   Begin Object Name="MaterialExpressionTextureSampleParameter2D_17"
+				  ParameterName="PreviousHeight1"
+				  ExpressionGUID=25A38FC34C188997E6B9008204282A6A
+				  Coordinates=(Expression=MaterialExpressionAdd'MaterialGraphNode_35.MaterialExpressionAdd_4')
+				  MaterialExpressionEditorX=-1824
+				  MaterialExpressionEditorY=-32
+				  MaterialExpressionGuid=8F3166474EF9A6F6A6825185B78360BE
+				  Material=PreviewMaterial'/Engine/Transient.HeightSim'
+			   End Object
+			   MaterialExpression=MaterialExpressionTextureSampleParameter2D'MaterialExpressionTextureSampleParameter2D_17'
+			   NodePosX=-1824
+			   NodePosY=-32
+			   bCanRenameNode=True
+			   NodeGuid=A327E0AC422AE56FE1D00B8D1CFEA84C
+			   CustomProperties Pin (PinId=8BC7DF074F21C2BBB568BDBD59679CA5,PinName="UVs",PinType.PinCategory="optional",PinType.PinSubCategoryObject=None,PinType.PinSubCategoryMemberReference=(),PinType.bIsArray=False,PinType.bIsReference=False,PinType.bIsConst=False,PinType.bIsWeakPointer=False,LinkedTo=(MaterialGraphNode_35 5C98594547A6BD67174460A26F665539,),PersistentGuid=00000000000000000000000000000000,bHidden=False,bNotConnectable=False,bDefaultValueIsReadOnly=False,bDefaultValueIsIgnored=False,bAdvancedView=False,)
+			   CustomProperties Pin (PinId=A87BD50449964337780E9099E43BFEA5,PinName="Tex",PinType.PinCategory="optional",PinType.PinSubCategoryObject=None,PinType.PinSubCategoryMemberReference=(),PinType.bIsArray=False,PinType.bIsReference=False,PinType.bIsConst=False,PinType.bIsWeakPointer=False,PersistentGuid=00000000000000000000000000000000,bHidden=False,bNotConnectable=False,bDefaultValueIsReadOnly=False,bDefaultValueIsIgnored=False,bAdvancedView=False,)
+			   CustomProperties Pin (PinId=548A183D4AA7D2E5D0E850AE4590B868,PinName="Output",PinFriendlyName=" ",Direction="EGPD_Output",PinType.PinCategory="mask",PinType.PinSubCategoryObject=None,PinType.PinSubCategoryMemberReference=(),PinType.bIsArray=False,PinType.bIsReference=False,PinType.bIsConst=False,PinType.bIsWeakPointer=False,PersistentGuid=00000000000000000000000000000000,bHidden=False,bNotConnectable=False,bDefaultValueIsReadOnly=False,bDefaultValueIsIgnored=False,bAdvancedView=False,)
+			   CustomProperties Pin (PinId=D613CC90425867D1D465F7A30BA7F588,PinName="Output2",PinFriendlyName=" ",Direction="EGPD_Output",PinType.PinCategory="mask",PinType.PinSubCategory="red",PinType.PinSubCategoryObject=None,PinType.PinSubCategoryMemberReference=(),PinType.bIsArray=False,PinType.bIsReference=False,PinType.bIsConst=False,PinType.bIsWeakPointer=False,LinkedTo=(MaterialGraphNode_33 3403324C41E1131264B0E59B4223E10C,),PersistentGuid=00000000000000000000000000000000,bHidden=False,bNotConnectable=False,bDefaultValueIsReadOnly=False,bDefaultValueIsIgnored=False,bAdvancedView=False,)
+			   CustomProperties Pin (PinId=4CA2D44D468AEB5DACCB6BBC088FEAFF,PinName="Output3",PinFriendlyName=" ",Direction="EGPD_Output",PinType.PinCategory="mask",PinType.PinSubCategory="green",PinType.PinSubCategoryObject=None,PinType.PinSubCategoryMemberReference=(),PinType.bIsArray=False,PinType.bIsReference=False,PinType.bIsConst=False,PinType.bIsWeakPointer=False,PersistentGuid=00000000000000000000000000000000,bHidden=False,bNotConnectable=False,bDefaultValueIsReadOnly=False,bDefaultValueIsIgnored=False,bAdvancedView=False,)
+			   CustomProperties Pin (PinId=B7F944E94C29C014824FB485823ABC02,PinName="Output4",PinFriendlyName=" ",Direction="EGPD_Output",PinType.PinCategory="mask",PinType.PinSubCategory="blue",PinType.PinSubCategoryObject=None,PinType.PinSubCategoryMemberReference=(),PinType.bIsArray=False,PinType.bIsReference=False,PinType.bIsConst=False,PinType.bIsWeakPointer=False,PersistentGuid=00000000000000000000000000000000,bHidden=False,bNotConnectable=False,bDefaultValueIsReadOnly=False,bDefaultValueIsIgnored=False,bAdvancedView=False,)
+			   CustomProperties Pin (PinId=88CC1AD94DEFD5DB32BBF08087F9D949,PinName="Output5",PinFriendlyName=" ",Direction="EGPD_Output",PinType.PinCategory="mask",PinType.PinSubCategory="alpha",PinType.PinSubCategoryObject=None,PinType.PinSubCategoryMemberReference=(),PinType.bIsArray=False,PinType.bIsReference=False,PinType.bIsConst=False,PinType.bIsWeakPointer=False,PersistentGuid=00000000000000000000000000000000,bHidden=False,bNotConnectable=False,bDefaultValueIsReadOnly=False,bDefaultValueIsIgnored=False,bAdvancedView=False,)
+			End Object
+			Begin Object Class=MaterialGraphNode Name="MaterialGraphNode_35"
+			   Begin Object Class=MaterialExpressionAdd Name="MaterialExpressionAdd_4"
+			   End Object
+			   Begin Object Name="MaterialExpressionAdd_4"
+				  A=(Expression=MaterialExpressionTextureCoordinate'MaterialGraphNode_36.MaterialExpressionTextureCoordinate_1')
+				  B=(Expression=MaterialExpressionMultiply'MaterialGraphNode_15.MaterialExpressionMultiply_9')
+				  MaterialExpressionEditorX=-1968
+				  MaterialExpressionEditorY=-16
+				  MaterialExpressionGuid=EF5DFD284086D3D0FDC39A86FEE5E825
+				  Material=PreviewMaterial'/Engine/Transient.HeightSim'
+			   End Object
+			   MaterialExpression=MaterialExpressionAdd'MaterialExpressionAdd_4'
+			   NodePosX=-1968
+			   NodePosY=-16
+			   NodeGuid=39253459460EF44C6A2D92AC3365EF8A
+			   CustomProperties Pin (PinId=A345B0514924018F1BD57E9F435FB86D,PinName="A",PinType.PinCategory="optional",PinType.PinSubCategoryObject=None,PinType.PinSubCategoryMemberReference=(),PinType.bIsArray=False,PinType.bIsReference=False,PinType.bIsConst=False,PinType.bIsWeakPointer=False,LinkedTo=(MaterialGraphNode_36 B81BD90748D4EA11825244AC94607E76,),PersistentGuid=00000000000000000000000000000000,bHidden=False,bNotConnectable=False,bDefaultValueIsReadOnly=False,bDefaultValueIsIgnored=False,bAdvancedView=False,)
+			   CustomProperties Pin (PinId=1576FA7A426EF5DB948BBCA5C72B5288,PinName="B",PinType.PinCategory="optional",PinType.PinSubCategoryObject=None,PinType.PinSubCategoryMemberReference=(),PinType.bIsArray=False,PinType.bIsReference=False,PinType.bIsConst=False,PinType.bIsWeakPointer=False,LinkedTo=(MaterialGraphNode_15 9FC9244144840A564AEA2BA8E1D244B7,),PersistentGuid=00000000000000000000000000000000,bHidden=False,bNotConnectable=False,bDefaultValueIsReadOnly=False,bDefaultValueIsIgnored=False,bAdvancedView=False,)
+			   CustomProperties Pin (PinId=5C98594547A6BD67174460A26F665539,PinName="Output",PinFriendlyName=" ",Direction="EGPD_Output",PinType.PinSubCategoryObject=None,PinType.PinSubCategoryMemberReference=(),PinType.bIsArray=False,PinType.bIsReference=False,PinType.bIsConst=False,PinType.bIsWeakPointer=False,LinkedTo=(MaterialGraphNode_34 8BC7DF074F21C2BBB568BDBD59679CA5,),PersistentGuid=00000000000000000000000000000000,bHidden=False,bNotConnectable=False,bDefaultValueIsReadOnly=False,bDefaultValueIsIgnored=False,bAdvancedView=False,)
+			End Object
+			Begin Object Class=MaterialGraphNode Name="MaterialGraphNode_36"
+			   Begin Object Class=MaterialExpressionTextureCoordinate Name="MaterialExpressionTextureCoordinate_1"
+			   End Object
+			   Begin Object Name="MaterialExpressionTextureCoordinate_1"
+				  MaterialExpressionEditorX=-2112
+				  MaterialExpressionEditorY=-64
+				  MaterialExpressionGuid=F1C1E2C24651DC43FD6AD5B0AF633576
+				  Material=PreviewMaterial'/Engine/Transient.HeightSim'
+			   End Object
+			   MaterialExpression=MaterialExpressionTextureCoordinate'MaterialExpressionTextureCoordinate_1'
+			   NodePosX=-2112
+			   NodePosY=-64
+			   NodeGuid=76F0D1EE4CD1D758EE3DCC9F18D52FB5
+			   CustomProperties Pin (PinId=B81BD90748D4EA11825244AC94607E76,PinName="Output",PinFriendlyName=" ",Direction="EGPD_Output",PinType.PinSubCategoryObject=None,PinType.PinSubCategoryMemberReference=(),PinType.bIsArray=False,PinType.bIsReference=False,PinType.bIsConst=False,PinType.bIsWeakPointer=False,LinkedTo=(MaterialGraphNode_35 A345B0514924018F1BD57E9F435FB86D,),PersistentGuid=00000000000000000000000000000000,bHidden=False,bNotConnectable=False,bDefaultValueIsReadOnly=False,bDefaultValueIsIgnored=False,bAdvancedView=False,)
+			End Object
+			Begin Object Class=MaterialGraphNode Name="MaterialGraphNode_37"
+			   Begin Object Class=MaterialExpressionConstant2Vector Name="MaterialExpressionConstant2Vector_1"
+			   End Object
+			   Begin Object Name="MaterialExpressionConstant2Vector_1"
+				  G=-1.000000
+				  MaterialExpressionEditorX=-2192
+				  MaterialExpressionEditorY=128
+				  MaterialExpressionGuid=F611381D424E60EA8824DA9E1A50EFEB
+				  Material=PreviewMaterial'/Engine/Transient.HeightSim'
+			   End Object
+			   MaterialExpression=MaterialExpressionConstant2Vector'MaterialExpressionConstant2Vector_1'
+			   NodePosX=-2192
+			   NodePosY=128
+			   NodeGuid=FA6A192B4B8F72EA56D30AB4EC079A49
+			   CustomProperties Pin (PinId=9AD9001E4BB978F2279A2E8107F639E1,PinName="Output",PinFriendlyName=" ",Direction="EGPD_Output",PinType.PinSubCategoryObject=None,PinType.PinSubCategoryMemberReference=(),PinType.bIsArray=False,PinType.bIsReference=False,PinType.bIsConst=False,PinType.bIsWeakPointer=False,LinkedTo=(MaterialGraphNode_15 59E2DAA140B7A021018B8980E6423224,),PersistentGuid=00000000000000000000000000000000,bHidden=False,bNotConnectable=False,bDefaultValueIsReadOnly=False,bDefaultValueIsIgnored=False,bAdvancedView=False,)
+			End Object
+			Begin Object Class=MaterialGraphNode Name="MaterialGraphNode_38"
+			   Begin Object Class=MaterialExpressionTextureSampleParameter2D Name="MaterialExpressionTextureSampleParameter2D_18"
+			   End Object
+			   Begin Object Name="MaterialExpressionTextureSampleParameter2D_18"
+				  ParameterName="PreviousHeight1"
+				  ExpressionGUID=94414C694CC9B627D1A6F996C5C9FDA9
+				  Coordinates=(Expression=MaterialExpressionAdd'MaterialGraphNode_39.MaterialExpressionAdd_5')
+				  MaterialExpressionEditorX=-1824
+				  MaterialExpressionEditorY=-288
+				  MaterialExpressionGuid=8F3166474EF9A6F6A6825185B78360BE
+				  Material=PreviewMaterial'/Engine/Transient.HeightSim'
+			   End Object
+			   MaterialExpression=MaterialExpressionTextureSampleParameter2D'MaterialExpressionTextureSampleParameter2D_18'
+			   NodePosX=-1824
+			   NodePosY=-288
+			   bCanRenameNode=True
+			   NodeGuid=A2D1BF1F405CB85C3CFD6A82006174B4
+			   CustomProperties Pin (PinId=1B861E0941ED374899F4D29C48C48C5D,PinName="UVs",PinType.PinCategory="optional",PinType.PinSubCategoryObject=None,PinType.PinSubCategoryMemberReference=(),PinType.bIsArray=False,PinType.bIsReference=False,PinType.bIsConst=False,PinType.bIsWeakPointer=False,LinkedTo=(MaterialGraphNode_39 745A6C3C4261FFCA78E401A1F32DC1A8,),PersistentGuid=00000000000000000000000000000000,bHidden=False,bNotConnectable=False,bDefaultValueIsReadOnly=False,bDefaultValueIsIgnored=False,bAdvancedView=False,)
+			   CustomProperties Pin (PinId=4AF62BDC48D23029F4B968A76EC44C57,PinName="Tex",PinType.PinCategory="optional",PinType.PinSubCategoryObject=None,PinType.PinSubCategoryMemberReference=(),PinType.bIsArray=False,PinType.bIsReference=False,PinType.bIsConst=False,PinType.bIsWeakPointer=False,PersistentGuid=00000000000000000000000000000000,bHidden=False,bNotConnectable=False,bDefaultValueIsReadOnly=False,bDefaultValueIsIgnored=False,bAdvancedView=False,)
+			   CustomProperties Pin (PinId=941FAC534C040326612DDE96F5D17825,PinName="Output",PinFriendlyName=" ",Direction="EGPD_Output",PinType.PinCategory="mask",PinType.PinSubCategoryObject=None,PinType.PinSubCategoryMemberReference=(),PinType.bIsArray=False,PinType.bIsReference=False,PinType.bIsConst=False,PinType.bIsWeakPointer=False,PersistentGuid=00000000000000000000000000000000,bHidden=False,bNotConnectable=False,bDefaultValueIsReadOnly=False,bDefaultValueIsIgnored=False,bAdvancedView=False,)
+			   CustomProperties Pin (PinId=BBF439FB451DE0A0AC85FA9F9D45895C,PinName="Output2",PinFriendlyName=" ",Direction="EGPD_Output",PinType.PinCategory="mask",PinType.PinSubCategory="red",PinType.PinSubCategoryObject=None,PinType.PinSubCategoryMemberReference=(),PinType.bIsArray=False,PinType.bIsReference=False,PinType.bIsConst=False,PinType.bIsWeakPointer=False,LinkedTo=(MaterialGraphNode_42 6A11559A47452394DE46EB95036FB2D5,),PersistentGuid=00000000000000000000000000000000,bHidden=False,bNotConnectable=False,bDefaultValueIsReadOnly=False,bDefaultValueIsIgnored=False,bAdvancedView=False,)
+			   CustomProperties Pin (PinId=20E55D5D400309BE1D1F9586221389D3,PinName="Output3",PinFriendlyName=" ",Direction="EGPD_Output",PinType.PinCategory="mask",PinType.PinSubCategory="green",PinType.PinSubCategoryObject=None,PinType.PinSubCategoryMemberReference=(),PinType.bIsArray=False,PinType.bIsReference=False,PinType.bIsConst=False,PinType.bIsWeakPointer=False,PersistentGuid=00000000000000000000000000000000,bHidden=False,bNotConnectable=False,bDefaultValueIsReadOnly=False,bDefaultValueIsIgnored=False,bAdvancedView=False,)
+			   CustomProperties Pin (PinId=B61098C04AE24D3CF698878CF0996603,PinName="Output4",PinFriendlyName=" ",Direction="EGPD_Output",PinType.PinCategory="mask",PinType.PinSubCategory="blue",PinType.PinSubCategoryObject=None,PinType.PinSubCategoryMemberReference=(),PinType.bIsArray=False,PinType.bIsReference=False,PinType.bIsConst=False,PinType.bIsWeakPointer=False,PersistentGuid=00000000000000000000000000000000,bHidden=False,bNotConnectable=False,bDefaultValueIsReadOnly=False,bDefaultValueIsIgnored=False,bAdvancedView=False,)
+			   CustomProperties Pin (PinId=1F84FA3A41936283441A8A9A1C99122A,PinName="Output5",PinFriendlyName=" ",Direction="EGPD_Output",PinType.PinCategory="mask",PinType.PinSubCategory="alpha",PinType.PinSubCategoryObject=None,PinType.PinSubCategoryMemberReference=(),PinType.bIsArray=False,PinType.bIsReference=False,PinType.bIsConst=False,PinType.bIsWeakPointer=False,PersistentGuid=00000000000000000000000000000000,bHidden=False,bNotConnectable=False,bDefaultValueIsReadOnly=False,bDefaultValueIsIgnored=False,bAdvancedView=False,)
+			End Object
+			Begin Object Class=MaterialGraphNode Name="MaterialGraphNode_39"
+			   Begin Object Class=MaterialExpressionAdd Name="MaterialExpressionAdd_5"
+			   End Object
+			   Begin Object Name="MaterialExpressionAdd_5"
+				  A=(Expression=MaterialExpressionTextureCoordinate'MaterialGraphNode_40.MaterialExpressionTextureCoordinate_2')
+				  B=(Expression=MaterialExpressionMultiply'MaterialGraphNode_16.MaterialExpressionMultiply_10')
+				  MaterialExpressionEditorX=-1952
+				  MaterialExpressionEditorY=-272
+				  MaterialExpressionGuid=EF5DFD284086D3D0FDC39A86FEE5E825
+				  Material=PreviewMaterial'/Engine/Transient.HeightSim'
+			   End Object
+			   MaterialExpression=MaterialExpressionAdd'MaterialExpressionAdd_5'
+			   NodePosX=-1952
+			   NodePosY=-272
+			   NodeGuid=253A180A4C8B414C297404B628353FFD
+			   CustomProperties Pin (PinId=5C5EC3DB4CC914962FC7E3BD10047D43,PinName="A",PinType.PinCategory="optional",PinType.PinSubCategoryObject=None,PinType.PinSubCategoryMemberReference=(),PinType.bIsArray=False,PinType.bIsReference=False,PinType.bIsConst=False,PinType.bIsWeakPointer=False,LinkedTo=(MaterialGraphNode_40 367F1A0A4683231CCBEA8ABDC38CC396,),PersistentGuid=00000000000000000000000000000000,bHidden=False,bNotConnectable=False,bDefaultValueIsReadOnly=False,bDefaultValueIsIgnored=False,bAdvancedView=False,)
+			   CustomProperties Pin (PinId=40DFB4A047C1871E144694A9EF25DBE5,PinName="B",PinType.PinCategory="optional",PinType.PinSubCategoryObject=None,PinType.PinSubCategoryMemberReference=(),PinType.bIsArray=False,PinType.bIsReference=False,PinType.bIsConst=False,PinType.bIsWeakPointer=False,LinkedTo=(MaterialGraphNode_16 DEB3772B4736834DF22CAFA5F034F702,),PersistentGuid=00000000000000000000000000000000,bHidden=False,bNotConnectable=False,bDefaultValueIsReadOnly=False,bDefaultValueIsIgnored=False,bAdvancedView=False,)
+			   CustomProperties Pin (PinId=745A6C3C4261FFCA78E401A1F32DC1A8,PinName="Output",PinFriendlyName=" ",Direction="EGPD_Output",PinType.PinSubCategoryObject=None,PinType.PinSubCategoryMemberReference=(),PinType.bIsArray=False,PinType.bIsReference=False,PinType.bIsConst=False,PinType.bIsWeakPointer=False,LinkedTo=(MaterialGraphNode_38 1B861E0941ED374899F4D29C48C48C5D,),PersistentGuid=00000000000000000000000000000000,bHidden=False,bNotConnectable=False,bDefaultValueIsReadOnly=False,bDefaultValueIsIgnored=False,bAdvancedView=False,)
+			End Object
+			Begin Object Class=MaterialGraphNode Name="MaterialGraphNode_40"
+			   Begin Object Class=MaterialExpressionTextureCoordinate Name="MaterialExpressionTextureCoordinate_2"
+			   End Object
+			   Begin Object Name="MaterialExpressionTextureCoordinate_2"
+				  MaterialExpressionEditorX=-2080
+				  MaterialExpressionEditorY=-320
+				  MaterialExpressionGuid=F1C1E2C24651DC43FD6AD5B0AF633576
+				  Material=PreviewMaterial'/Engine/Transient.HeightSim'
+			   End Object
+			   MaterialExpression=MaterialExpressionTextureCoordinate'MaterialExpressionTextureCoordinate_2'
+			   NodePosX=-2080
+			   NodePosY=-320
+			   NodeGuid=14DC8A96442661C62C4E2D8CCC6D7BC6
+			   CustomProperties Pin (PinId=367F1A0A4683231CCBEA8ABDC38CC396,PinName="Output",PinFriendlyName=" ",Direction="EGPD_Output",PinType.PinSubCategoryObject=None,PinType.PinSubCategoryMemberReference=(),PinType.bIsArray=False,PinType.bIsReference=False,PinType.bIsConst=False,PinType.bIsWeakPointer=False,LinkedTo=(MaterialGraphNode_39 5C5EC3DB4CC914962FC7E3BD10047D43,),PersistentGuid=00000000000000000000000000000000,bHidden=False,bNotConnectable=False,bDefaultValueIsReadOnly=False,bDefaultValueIsIgnored=False,bAdvancedView=False,)
+			End Object
+			Begin Object Class=MaterialGraphNode Name="MaterialGraphNode_41"
+			   Begin Object Class=MaterialExpressionConstant2Vector Name="MaterialExpressionConstant2Vector_2"
+			   End Object
+			   Begin Object Name="MaterialExpressionConstant2Vector_2"
+				  R=1.000000
+				  MaterialExpressionEditorX=-2192
+				  MaterialExpressionEditorY=-128
+				  MaterialExpressionGuid=F611381D424E60EA8824DA9E1A50EFEB
+				  Material=PreviewMaterial'/Engine/Transient.HeightSim'
+			   End Object
+			   MaterialExpression=MaterialExpressionConstant2Vector'MaterialExpressionConstant2Vector_2'
+			   NodePosX=-2192
+			   NodePosY=-128
+			   NodeGuid=6407310F45C901CF8E2BCFAE84BB05F6
+			   CustomProperties Pin (PinId=B72141BD46005B6680C210816BEEB6B1,PinName="Output",PinFriendlyName=" ",Direction="EGPD_Output",PinType.PinSubCategoryObject=None,PinType.PinSubCategoryMemberReference=(),PinType.bIsArray=False,PinType.bIsReference=False,PinType.bIsConst=False,PinType.bIsWeakPointer=False,LinkedTo=(MaterialGraphNode_16 8C0F76D241F89C6DCA6B3E9FD3BA69CD,),PersistentGuid=00000000000000000000000000000000,bHidden=False,bNotConnectable=False,bDefaultValueIsReadOnly=False,bDefaultValueIsIgnored=False,bAdvancedView=False,)
+			End Object
+			Begin Object Class=MaterialGraphNode Name="MaterialGraphNode_42"
+			   Begin Object Class=MaterialExpressionAdd Name="MaterialExpressionAdd_6"
+			   End Object
+			   Begin Object Name="MaterialExpressionAdd_6"
+				  A=(Expression=MaterialExpressionTextureSampleParameter2D'MaterialGraphNode_46.MaterialExpressionTextureSampleParameter2D_19',OutputIndex=1,Mask=1,MaskR=1)
+				  B=(Expression=MaterialExpressionTextureSampleParameter2D'MaterialGraphNode_38.MaterialExpressionTextureSampleParameter2D_18',OutputIndex=1,Mask=1,MaskR=1)
+				  MaterialExpressionEditorX=-1600
+				  MaterialExpressionEditorY=-384
+				  MaterialExpressionGuid=7C09BC9848A1D28AA8F90791329F6B40
+				  Material=PreviewMaterial'/Engine/Transient.HeightSim'
+			   End Object
+			   MaterialExpression=MaterialExpressionAdd'MaterialExpressionAdd_6'
+			   NodePosX=-1600
+			   NodePosY=-384
+			   NodeGuid=567BB1DF4A18792F1CA84B894BD2030B
+			   CustomProperties Pin (PinId=8572ED124D814B37D6DDA7A3755427E9,PinName="A",PinType.PinCategory="optional",PinType.PinSubCategoryObject=None,PinType.PinSubCategoryMemberReference=(),PinType.bIsArray=False,PinType.bIsReference=False,PinType.bIsConst=False,PinType.bIsWeakPointer=False,LinkedTo=(MaterialGraphNode_46 BACA0BAB446BC16D64B6E580555CA5D7,),PersistentGuid=00000000000000000000000000000000,bHidden=False,bNotConnectable=False,bDefaultValueIsReadOnly=False,bDefaultValueIsIgnored=False,bAdvancedView=False,)
+			   CustomProperties Pin (PinId=6A11559A47452394DE46EB95036FB2D5,PinName="B",PinType.PinCategory="optional",PinType.PinSubCategoryObject=None,PinType.PinSubCategoryMemberReference=(),PinType.bIsArray=False,PinType.bIsReference=False,PinType.bIsConst=False,PinType.bIsWeakPointer=False,LinkedTo=(MaterialGraphNode_38 BBF439FB451DE0A0AC85FA9F9D45895C,),PersistentGuid=00000000000000000000000000000000,bHidden=False,bNotConnectable=False,bDefaultValueIsReadOnly=False,bDefaultValueIsIgnored=False,bAdvancedView=False,)
+			   CustomProperties Pin (PinId=EA5FD970428A72F86BD8A0B7F397B30E,PinName="Output",PinFriendlyName=" ",Direction="EGPD_Output",PinType.PinSubCategoryObject=None,PinType.PinSubCategoryMemberReference=(),PinType.bIsArray=False,PinType.bIsReference=False,PinType.bIsConst=False,PinType.bIsWeakPointer=False,LinkedTo=(MaterialGraphNode_28 BB7B4EAB4EDB7DDEB9548E9EBE9E2B2D,),PersistentGuid=00000000000000000000000000000000,bHidden=False,bNotConnectable=False,bDefaultValueIsReadOnly=False,bDefaultValueIsIgnored=False,bAdvancedView=False,)
+			End Object
+			Begin Object Class=MaterialGraphNode Name="MaterialGraphNode_43"
+			   Begin Object Class=MaterialExpressionAdd Name="MaterialExpressionAdd_7"
+			   End Object
+			   Begin Object Name="MaterialExpressionAdd_7"
+				  A=(Expression=MaterialExpressionTextureCoordinate'MaterialGraphNode_44.MaterialExpressionTextureCoordinate_3')
+				  B=(Expression=MaterialExpressionMultiply'MaterialGraphNode_17.MaterialExpressionMultiply_11')
+				  MaterialExpressionEditorX=-1936
+				  MaterialExpressionEditorY=-528
+				  MaterialExpressionGuid=EF5DFD284086D3D0FDC39A86FEE5E825
+				  Material=PreviewMaterial'/Engine/Transient.HeightSim'
+			   End Object
+			   MaterialExpression=MaterialExpressionAdd'MaterialExpressionAdd_7'
+			   NodePosX=-1936
+			   NodePosY=-528
+			   NodeGuid=A3AFAC63478EA5C7B7BAD39275A4B3A7
+			   CustomProperties Pin (PinId=8BDD1966431DA006B23674B206989D78,PinName="A",PinType.PinCategory="optional",PinType.PinSubCategoryObject=None,PinType.PinSubCategoryMemberReference=(),PinType.bIsArray=False,PinType.bIsReference=False,PinType.bIsConst=False,PinType.bIsWeakPointer=False,LinkedTo=(MaterialGraphNode_44 8969850648806C68A5877181BDF7077C,),PersistentGuid=00000000000000000000000000000000,bHidden=False,bNotConnectable=False,bDefaultValueIsReadOnly=False,bDefaultValueIsIgnored=False,bAdvancedView=False,)
+			   CustomProperties Pin (PinId=E936B7EF4C08A3F0F256FF95563C7509,PinName="B",PinType.PinCategory="optional",PinType.PinSubCategoryObject=None,PinType.PinSubCategoryMemberReference=(),PinType.bIsArray=False,PinType.bIsReference=False,PinType.bIsConst=False,PinType.bIsWeakPointer=False,LinkedTo=(MaterialGraphNode_17 573CDBF54E19966A3470138C02D5DBF8,),PersistentGuid=00000000000000000000000000000000,bHidden=False,bNotConnectable=False,bDefaultValueIsReadOnly=False,bDefaultValueIsIgnored=False,bAdvancedView=False,)
+			   CustomProperties Pin (PinId=9674E90D4C500C27D22F0EA7972074D7,PinName="Output",PinFriendlyName=" ",Direction="EGPD_Output",PinType.PinSubCategoryObject=None,PinType.PinSubCategoryMemberReference=(),PinType.bIsArray=False,PinType.bIsReference=False,PinType.bIsConst=False,PinType.bIsWeakPointer=False,LinkedTo=(MaterialGraphNode_46 72949FE748F7894D5AD8D0B971970408,),PersistentGuid=00000000000000000000000000000000,bHidden=False,bNotConnectable=False,bDefaultValueIsReadOnly=False,bDefaultValueIsIgnored=False,bAdvancedView=False,)
+			End Object
+			Begin Object Class=MaterialGraphNode Name="MaterialGraphNode_44"
+			   Begin Object Class=MaterialExpressionTextureCoordinate Name="MaterialExpressionTextureCoordinate_3"
+			   End Object
+			   Begin Object Name="MaterialExpressionTextureCoordinate_3"
+				  MaterialExpressionEditorX=-2112
+				  MaterialExpressionEditorY=-560
+				  MaterialExpressionGuid=F1C1E2C24651DC43FD6AD5B0AF633576
+				  Material=PreviewMaterial'/Engine/Transient.HeightSim'
+			   End Object
+			   MaterialExpression=MaterialExpressionTextureCoordinate'MaterialExpressionTextureCoordinate_3'
+			   NodePosX=-2112
+			   NodePosY=-560
+			   NodeGuid=FB8886B041C30FF467D7E1A7D2942D0B
+			   CustomProperties Pin (PinId=8969850648806C68A5877181BDF7077C,PinName="Output",PinFriendlyName=" ",Direction="EGPD_Output",PinType.PinSubCategoryObject=None,PinType.PinSubCategoryMemberReference=(),PinType.bIsArray=False,PinType.bIsReference=False,PinType.bIsConst=False,PinType.bIsWeakPointer=False,LinkedTo=(MaterialGraphNode_43 8BDD1966431DA006B23674B206989D78,),PersistentGuid=00000000000000000000000000000000,bHidden=False,bNotConnectable=False,bDefaultValueIsReadOnly=False,bDefaultValueIsIgnored=False,bAdvancedView=False,)
+			End Object
+			Begin Object Class=MaterialGraphNode Name="MaterialGraphNode_45"
+			   Begin Object Class=MaterialExpressionConstant2Vector Name="MaterialExpressionConstant2Vector_3"
+			   End Object
+			   Begin Object Name="MaterialExpressionConstant2Vector_3"
+				  R=-1.000000
+				  MaterialExpressionEditorX=-2192
+				  MaterialExpressionEditorY=-384
+				  MaterialExpressionGuid=F611381D424E60EA8824DA9E1A50EFEB
+				  Material=PreviewMaterial'/Engine/Transient.HeightSim'
+			   End Object
+			   MaterialExpression=MaterialExpressionConstant2Vector'MaterialExpressionConstant2Vector_3'
+			   NodePosX=-2192
+			   NodePosY=-384
+			   NodeGuid=058CEF764CDDF897E2BC7694CC391588
+			   CustomProperties Pin (PinId=150FAB2F48D1DD64E87EC4AFA4C3D9EC,PinName="Output",PinFriendlyName=" ",Direction="EGPD_Output",PinType.PinSubCategoryObject=None,PinType.PinSubCategoryMemberReference=(),PinType.bIsArray=False,PinType.bIsReference=False,PinType.bIsConst=False,PinType.bIsWeakPointer=False,LinkedTo=(MaterialGraphNode_17 9E785A3944293393B058D6A68927A4CF,),PersistentGuid=00000000000000000000000000000000,bHidden=False,bNotConnectable=False,bDefaultValueIsReadOnly=False,bDefaultValueIsIgnored=False,bAdvancedView=False,)
+			End Object
+			Begin Object Class=MaterialGraphNode Name="MaterialGraphNode_46"
+			   Begin Object Class=MaterialExpressionTextureSampleParameter2D Name="MaterialExpressionTextureSampleParameter2D_19"
+			   End Object
+			   Begin Object Name="MaterialExpressionTextureSampleParameter2D_19"
+				  ParameterName="PreviousHeight1"
+				  ExpressionGUID=113561E34B36349BF7C41298386B1FE4
+				  Coordinates=(Expression=MaterialExpressionAdd'MaterialGraphNode_43.MaterialExpressionAdd_7')
+				  MaterialExpressionEditorX=-1792
+				  MaterialExpressionEditorY=-544
+				  MaterialExpressionGuid=8F3166474EF9A6F6A6825185B78360BE
+				  Material=PreviewMaterial'/Engine/Transient.HeightSim'
+			   End Object
+			   MaterialExpression=MaterialExpressionTextureSampleParameter2D'MaterialExpressionTextureSampleParameter2D_19'
+			   NodePosX=-1792
+			   NodePosY=-544
+			   bCanRenameNode=True
+			   NodeGuid=83400B3146917BE83D12A194A82E33E6
+			   CustomProperties Pin (PinId=72949FE748F7894D5AD8D0B971970408,PinName="UVs",PinType.PinCategory="optional",PinType.PinSubCategoryObject=None,PinType.PinSubCategoryMemberReference=(),PinType.bIsArray=False,PinType.bIsReference=False,PinType.bIsConst=False,PinType.bIsWeakPointer=False,LinkedTo=(MaterialGraphNode_43 9674E90D4C500C27D22F0EA7972074D7,),PersistentGuid=00000000000000000000000000000000,bHidden=False,bNotConnectable=False,bDefaultValueIsReadOnly=False,bDefaultValueIsIgnored=False,bAdvancedView=False,)
+			   CustomProperties Pin (PinId=1A6CF9C44C562FABC4286492E467F6E7,PinName="Tex",PinType.PinCategory="optional",PinType.PinSubCategoryObject=None,PinType.PinSubCategoryMemberReference=(),PinType.bIsArray=False,PinType.bIsReference=False,PinType.bIsConst=False,PinType.bIsWeakPointer=False,PersistentGuid=00000000000000000000000000000000,bHidden=False,bNotConnectable=False,bDefaultValueIsReadOnly=False,bDefaultValueIsIgnored=False,bAdvancedView=False,)
+			   CustomProperties Pin (PinId=5FEAEDCA4870F2030447D19633F6DB13,PinName="Output",PinFriendlyName=" ",Direction="EGPD_Output",PinType.PinCategory="mask",PinType.PinSubCategoryObject=None,PinType.PinSubCategoryMemberReference=(),PinType.bIsArray=False,PinType.bIsReference=False,PinType.bIsConst=False,PinType.bIsWeakPointer=False,PersistentGuid=00000000000000000000000000000000,bHidden=False,bNotConnectable=False,bDefaultValueIsReadOnly=False,bDefaultValueIsIgnored=False,bAdvancedView=False,)
+			   CustomProperties Pin (PinId=BACA0BAB446BC16D64B6E580555CA5D7,PinName="Output2",PinFriendlyName=" ",Direction="EGPD_Output",PinType.PinCategory="mask",PinType.PinSubCategory="red",PinType.PinSubCategoryObject=None,PinType.PinSubCategoryMemberReference=(),PinType.bIsArray=False,PinType.bIsReference=False,PinType.bIsConst=False,PinType.bIsWeakPointer=False,LinkedTo=(MaterialGraphNode_42 8572ED124D814B37D6DDA7A3755427E9,),PersistentGuid=00000000000000000000000000000000,bHidden=False,bNotConnectable=False,bDefaultValueIsReadOnly=False,bDefaultValueIsIgnored=False,bAdvancedView=False,)
+			   CustomProperties Pin (PinId=BEECF18A4D76DF6046E91881B2F88CB8,PinName="Output3",PinFriendlyName=" ",Direction="EGPD_Output",PinType.PinCategory="mask",PinType.PinSubCategory="green",PinType.PinSubCategoryObject=None,PinType.PinSubCategoryMemberReference=(),PinType.bIsArray=False,PinType.bIsReference=False,PinType.bIsConst=False,PinType.bIsWeakPointer=False,PersistentGuid=00000000000000000000000000000000,bHidden=False,bNotConnectable=False,bDefaultValueIsReadOnly=False,bDefaultValueIsIgnored=False,bAdvancedView=False,)
+			   CustomProperties Pin (PinId=FB10889544605B91E98EC78FEE9926F7,PinName="Output4",PinFriendlyName=" ",Direction="EGPD_Output",PinType.PinCategory="mask",PinType.PinSubCategory="blue",PinType.PinSubCategoryObject=None,PinType.PinSubCategoryMemberReference=(),PinType.bIsArray=False,PinType.bIsReference=False,PinType.bIsConst=False,PinType.bIsWeakPointer=False,PersistentGuid=00000000000000000000000000000000,bHidden=False,bNotConnectable=False,bDefaultValueIsReadOnly=False,bDefaultValueIsIgnored=False,bAdvancedView=False,)
+			   CustomProperties Pin (PinId=ED761DCA43592672E68224B312AC4481,PinName="Output5",PinFriendlyName=" ",Direction="EGPD_Output",PinType.PinCategory="mask",PinType.PinSubCategory="alpha",PinType.PinSubCategoryObject=None,PinType.PinSubCategoryMemberReference=(),PinType.bIsArray=False,PinType.bIsReference=False,PinType.bIsConst=False,PinType.bIsWeakPointer=False,PersistentGuid=00000000000000000000000000000000,bHidden=False,bNotConnectable=False,bDefaultValueIsReadOnly=False,bDefaultValueIsIgnored=False,bAdvancedView=False,)
+			End Object
+			Begin Object Class=MaterialGraphNode Name="MaterialGraphNode_47"
+			   Begin Object Class=MaterialExpressionSubtract Name="MaterialExpressionSubtract_1"
+			   End Object
+			   Begin Object Name="MaterialExpressionSubtract_1"
+				  A=(Expression=MaterialExpressionMultiply'MaterialGraphNode_21.MaterialExpressionMultiply_12')
+				  B=(Expression=MaterialExpressionTextureSampleParameter2D'MaterialGraphNode_20.MaterialExpressionTextureSampleParameter2D_14',OutputIndex=1,Mask=1,MaskR=1)
+				  MaterialExpressionEditorX=-576
+				  MaterialExpressionEditorY=-208
+				  MaterialExpressionGuid=290D1BC24729BE42230F35BDDD9EDBDD
+				  Material=PreviewMaterial'/Engine/Transient.HeightSim'
+			   End Object
+			   MaterialExpression=MaterialExpressionSubtract'MaterialExpressionSubtract_1'
+			   NodePosX=-576
+			   NodePosY=-208
+			   NodeGuid=5740D385435DCDD8420B8788665382D7
+			   CustomProperties Pin (PinId=836C63F044D0CE95EF676390B28B3921,PinName="A",PinType.PinCategory="optional",PinType.PinSubCategoryObject=None,PinType.PinSubCategoryMemberReference=(),PinType.bIsArray=False,PinType.bIsReference=False,PinType.bIsConst=False,PinType.bIsWeakPointer=False,LinkedTo=(MaterialGraphNode_21 F7B1B2994F6F5A2E84B9EFA1485EFD4A,),PersistentGuid=00000000000000000000000000000000,bHidden=False,bNotConnectable=False,bDefaultValueIsReadOnly=False,bDefaultValueIsIgnored=False,bAdvancedView=False,)
+			   CustomProperties Pin (PinId=C91C84A64517D7F47FA7B39B116852A6,PinName="B",PinType.PinCategory="optional",PinType.PinSubCategoryObject=None,PinType.PinSubCategoryMemberReference=(),PinType.bIsArray=False,PinType.bIsReference=False,PinType.bIsConst=False,PinType.bIsWeakPointer=False,LinkedTo=(MaterialGraphNode_20 9E1D6AE64F3486451D911EA6B84ECC06,),PersistentGuid=00000000000000000000000000000000,bHidden=False,bNotConnectable=False,bDefaultValueIsReadOnly=False,bDefaultValueIsIgnored=False,bAdvancedView=False,)
+			   CustomProperties Pin (PinId=297B10D944D15582C2242E9387BD8FA9,PinName="Output",PinFriendlyName=" ",Direction="EGPD_Output",PinType.PinSubCategoryObject=None,PinType.PinSubCategoryMemberReference=(),PinType.bIsArray=False,PinType.bIsReference=False,PinType.bIsConst=False,PinType.bIsWeakPointer=False,LinkedTo=(MaterialGraphNode_12 AE4E493242D5E3BA3C1D83A2BF5D2295,),PersistentGuid=00000000000000000000000000000000,bHidden=False,bNotConnectable=False,bDefaultValueIsReadOnly=False,bDefaultValueIsIgnored=False,bAdvancedView=False,)
+			End Object
+			END OBJECT
+	[/PARAM]
+[/OBJECT]
+[/EXCERPT:HeightSim]
+[EXCERPT:WaterMaterial]	
+[OBJECT:Snippet]
+	[PARAMLITERAL:id]
+		WaterMaterial
+	[/PARAMLITERAL]
+	[PARAMLITERAL:type]
+		material
+	[/PARAMLITERAL]
+	[PARAM:image]
+		![](Engine\Rendering\RenderTargets\BlueprintRenderTargets\HowTo\FluidSurface\FS_Water_Material.png)
+	[/PARAM]
+	[PARAM:snippet]
+			BEGIN OBJECT
+				Begin Object Class=MaterialGraphNode Name="MaterialGraphNode_60"
+			   Begin Object Class=MaterialExpressionConstant Name="MaterialExpressionConstant_9"
+			   End Object
+			   Begin Object Name="MaterialExpressionConstant_9"
+				  R=1.000000
+				  MaterialExpressionEditorX=-352
+				  MaterialExpressionGuid=BD13899F48908744D85F3990611CA97C
+				  Material=PreviewMaterial'/Engine/Transient.WaterMaterial'
+			   End Object
+			   MaterialExpression=MaterialExpressionConstant'MaterialExpressionConstant_9'
+			   NodePosX=-352
+			   NodeGuid=01ADF5574FE81EAE21FEFCA30E5F23D2
+			   CustomProperties Pin (PinId=B9006E604FA8EA9C334592B22E756F16,PinName="Output",PinFriendlyName=" ",Direction="EGPD_Output",PinType.PinSubCategoryObject=None,PinType.PinSubCategoryMemberReference=(),PinType.bIsArray=False,PinType.bIsReference=False,PinType.bIsConst=False,PinType.bIsWeakPointer=False,LinkedTo=(MaterialGraphNode_Root_2 7F0A6C7D4947D45D97B9CFBD704317BE,MaterialGraphNode_Root_2 021DB70C4B501876C28FE5BB7DE8174F,),PersistentGuid=00000000000000000000000000000000,bHidden=False,bNotConnectable=False,bDefaultValueIsReadOnly=False,bDefaultValueIsIgnored=False,bAdvancedView=False,)
+			End Object
+			Begin Object Class=MaterialGraphNode Name="MaterialGraphNode_61"
+			   Begin Object Class=MaterialExpressionConstant Name="MaterialExpressionConstant_10"
+			   End Object
+			   Begin Object Name="MaterialExpressionConstant_10"
+				  MaterialExpressionEditorX=-348
+				  MaterialExpressionEditorY=89
+				  MaterialExpressionGuid=119740984BEF4C865F0E369FD833C5BA
+				  Material=PreviewMaterial'/Engine/Transient.WaterMaterial'
+			   End Object
+			   MaterialExpression=MaterialExpressionConstant'MaterialExpressionConstant_10'
+			   NodePosX=-348
+			   NodePosY=89
+			   NodeGuid=FF4379D048BEFF36334487A1AF234F23
+			   CustomProperties Pin (PinId=B2943B95491C0AA1A0B20484A0ED654A,PinName="Output",PinFriendlyName=" ",Direction="EGPD_Output",PinType.PinSubCategoryObject=None,PinType.PinSubCategoryMemberReference=(),PinType.bIsArray=False,PinType.bIsReference=False,PinType.bIsConst=False,PinType.bIsWeakPointer=False,LinkedTo=(MaterialGraphNode_Root_2 00BDA30946B8D573063A2DB2660624DC,),PersistentGuid=00000000000000000000000000000000,bHidden=False,bNotConnectable=False,bDefaultValueIsReadOnly=False,bDefaultValueIsIgnored=False,bAdvancedView=False,)
+			End Object
+			Begin Object Class=MaterialGraphNode Name="MaterialGraphNode_62"
+			   Begin Object Class=MaterialExpressionTextureSampleParameter2D Name="MaterialExpressionTextureSampleParameter2D_6"
+			   End Object
+			   Begin Object Name="MaterialExpressionTextureSampleParameter2D_6"
+				  ParameterName="HeightfieldNormal"
+				  ExpressionGUID=5A0C4B924E41F82F138AE8B43E6E3AF5
+				  Texture=Texture2D'/Game/ExampleContent/BlueprintFluidSurface/DefaultHeightfieldNormal.DefaultHeightfieldNormal'
+				  SamplerType=SAMPLERTYPE_LinearColor
+				  MaterialExpressionEditorX=-400
+				  MaterialExpressionEditorY=176
+				  MaterialExpressionGuid=1237076F46529CF6EB453EA957551186
+				  Material=PreviewMaterial'/Engine/Transient.WaterMaterial'
+				  bCollapsed=True
+			   End Object
+			   MaterialExpression=MaterialExpressionTextureSampleParameter2D'MaterialExpressionTextureSampleParameter2D_6'
+			   NodePosX=-400
+			   NodePosY=176
+			   bCanRenameNode=True
+			   NodeGuid=FE6544424BF3DF67C4AD81B4276CC41C
+			   CustomProperties Pin (PinId=893FAD534E85959D66F66EA33D3E3CA5,PinName="UVs",PinType.PinCategory="optional",PinType.PinSubCategoryObject=None,PinType.PinSubCategoryMemberReference=(),PinType.bIsArray=False,PinType.bIsReference=False,PinType.bIsConst=False,PinType.bIsWeakPointer=False,PersistentGuid=00000000000000000000000000000000,bHidden=False,bNotConnectable=False,bDefaultValueIsReadOnly=False,bDefaultValueIsIgnored=False,bAdvancedView=False,)
+			   CustomProperties Pin (PinId=4B464619467DC6C993DDD08F2C9C42DF,PinName="Tex",PinType.PinCategory="optional",PinType.PinSubCategoryObject=None,PinType.PinSubCategoryMemberReference=(),PinType.bIsArray=False,PinType.bIsReference=False,PinType.bIsConst=False,PinType.bIsWeakPointer=False,PersistentGuid=00000000000000000000000000000000,bHidden=False,bNotConnectable=False,bDefaultValueIsReadOnly=False,bDefaultValueIsIgnored=False,bAdvancedView=False,)
+			   CustomProperties Pin (PinId=65D8F1DC45FC44C455168D8B12D9345F,PinName="Output",PinFriendlyName=" ",Direction="EGPD_Output",PinType.PinCategory="mask",PinType.PinSubCategoryObject=None,PinType.PinSubCategoryMemberReference=(),PinType.bIsArray=False,PinType.bIsReference=False,PinType.bIsConst=False,PinType.bIsWeakPointer=False,LinkedTo=(MaterialGraphNode_Root_2 2E7D5AAF4CC55C908D11C0A6BE9CFE87,),PersistentGuid=00000000000000000000000000000000,bHidden=False,bNotConnectable=False,bDefaultValueIsReadOnly=False,bDefaultValueIsIgnored=False,bAdvancedView=False,)
+			   CustomProperties Pin (PinId=E463785A408CF5F9483797AFB4183EE8,PinName="Output2",PinFriendlyName=" ",Direction="EGPD_Output",PinType.PinCategory="mask",PinType.PinSubCategory="red",PinType.PinSubCategoryObject=None,PinType.PinSubCategoryMemberReference=(),PinType.bIsArray=False,PinType.bIsReference=False,PinType.bIsConst=False,PinType.bIsWeakPointer=False,PersistentGuid=00000000000000000000000000000000,bHidden=False,bNotConnectable=False,bDefaultValueIsReadOnly=False,bDefaultValueIsIgnored=False,bAdvancedView=False,)
+			   CustomProperties Pin (PinId=B4BFCA5A4E74759EF6B5F4B85568422D,PinName="Output3",PinFriendlyName=" ",Direction="EGPD_Output",PinType.PinCategory="mask",PinType.PinSubCategory="green",PinType.PinSubCategoryObject=None,PinType.PinSubCategoryMemberReference=(),PinType.bIsArray=False,PinType.bIsReference=False,PinType.bIsConst=False,PinType.bIsWeakPointer=False,PersistentGuid=00000000000000000000000000000000,bHidden=False,bNotConnectable=False,bDefaultValueIsReadOnly=False,bDefaultValueIsIgnored=False,bAdvancedView=False,)
+			   CustomProperties Pin (PinId=A5BAFDEC4474F71C856AE987B8139D33,PinName="Output4",PinFriendlyName=" ",Direction="EGPD_Output",PinType.PinCategory="mask",PinType.PinSubCategory="blue",PinType.PinSubCategoryObject=None,PinType.PinSubCategoryMemberReference=(),PinType.bIsArray=False,PinType.bIsReference=False,PinType.bIsConst=False,PinType.bIsWeakPointer=False,PersistentGuid=00000000000000000000000000000000,bHidden=False,bNotConnectable=False,bDefaultValueIsReadOnly=False,bDefaultValueIsIgnored=False,bAdvancedView=False,)
+			   CustomProperties Pin (PinId=ED0A84A64616383766BB45B16A203DE6,PinName="Output5",PinFriendlyName=" ",Direction="EGPD_Output",PinType.PinCategory="mask",PinType.PinSubCategory="alpha",PinType.PinSubCategoryObject=None,PinType.PinSubCategoryMemberReference=(),PinType.bIsArray=False,PinType.bIsReference=False,PinType.bIsConst=False,PinType.bIsWeakPointer=False,PersistentGuid=00000000000000000000000000000000,bHidden=False,bNotConnectable=False,bDefaultValueIsReadOnly=False,bDefaultValueIsIgnored=False,bAdvancedView=False,)
+			End Object
+			Begin Object Class=MaterialGraphNode Name="MaterialGraphNode_63"
+			   Begin Object Class=MaterialExpressionFresnel Name="MaterialExpressionFresnel_3"
+			   End Object
+			   Begin Object Name="MaterialExpressionFresnel_3"
+				  BaseReflectFractionIn=(Expression=MaterialExpressionScalarParameter'MaterialGraphNode_66.MaterialExpressionScalarParameter_10')
+				  MaterialExpressionEditorX=-624
+				  MaterialExpressionEditorY=128
+				  MaterialExpressionGuid=3E39F0A54150DB21FCB3B2AC934DDCCB
+				  Material=PreviewMaterial'/Engine/Transient.WaterMaterial'
+			   End Object
+			   MaterialExpression=MaterialExpressionFresnel'MaterialExpressionFresnel_3'
+			   NodePosX=-624
+			   NodePosY=128
+			   NodeGuid=E6453672434ACA54DBD00D8E68E330E9
+			   CustomProperties Pin (PinId=112A41D044C56DA20A3A238841D1EC58,PinName="ExponentIn",PinType.PinCategory="optional",PinType.PinSubCategoryObject=None,PinType.PinSubCategoryMemberReference=(),PinType.bIsArray=False,PinType.bIsReference=False,PinType.bIsConst=False,PinType.bIsWeakPointer=False,PersistentGuid=00000000000000000000000000000000,bHidden=False,bNotConnectable=False,bDefaultValueIsReadOnly=False,bDefaultValueIsIgnored=False,bAdvancedView=False,)
+			   CustomProperties Pin (PinId=EC3301884D93BD09D46B3EA2C035F4D9,PinName="BaseReflectFractionIn",PinType.PinCategory="optional",PinType.PinSubCategoryObject=None,PinType.PinSubCategoryMemberReference=(),PinType.bIsArray=False,PinType.bIsReference=False,PinType.bIsConst=False,PinType.bIsWeakPointer=False,LinkedTo=(MaterialGraphNode_66 8E1B168A4808C8A990C854982AF51A08,),PersistentGuid=00000000000000000000000000000000,bHidden=False,bNotConnectable=False,bDefaultValueIsReadOnly=False,bDefaultValueIsIgnored=False,bAdvancedView=False,)
+			   CustomProperties Pin (PinId=CAF6DCCF422EA89AFDEEB288D21D7EFF,PinName="Normal",PinType.PinCategory="optional",PinType.PinSubCategoryObject=None,PinType.PinSubCategoryMemberReference=(),PinType.bIsArray=False,PinType.bIsReference=False,PinType.bIsConst=False,PinType.bIsWeakPointer=False,PersistentGuid=00000000000000000000000000000000,bHidden=False,bNotConnectable=False,bDefaultValueIsReadOnly=False,bDefaultValueIsIgnored=False,bAdvancedView=False,)
+			   CustomProperties Pin (PinId=4C7E4EBA499DD92D7B860FBCC63502B4,PinName="Output",PinFriendlyName=" ",Direction="EGPD_Output",PinType.PinSubCategoryObject=None,PinType.PinSubCategoryMemberReference=(),PinType.bIsArray=False,PinType.bIsReference=False,PinType.bIsConst=False,PinType.bIsWeakPointer=False,LinkedTo=(MaterialGraphNode_Root_2 1014AE2243D06A14EADF76A1404DF996,),PersistentGuid=00000000000000000000000000000000,bHidden=False,bNotConnectable=False,bDefaultValueIsReadOnly=False,bDefaultValueIsIgnored=False,bAdvancedView=False,)
+			End Object
+			Begin Object Class=MaterialGraphNode Name="MaterialGraphNode_64"
+			   Begin Object Class=MaterialExpressionAppendVector Name="MaterialExpressionAppendVector_6"
+			   End Object
+			   Begin Object Name="MaterialExpressionAppendVector_6"
+				  A=(Expression=MaterialExpressionAppendVector'MaterialGraphNode_71.MaterialExpressionAppendVector_7')
+				  B=(Expression=MaterialExpressionMultiply'MaterialGraphNode_69.MaterialExpressionMultiply_3')
+				  MaterialExpressionEditorX=-480
+				  MaterialExpressionEditorY=384
+				  MaterialExpressionGuid=C32396324C82B28E46B485B62A23697F
+				  Material=PreviewMaterial'/Engine/Transient.WaterMaterial'
+			   End Object
+			   MaterialExpression=MaterialExpressionAppendVector'MaterialExpressionAppendVector_6'
+			   NodePosX=-480
+			   NodePosY=384
+			   NodeGuid=267E595E42C4993B43FEA693E004D02E
+			   CustomProperties Pin (PinId=0D21E58243C6C7C77F26E7B3A92662EE,PinName="A",PinType.PinCategory="required",PinType.PinSubCategoryObject=None,PinType.PinSubCategoryMemberReference=(),PinType.bIsArray=False,PinType.bIsReference=False,PinType.bIsConst=False,PinType.bIsWeakPointer=False,LinkedTo=(MaterialGraphNode_71 38DF7322459B1B42ACB22AB359E5788D,),PersistentGuid=00000000000000000000000000000000,bHidden=False,bNotConnectable=False,bDefaultValueIsReadOnly=False,bDefaultValueIsIgnored=False,bAdvancedView=False,)
+			   CustomProperties Pin (PinId=9CD7F0834D0F392098D6B7BF36EB5C3D,PinName="B",PinType.PinCategory="required",PinType.PinSubCategoryObject=None,PinType.PinSubCategoryMemberReference=(),PinType.bIsArray=False,PinType.bIsReference=False,PinType.bIsConst=False,PinType.bIsWeakPointer=False,LinkedTo=(MaterialGraphNode_69 6766DD6A4C02E9BC055F24BB7824D6CC,),PersistentGuid=00000000000000000000000000000000,bHidden=False,bNotConnectable=False,bDefaultValueIsReadOnly=False,bDefaultValueIsIgnored=False,bAdvancedView=False,)
+			   CustomProperties Pin (PinId=F87D1BA8407D97684FF66282E0CC6A57,PinName="Output",PinFriendlyName=" ",Direction="EGPD_Output",PinType.PinSubCategoryObject=None,PinType.PinSubCategoryMemberReference=(),PinType.bIsArray=False,PinType.bIsReference=False,PinType.bIsConst=False,PinType.bIsWeakPointer=False,LinkedTo=(MaterialGraphNode_Root_2 D3DBC8E446DA3DFE598078974970A967,),PersistentGuid=00000000000000000000000000000000,bHidden=False,bNotConnectable=False,bDefaultValueIsReadOnly=False,bDefaultValueIsIgnored=False,bAdvancedView=False,)
+			End Object
+			Begin Object Class=MaterialGraphNode Name="MaterialGraphNode_65"
+			   Begin Object Class=MaterialExpressionScalarParameter Name="MaterialExpressionScalarParameter_9"
+			   End Object
+			   Begin Object Name="MaterialExpressionScalarParameter_9"
+				  DefaultValue=2.000000
+				  ParameterName="RefractionScale"
+				  ExpressionGUID=EBE99F6345C6BF78C002DCA1325673E0
+				  MaterialExpressionEditorX=-240
+				  MaterialExpressionEditorY=512
+				  MaterialExpressionGuid=C88601E34A9E8DC112FED7922B9DA8E2
+				  Material=PreviewMaterial'/Engine/Transient.WaterMaterial'
+			   End Object
+			   MaterialExpression=MaterialExpressionScalarParameter'MaterialExpressionScalarParameter_9'
+			   NodePosX=-240
+			   NodePosY=512
+			   bCanRenameNode=True
+			   NodeGuid=9D78CC03454FD51F20A87BBF7AAF2E5D
+			   CustomProperties Pin (PinId=A516DDF24CA7B9EBB3FD6A826951AC4D,PinName="Output",PinFriendlyName=" ",Direction="EGPD_Output",PinType.PinSubCategoryObject=None,PinType.PinSubCategoryMemberReference=(),PinType.bIsArray=False,PinType.bIsReference=False,PinType.bIsConst=False,PinType.bIsWeakPointer=False,LinkedTo=(MaterialGraphNode_Root_2 0BE38E1142FD35FFB20C138FDA562F83,),PersistentGuid=00000000000000000000000000000000,bHidden=False,bNotConnectable=False,bDefaultValueIsReadOnly=False,bDefaultValueIsIgnored=False,bAdvancedView=False,)
+			End Object
+			Begin Object Class=MaterialGraphNode Name="MaterialGraphNode_66"
+			   Begin Object Class=MaterialExpressionScalarParameter Name="MaterialExpressionScalarParameter_10"
+			   End Object
+			   Begin Object Name="MaterialExpressionScalarParameter_10"
+				  DefaultValue=0.800000
+				  ParameterName="Opacity"
+				  ExpressionGUID=EDB745FB4D2AFB05123455BE44C9F12D
+				  MaterialExpressionEditorX=-883
+				  MaterialExpressionEditorY=239
+				  MaterialExpressionGuid=05D5CEC740828D5694C0019A6B5F29A6
+				  Material=PreviewMaterial'/Engine/Transient.WaterMaterial'
+			   End Object
+			   MaterialExpression=MaterialExpressionScalarParameter'MaterialExpressionScalarParameter_10'
+			   NodePosX=-883
+			   NodePosY=239
+			   bCanRenameNode=True
+			   NodeGuid=34C2EF99488BAD95F45F978F8D1D0846
+			   CustomProperties Pin (PinId=8E1B168A4808C8A990C854982AF51A08,PinName="Output",PinFriendlyName=" ",Direction="EGPD_Output",PinType.PinSubCategoryObject=None,PinType.PinSubCategoryMemberReference=(),PinType.bIsArray=False,PinType.bIsReference=False,PinType.bIsConst=False,PinType.bIsWeakPointer=False,LinkedTo=(MaterialGraphNode_63 EC3301884D93BD09D46B3EA2C035F4D9,),PersistentGuid=00000000000000000000000000000000,bHidden=False,bNotConnectable=False,bDefaultValueIsReadOnly=False,bDefaultValueIsIgnored=False,bAdvancedView=False,)
+			End Object
+			Begin Object Class=MaterialGraphNode Name="MaterialGraphNode_67"
+			   Begin Object Class=MaterialExpressionTextureSampleParameter2D Name="MaterialExpressionTextureSampleParameter2D_7"
+			   End Object
+			   Begin Object Name="MaterialExpressionTextureSampleParameter2D_7"
+				  ParameterName="Heightfield"
+				  ExpressionGUID=83D44DDC41F0BEFEFE9496A95528DB7F
+				  Texture=Texture2D'/Game/ExampleContent/BlueprintFluidSurface/DefaultHeightfieldNormal.DefaultHeightfieldNormal'
+				  SamplerType=SAMPLERTYPE_LinearColor
+				  MaterialExpressionEditorX=-1072
+				  MaterialExpressionEditorY=288
+				  MaterialExpressionGuid=BD05F92B4EB1F08272537AB8E283F06A
+				  Material=PreviewMaterial'/Engine/Transient.WaterMaterial'
+				  bCollapsed=True
+			   End Object
+			   MaterialExpression=MaterialExpressionTextureSampleParameter2D'MaterialExpressionTextureSampleParameter2D_7'
+			   NodePosX=-1072
+			   NodePosY=288
+			   bCanRenameNode=True
+			   NodeGuid=B787AB4C4D189B23B0B4408A36EC5A57
+			   CustomProperties Pin (PinId=6964551F4CA421CC94A5E79E1E68B8AF,PinName="UVs",PinType.PinCategory="optional",PinType.PinSubCategoryObject=None,PinType.PinSubCategoryMemberReference=(),PinType.bIsArray=False,PinType.bIsReference=False,PinType.bIsConst=False,PinType.bIsWeakPointer=False,PersistentGuid=00000000000000000000000000000000,bHidden=False,bNotConnectable=False,bDefaultValueIsReadOnly=False,bDefaultValueIsIgnored=False,bAdvancedView=False,)
+			   CustomProperties Pin (PinId=29E48CFA4C080F1CB572EC949500A5C5,PinName="Tex",PinType.PinCategory="optional",PinType.PinSubCategoryObject=None,PinType.PinSubCategoryMemberReference=(),PinType.bIsArray=False,PinType.bIsReference=False,PinType.bIsConst=False,PinType.bIsWeakPointer=False,PersistentGuid=00000000000000000000000000000000,bHidden=False,bNotConnectable=False,bDefaultValueIsReadOnly=False,bDefaultValueIsIgnored=False,bAdvancedView=False,)
+			   CustomProperties Pin (PinId=05865EDA4EF4A60FF14703918FFDDBD3,PinName="Output",PinFriendlyName=" ",Direction="EGPD_Output",PinType.PinCategory="mask",PinType.PinSubCategoryObject=None,PinType.PinSubCategoryMemberReference=(),PinType.bIsArray=False,PinType.bIsReference=False,PinType.bIsConst=False,PinType.bIsWeakPointer=False,PersistentGuid=00000000000000000000000000000000,bHidden=False,bNotConnectable=False,bDefaultValueIsReadOnly=False,bDefaultValueIsIgnored=False,bAdvancedView=False,)
+			   CustomProperties Pin (PinId=E4C0D9F341000C6878BCFDA0DDE682F5,PinName="Output2",PinFriendlyName=" ",Direction="EGPD_Output",PinType.PinCategory="mask",PinType.PinSubCategory="red",PinType.PinSubCategoryObject=None,PinType.PinSubCategoryMemberReference=(),PinType.bIsArray=False,PinType.bIsReference=False,PinType.bIsConst=False,PinType.bIsWeakPointer=False,LinkedTo=(MaterialGraphNode_69 E92D6F92425F1A509E3B949E94DD125A,),PersistentGuid=00000000000000000000000000000000,bHidden=False,bNotConnectable=False,bDefaultValueIsReadOnly=False,bDefaultValueIsIgnored=False,bAdvancedView=False,)
+			   CustomProperties Pin (PinId=C31E11B74E29A24793C7AB8B9951E397,PinName="Output3",PinFriendlyName=" ",Direction="EGPD_Output",PinType.PinCategory="mask",PinType.PinSubCategory="green",PinType.PinSubCategoryObject=None,PinType.PinSubCategoryMemberReference=(),PinType.bIsArray=False,PinType.bIsReference=False,PinType.bIsConst=False,PinType.bIsWeakPointer=False,PersistentGuid=00000000000000000000000000000000,bHidden=False,bNotConnectable=False,bDefaultValueIsReadOnly=False,bDefaultValueIsIgnored=False,bAdvancedView=False,)
+			   CustomProperties Pin (PinId=CF42910F4F59023877478EA6AB54170A,PinName="Output4",PinFriendlyName=" ",Direction="EGPD_Output",PinType.PinCategory="mask",PinType.PinSubCategory="blue",PinType.PinSubCategoryObject=None,PinType.PinSubCategoryMemberReference=(),PinType.bIsArray=False,PinType.bIsReference=False,PinType.bIsConst=False,PinType.bIsWeakPointer=False,PersistentGuid=00000000000000000000000000000000,bHidden=False,bNotConnectable=False,bDefaultValueIsReadOnly=False,bDefaultValueIsIgnored=False,bAdvancedView=False,)
+			   CustomProperties Pin (PinId=732306FC402A82C95F929BBE4D8765C8,PinName="Output5",PinFriendlyName=" ",Direction="EGPD_Output",PinType.PinCategory="mask",PinType.PinSubCategory="alpha",PinType.PinSubCategoryObject=None,PinType.PinSubCategoryMemberReference=(),PinType.bIsArray=False,PinType.bIsReference=False,PinType.bIsConst=False,PinType.bIsWeakPointer=False,PersistentGuid=00000000000000000000000000000000,bHidden=False,bNotConnectable=False,bDefaultValueIsReadOnly=False,bDefaultValueIsIgnored=False,bAdvancedView=False,)
+			End Object
+			Begin Object Class=MaterialGraphNode Name="MaterialGraphNode_68"
+			   Begin Object Class=MaterialExpressionScalarParameter Name="MaterialExpressionScalarParameter_11"
+			   End Object
+			   Begin Object Name="MaterialExpressionScalarParameter_11"
+				  DefaultValue=10.000000
+				  ParameterName="HeightScale"
+				  ExpressionGUID=4A4E6B2A4353ECD525ECD5A679DDCF17
+				  MaterialExpressionEditorX=-992
+				  MaterialExpressionEditorY=496
+				  MaterialExpressionGuid=8365D17A4B4AED6D11D7DAADCCD07F11
+				  Material=PreviewMaterial'/Engine/Transient.WaterMaterial'
+			   End Object
+			   MaterialExpression=MaterialExpressionScalarParameter'MaterialExpressionScalarParameter_11'
+			   NodePosX=-992
+			   NodePosY=496
+			   bCanRenameNode=True
+			   NodeGuid=0D78ED1A45A7525223DBE5957378FB92
+			   CustomProperties Pin (PinId=3E29E7294A7BA209A648EEADBFAD8D47,PinName="Output",PinFriendlyName=" ",Direction="EGPD_Output",PinType.PinSubCategoryObject=None,PinType.PinSubCategoryMemberReference=(),PinType.bIsArray=False,PinType.bIsReference=False,PinType.bIsConst=False,PinType.bIsWeakPointer=False,LinkedTo=(MaterialGraphNode_69 14BAEB944E14E9EC9FE5F6A49D476D31,),PersistentGuid=00000000000000000000000000000000,bHidden=False,bNotConnectable=False,bDefaultValueIsReadOnly=False,bDefaultValueIsIgnored=False,bAdvancedView=False,)
+			End Object
+			Begin Object Class=MaterialGraphNode Name="MaterialGraphNode_69"
+			   Begin Object Class=MaterialExpressionMultiply Name="MaterialExpressionMultiply_3"
+			   End Object
+			   Begin Object Name="MaterialExpressionMultiply_3"
+				  A=(Expression=MaterialExpressionTextureSampleParameter2D'MaterialGraphNode_67.MaterialExpressionTextureSampleParameter2D_7',OutputIndex=1,Mask=1,MaskR=1)
+				  B=(Expression=MaterialExpressionScalarParameter'MaterialGraphNode_68.MaterialExpressionScalarParameter_11')
+				  MaterialExpressionEditorX=-800
+				  MaterialExpressionEditorY=432
+				  MaterialExpressionGuid=A30B87704A8B72865FC4CF9471068EA7
+				  Material=PreviewMaterial'/Engine/Transient.WaterMaterial'
+			   End Object
+			   MaterialExpression=MaterialExpressionMultiply'MaterialExpressionMultiply_3'
+			   NodePosX=-800
+			   NodePosY=432
+			   NodeGuid=3D71AD774E138E5BEB6CEA8D50192B72
+			   CustomProperties Pin (PinId=E92D6F92425F1A509E3B949E94DD125A,PinName="A",PinType.PinCategory="optional",PinType.PinSubCategoryObject=None,PinType.PinSubCategoryMemberReference=(),PinType.bIsArray=False,PinType.bIsReference=False,PinType.bIsConst=False,PinType.bIsWeakPointer=False,LinkedTo=(MaterialGraphNode_67 E4C0D9F341000C6878BCFDA0DDE682F5,),PersistentGuid=00000000000000000000000000000000,bHidden=False,bNotConnectable=False,bDefaultValueIsReadOnly=False,bDefaultValueIsIgnored=False,bAdvancedView=False,)
+			   CustomProperties Pin (PinId=14BAEB944E14E9EC9FE5F6A49D476D31,PinName="B",PinType.PinCategory="optional",PinType.PinSubCategoryObject=None,PinType.PinSubCategoryMemberReference=(),PinType.bIsArray=False,PinType.bIsReference=False,PinType.bIsConst=False,PinType.bIsWeakPointer=False,LinkedTo=(MaterialGraphNode_68 3E29E7294A7BA209A648EEADBFAD8D47,),PersistentGuid=00000000000000000000000000000000,bHidden=False,bNotConnectable=False,bDefaultValueIsReadOnly=False,bDefaultValueIsIgnored=False,bAdvancedView=False,)
+			   CustomProperties Pin (PinId=6766DD6A4C02E9BC055F24BB7824D6CC,PinName="Output",PinFriendlyName=" ",Direction="EGPD_Output",PinType.PinSubCategoryObject=None,PinType.PinSubCategoryMemberReference=(),PinType.bIsArray=False,PinType.bIsReference=False,PinType.bIsConst=False,PinType.bIsWeakPointer=False,LinkedTo=(MaterialGraphNode_64 9CD7F0834D0F392098D6B7BF36EB5C3D,),PersistentGuid=00000000000000000000000000000000,bHidden=False,bNotConnectable=False,bDefaultValueIsReadOnly=False,bDefaultValueIsIgnored=False,bAdvancedView=False,)
+			End Object
+			Begin Object Class=MaterialGraphNode Name="MaterialGraphNode_70"
+			   Begin Object Class=MaterialExpressionConstant Name="MaterialExpressionConstant_11"
+			   End Object
+			   Begin Object Name="MaterialExpressionConstant_11"
+				  MaterialExpressionEditorX=-768
+				  MaterialExpressionEditorY=368
+				  MaterialExpressionGuid=7D7C334346B0539AD2AF248F7F95C517
+				  Material=PreviewMaterial'/Engine/Transient.WaterMaterial'
+			   End Object
+			   MaterialExpression=MaterialExpressionConstant'MaterialExpressionConstant_11'
+			   NodePosX=-768
+			   NodePosY=368
+			   NodeGuid=160AEF0B413019BB7C2720B01915DEF8
+			   CustomProperties Pin (PinId=987C42784FF10C1B6C5F66A0CC164817,PinName="Output",PinFriendlyName=" ",Direction="EGPD_Output",PinType.PinSubCategoryObject=None,PinType.PinSubCategoryMemberReference=(),PinType.bIsArray=False,PinType.bIsReference=False,PinType.bIsConst=False,PinType.bIsWeakPointer=False,LinkedTo=(MaterialGraphNode_71 E4B2A48B405C5BE1ECA917B7DC4BCDA7,MaterialGraphNode_71 207CC92F40241634F7D74290846DA5E3,),PersistentGuid=00000000000000000000000000000000,bHidden=False,bNotConnectable=False,bDefaultValueIsReadOnly=False,bDefaultValueIsIgnored=False,bAdvancedView=False,)
+			End Object
+			Begin Object Class=MaterialGraphNode Name="MaterialGraphNode_71"
+			   Begin Object Class=MaterialExpressionAppendVector Name="MaterialExpressionAppendVector_7"
+			   End Object
+			   Begin Object Name="MaterialExpressionAppendVector_7"
+				  A=(Expression=MaterialExpressionConstant'MaterialGraphNode_70.MaterialExpressionConstant_11')
+				  B=(Expression=MaterialExpressionConstant'MaterialGraphNode_70.MaterialExpressionConstant_11')
+				  MaterialExpressionEditorX=-656
+				  MaterialExpressionEditorY=352
+				  MaterialExpressionGuid=C32396324C82B28E46B485B62A23697F
+				  Material=PreviewMaterial'/Engine/Transient.WaterMaterial'
+			   End Object
+			   MaterialExpression=MaterialExpressionAppendVector'MaterialExpressionAppendVector_7'
+			   NodePosX=-656
+			   NodePosY=352
+			   NodeGuid=D5A877264C613DDA70C47D84E5E33120
+			   CustomProperties Pin (PinId=E4B2A48B405C5BE1ECA917B7DC4BCDA7,PinName="A",PinType.PinCategory="required",PinType.PinSubCategoryObject=None,PinType.PinSubCategoryMemberReference=(),PinType.bIsArray=False,PinType.bIsReference=False,PinType.bIsConst=False,PinType.bIsWeakPointer=False,LinkedTo=(MaterialGraphNode_70 987C42784FF10C1B6C5F66A0CC164817,),PersistentGuid=00000000000000000000000000000000,bHidden=False,bNotConnectable=False,bDefaultValueIsReadOnly=False,bDefaultValueIsIgnored=False,bAdvancedView=False,)
+			   CustomProperties Pin (PinId=207CC92F40241634F7D74290846DA5E3,PinName="B",PinType.PinCategory="required",PinType.PinSubCategoryObject=None,PinType.PinSubCategoryMemberReference=(),PinType.bIsArray=False,PinType.bIsReference=False,PinType.bIsConst=False,PinType.bIsWeakPointer=False,LinkedTo=(MaterialGraphNode_70 987C42784FF10C1B6C5F66A0CC164817,),PersistentGuid=00000000000000000000000000000000,bHidden=False,bNotConnectable=False,bDefaultValueIsReadOnly=False,bDefaultValueIsIgnored=False,bAdvancedView=False,)
+			   CustomProperties Pin (PinId=38DF7322459B1B42ACB22AB359E5788D,PinName="Output",PinFriendlyName=" ",Direction="EGPD_Output",PinType.PinSubCategoryObject=None,PinType.PinSubCategoryMemberReference=(),PinType.bIsArray=False,PinType.bIsReference=False,PinType.bIsConst=False,PinType.bIsWeakPointer=False,LinkedTo=(MaterialGraphNode_64 0D21E58243C6C7C77F26E7B3A92662EE,),PersistentGuid=00000000000000000000000000000000,bHidden=False,bNotConnectable=False,bDefaultValueIsReadOnly=False,bDefaultValueIsIgnored=False,bAdvancedView=False,)
+			End Object
+
+
+			END OBJECT
+	[/PARAM]
+[/OBJECT]
+[/EXCERPT:WaterMaterial]
+
+
+[EXCERPT:CompNormal]	
+[OBJECT:Snippet]
+	[PARAMLITERAL:id]
+		CompNormal
+	[/PARAMLITERAL]
+	[PARAMLITERAL:type]
+		material
+	[/PARAMLITERAL]
+	[PARAM:image]
+		![](Engine\Rendering\RenderTargets\BlueprintRenderTargets\HowTo\FluidSurface\FS_ComputeNormal_Material.png)
+	[/PARAM]
+	[PARAM:snippet]
+			BEGIN OBJECT
+					Begin Object Class=MaterialGraphNode_Root Name="MaterialGraphNode_Root_4"
+					   Material=PreviewMaterial'/Engine/Transient.PreviewMaterial_0'
+					   NodePosY=144
+					   NodeGuid=5954F5F14E0138F3422A48A68235CAC6
+					   CustomProperties Pin (PinId=891A561F41DBEC54AFF6D7A28BF76150,PinName="Base Color",PinType.PinCategory="materialinput",PinType.PinSubCategoryObject=None,PinType.PinSubCategoryMemberReference=(),PinType.bIsArray=False,PinType.bIsReference=False,PinType.bIsConst=False,PinType.bIsWeakPointer=False,PersistentGuid=00000000000000000000000000000000,bHidden=False,bNotConnectable=False,bDefaultValueIsReadOnly=False,bDefaultValueIsIgnored=False,bAdvancedView=False,)
+					   CustomProperties Pin (PinId=19059EC34A30A844571369BC97493F65,PinName="Metallic",PinType.PinCategory="materialinput",PinType.PinSubCategoryObject=None,PinType.PinSubCategoryMemberReference=(),PinType.bIsArray=False,PinType.bIsReference=False,PinType.bIsConst=False,PinType.bIsWeakPointer=False,PersistentGuid=00000000000000000000000000000000,bHidden=False,bNotConnectable=False,bDefaultValueIsReadOnly=False,bDefaultValueIsIgnored=False,bAdvancedView=False,)
+					   CustomProperties Pin (PinId=29EF1B3E4362B18D296378A0054ED148,PinName="Specular",PinType.PinCategory="materialinput",PinType.PinSubCategoryObject=None,PinType.PinSubCategoryMemberReference=(),PinType.bIsArray=False,PinType.bIsReference=False,PinType.bIsConst=False,PinType.bIsWeakPointer=False,PersistentGuid=00000000000000000000000000000000,bHidden=False,bNotConnectable=False,bDefaultValueIsReadOnly=False,bDefaultValueIsIgnored=False,bAdvancedView=False,)
+					   CustomProperties Pin (PinId=B371827A4DACDB4386BDA497622714D0,PinName="Roughness",PinType.PinCategory="materialinput",PinType.PinSubCategoryObject=None,PinType.PinSubCategoryMemberReference=(),PinType.bIsArray=False,PinType.bIsReference=False,PinType.bIsConst=False,PinType.bIsWeakPointer=False,PersistentGuid=00000000000000000000000000000000,bHidden=False,bNotConnectable=False,bDefaultValueIsReadOnly=False,bDefaultValueIsIgnored=False,bAdvancedView=False,)
+					   CustomProperties Pin (PinId=6347DBDD4486F550433714BA39F981CA,PinName="Emissive Color",PinType.PinCategory="materialinput",PinType.PinSubCategoryObject=None,PinType.PinSubCategoryMemberReference=(),PinType.bIsArray=False,PinType.bIsReference=False,PinType.bIsConst=False,PinType.bIsWeakPointer=False,LinkedTo=(MaterialGraphNode_135 0D06E5614F8EDD4351BF188D540DA472,),PersistentGuid=00000000000000000000000000000000,bHidden=False,bNotConnectable=False,bDefaultValueIsReadOnly=False,bDefaultValueIsIgnored=False,bAdvancedView=False,)
+					   CustomProperties Pin (PinId=FBAE1FAD4E4E003359F7B39478BEFE26,PinName="Opacity",PinType.PinCategory="materialinput",PinType.PinSubCategoryObject=None,PinType.PinSubCategoryMemberReference=(),PinType.bIsArray=False,PinType.bIsReference=False,PinType.bIsConst=False,PinType.bIsWeakPointer=False,PersistentGuid=00000000000000000000000000000000,bHidden=False,bNotConnectable=False,bDefaultValueIsReadOnly=False,bDefaultValueIsIgnored=False,bAdvancedView=False,)
+					   CustomProperties Pin (PinId=B5152A01458512CAA481A2BE3F2086EE,PinName="Opacity Mask",PinType.PinCategory="materialinput",PinType.PinSubCategoryObject=None,PinType.PinSubCategoryMemberReference=(),PinType.bIsArray=False,PinType.bIsReference=False,PinType.bIsConst=False,PinType.bIsWeakPointer=False,PersistentGuid=00000000000000000000000000000000,bHidden=False,bNotConnectable=False,bDefaultValueIsReadOnly=False,bDefaultValueIsIgnored=False,bAdvancedView=False,)
+					   CustomProperties Pin (PinId=66E6F03E4C5F58DBDDAF88B604C3EC75,PinName="Normal",PinType.PinCategory="materialinput",PinType.PinSubCategoryObject=None,PinType.PinSubCategoryMemberReference=(),PinType.bIsArray=False,PinType.bIsReference=False,PinType.bIsConst=False,PinType.bIsWeakPointer=False,PersistentGuid=00000000000000000000000000000000,bHidden=False,bNotConnectable=False,bDefaultValueIsReadOnly=False,bDefaultValueIsIgnored=False,bAdvancedView=False,)
+					   CustomProperties Pin (PinId=B18C254B48767D901FFF0F8A74B6E8AE,PinName="World Position Offset",PinType.PinCategory="materialinput",PinType.PinSubCategoryObject=None,PinType.PinSubCategoryMemberReference=(),PinType.bIsArray=False,PinType.bIsReference=False,PinType.bIsConst=False,PinType.bIsWeakPointer=False,PersistentGuid=00000000000000000000000000000000,bHidden=False,bNotConnectable=False,bDefaultValueIsReadOnly=False,bDefaultValueIsIgnored=False,bAdvancedView=False,)
+					   CustomProperties Pin (PinId=B88DF77B4F2B85FC4968B1ADFBFF2A5D,PinName="World Displacement",PinType.PinCategory="materialinput",PinType.PinSubCategoryObject=None,PinType.PinSubCategoryMemberReference=(),PinType.bIsArray=False,PinType.bIsReference=False,PinType.bIsConst=False,PinType.bIsWeakPointer=False,PersistentGuid=00000000000000000000000000000000,bHidden=False,bNotConnectable=False,bDefaultValueIsReadOnly=False,bDefaultValueIsIgnored=False,bAdvancedView=False,)
+					   CustomProperties Pin (PinId=CF4B41C9441F7364D8F84B87C7CCC685,PinName="Tessellation Multiplier",PinType.PinCategory="materialinput",PinType.PinSubCategoryObject=None,PinType.PinSubCategoryMemberReference=(),PinType.bIsArray=False,PinType.bIsReference=False,PinType.bIsConst=False,PinType.bIsWeakPointer=False,PersistentGuid=00000000000000000000000000000000,bHidden=False,bNotConnectable=False,bDefaultValueIsReadOnly=False,bDefaultValueIsIgnored=False,bAdvancedView=False,)
+					   CustomProperties Pin (PinId=6BF0D76248B5D92A8A6C958ECC2C1680,PinName="Subsurface Color",PinType.PinCategory="materialinput",PinType.PinSubCategoryObject=None,PinType.PinSubCategoryMemberReference=(),PinType.bIsArray=False,PinType.bIsReference=False,PinType.bIsConst=False,PinType.bIsWeakPointer=False,PersistentGuid=00000000000000000000000000000000,bHidden=False,bNotConnectable=False,bDefaultValueIsReadOnly=False,bDefaultValueIsIgnored=False,bAdvancedView=False,)
+					   CustomProperties Pin (PinId=28615DA94B028DC42FC426A3E2CB3D6D,PinName="Custom Data 0",PinType.PinCategory="materialinput",PinType.PinSubCategoryObject=None,PinType.PinSubCategoryMemberReference=(),PinType.bIsArray=False,PinType.bIsReference=False,PinType.bIsConst=False,PinType.bIsWeakPointer=False,PersistentGuid=00000000000000000000000000000000,bHidden=False,bNotConnectable=False,bDefaultValueIsReadOnly=False,bDefaultValueIsIgnored=False,bAdvancedView=False,)
+					   CustomProperties Pin (PinId=01AE9CE94D3523E5008F06ACB0A61783,PinName="Custom Data 1",PinType.PinCategory="materialinput",PinType.PinSubCategoryObject=None,PinType.PinSubCategoryMemberReference=(),PinType.bIsArray=False,PinType.bIsReference=False,PinType.bIsConst=False,PinType.bIsWeakPointer=False,PersistentGuid=00000000000000000000000000000000,bHidden=False,bNotConnectable=False,bDefaultValueIsReadOnly=False,bDefaultValueIsIgnored=False,bAdvancedView=False,)
+					   CustomProperties Pin (PinId=8EC030B2431EB88E52AE24990AA0C133,PinName="Ambient Occlusion",PinType.PinCategory="materialinput",PinType.PinSubCategoryObject=None,PinType.PinSubCategoryMemberReference=(),PinType.bIsArray=False,PinType.bIsReference=False,PinType.bIsConst=False,PinType.bIsWeakPointer=False,PersistentGuid=00000000000000000000000000000000,bHidden=False,bNotConnectable=False,bDefaultValueIsReadOnly=False,bDefaultValueIsIgnored=False,bAdvancedView=False,)
+					   CustomProperties Pin (PinId=E05AF7694D45C6F800018995E037E8CE,PinName="Refraction",PinType.PinCategory="materialinput",PinType.PinSubCategoryObject=None,PinType.PinSubCategoryMemberReference=(),PinType.bIsArray=False,PinType.bIsReference=False,PinType.bIsConst=False,PinType.bIsWeakPointer=False,PersistentGuid=00000000000000000000000000000000,bHidden=False,bNotConnectable=False,bDefaultValueIsReadOnly=False,bDefaultValueIsIgnored=False,bAdvancedView=False,)
+					   CustomProperties Pin (PinId=536E830440B68A61431D93A364D1C71A,PinName="Customized UV0",PinType.PinCategory="materialinput",PinType.PinSubCategoryObject=None,PinType.PinSubCategoryMemberReference=(),PinType.bIsArray=False,PinType.bIsReference=False,PinType.bIsConst=False,PinType.bIsWeakPointer=False,PersistentGuid=00000000000000000000000000000000,bHidden=False,bNotConnectable=False,bDefaultValueIsReadOnly=False,bDefaultValueIsIgnored=False,bAdvancedView=False,)
+					   CustomProperties Pin (PinId=06D5B2D04104C80B68BF62AD508ACF8B,PinName="Customized UV1",PinType.PinCategory="materialinput",PinType.PinSubCategoryObject=None,PinType.PinSubCategoryMemberReference=(),PinType.bIsArray=False,PinType.bIsReference=False,PinType.bIsConst=False,PinType.bIsWeakPointer=False,PersistentGuid=00000000000000000000000000000000,bHidden=False,bNotConnectable=False,bDefaultValueIsReadOnly=False,bDefaultValueIsIgnored=False,bAdvancedView=False,)
+					   CustomProperties Pin (PinId=FD6A1E394E7877B361194B867D1908F6,PinName="Customized UV2",PinType.PinCategory="materialinput",PinType.PinSubCategoryObject=None,PinType.PinSubCategoryMemberReference=(),PinType.bIsArray=False,PinType.bIsReference=False,PinType.bIsConst=False,PinType.bIsWeakPointer=False,PersistentGuid=00000000000000000000000000000000,bHidden=False,bNotConnectable=False,bDefaultValueIsReadOnly=False,bDefaultValueIsIgnored=False,bAdvancedView=False,)
+					   CustomProperties Pin (PinId=A7462B8147DFAD9AC160799F63DC8B68,PinName="Customized UV3",PinType.PinCategory="materialinput",PinType.PinSubCategoryObject=None,PinType.PinSubCategoryMemberReference=(),PinType.bIsArray=False,PinType.bIsReference=False,PinType.bIsConst=False,PinType.bIsWeakPointer=False,PersistentGuid=00000000000000000000000000000000,bHidden=False,bNotConnectable=False,bDefaultValueIsReadOnly=False,bDefaultValueIsIgnored=False,bAdvancedView=False,)
+					   CustomProperties Pin (PinId=417C0379415EC134FF3353AEDE895AEE,PinName="Customized UV4",PinType.PinCategory="materialinput",PinType.PinSubCategoryObject=None,PinType.PinSubCategoryMemberReference=(),PinType.bIsArray=False,PinType.bIsReference=False,PinType.bIsConst=False,PinType.bIsWeakPointer=False,PersistentGuid=00000000000000000000000000000000,bHidden=False,bNotConnectable=False,bDefaultValueIsReadOnly=False,bDefaultValueIsIgnored=False,bAdvancedView=False,)
+					   CustomProperties Pin (PinId=E6EF7D824031BBB0AE12519417017F95,PinName="Customized UV5",PinType.PinCategory="materialinput",PinType.PinSubCategoryObject=None,PinType.PinSubCategoryMemberReference=(),PinType.bIsArray=False,PinType.bIsReference=False,PinType.bIsConst=False,PinType.bIsWeakPointer=False,PersistentGuid=00000000000000000000000000000000,bHidden=False,bNotConnectable=False,bDefaultValueIsReadOnly=False,bDefaultValueIsIgnored=False,bAdvancedView=False,)
+					   CustomProperties Pin (PinId=6002F6094ABF8FEFE06D33AEB5F1B7EA,PinName="Customized UV6",PinType.PinCategory="materialinput",PinType.PinSubCategoryObject=None,PinType.PinSubCategoryMemberReference=(),PinType.bIsArray=False,PinType.bIsReference=False,PinType.bIsConst=False,PinType.bIsWeakPointer=False,PersistentGuid=00000000000000000000000000000000,bHidden=False,bNotConnectable=False,bDefaultValueIsReadOnly=False,bDefaultValueIsIgnored=False,bAdvancedView=False,)
+					   CustomProperties Pin (PinId=A471EB104E644D5D33C1C1903AC479FC,PinName="Customized UV7",PinType.PinCategory="materialinput",PinType.PinSubCategoryObject=None,PinType.PinSubCategoryMemberReference=(),PinType.bIsArray=False,PinType.bIsReference=False,PinType.bIsConst=False,PinType.bIsWeakPointer=False,PersistentGuid=00000000000000000000000000000000,bHidden=False,bNotConnectable=False,bDefaultValueIsReadOnly=False,bDefaultValueIsIgnored=False,bAdvancedView=False,)
+					   CustomProperties Pin (PinId=3815D46B4EB12FBC929AD5AE181FFEF1,PinName="Pixel Depth Offset",PinType.PinCategory="materialinput",PinType.PinSubCategoryObject=None,PinType.PinSubCategoryMemberReference=(),PinType.bIsArray=False,PinType.bIsReference=False,PinType.bIsConst=False,PinType.bIsWeakPointer=False,PersistentGuid=00000000000000000000000000000000,bHidden=False,bNotConnectable=False,bDefaultValueIsReadOnly=False,bDefaultValueIsIgnored=False,bAdvancedView=False,)
+					   CustomProperties Pin (PinId=1DBD479B4A8904458FA3EC87181A36A9,PinName="Material Attributes",PinType.PinCategory="materialinput",PinType.PinSubCategoryObject=None,PinType.PinSubCategoryMemberReference=(),PinType.bIsArray=False,PinType.bIsReference=False,PinType.bIsConst=False,PinType.bIsWeakPointer=False,PersistentGuid=00000000000000000000000000000000,bHidden=False,bNotConnectable=False,bDefaultValueIsReadOnly=False,bDefaultValueIsIgnored=False,bAdvancedView=False,)
+					End Object
+					Begin Object Class=MaterialGraphNode Name="MaterialGraphNode_135"
+					   Begin Object Class=MaterialExpressionNormalize Name="MaterialExpressionNormalize_0"
+					   End Object
+					   Begin Object Name="MaterialExpressionNormalize_0"
+						  VectorInput=(Expression=MaterialExpressionCrossProduct'MaterialGraphNode_148.MaterialExpressionCrossProduct_0')
+						  MaterialExpressionEditorX=-176
+						  MaterialExpressionEditorY=240
+						  MaterialExpressionGuid=5237C97D461F42F79BD14D9ABA01D22F
+						  Material=PreviewMaterial'/Engine/Transient.PreviewMaterial_0'
+					   End Object
+					   MaterialExpression=MaterialExpressionNormalize'MaterialExpressionNormalize_0'
+					   NodePosX=-176
+					   NodePosY=240
+					   NodeGuid=5FCF584646DE90CB29DF659A58CBABEA
+					   CustomProperties Pin (PinId=AE4DD0784D7EA7E1197F9A95EB7A13B7,PinName="VectorInput",PinType.PinCategory="required",PinType.PinSubCategoryObject=None,PinType.PinSubCategoryMemberReference=(),PinType.bIsArray=False,PinType.bIsReference=False,PinType.bIsConst=False,PinType.bIsWeakPointer=False,LinkedTo=(MaterialGraphNode_148 2059BC774C1CE4CC4936FCBB1A3288CC,),PersistentGuid=00000000000000000000000000000000,bHidden=False,bNotConnectable=False,bDefaultValueIsReadOnly=False,bDefaultValueIsIgnored=False,bAdvancedView=False,)
+					   CustomProperties Pin (PinId=0D06E5614F8EDD4351BF188D540DA472,PinName="Output",PinFriendlyName=" ",Direction="EGPD_Output",PinType.PinSubCategoryObject=None,PinType.PinSubCategoryMemberReference=(),PinType.bIsArray=False,PinType.bIsReference=False,PinType.bIsConst=False,PinType.bIsWeakPointer=False,LinkedTo=(MaterialGraphNode_Root_4 6347DBDD4486F550433714BA39F981CA,),PersistentGuid=00000000000000000000000000000000,bHidden=False,bNotConnectable=False,bDefaultValueIsReadOnly=False,bDefaultValueIsIgnored=False,bAdvancedView=False,)
+					End Object
+					Begin Object Class=MaterialGraphNode Name="MaterialGraphNode_136"
+					   Begin Object Class=MaterialExpressionScalarParameter Name="MaterialExpressionScalarParameter_27"
+					   End Object
+					   Begin Object Name="MaterialExpressionScalarParameter_27"
+						  DefaultValue=0.003906
+						  ParameterName="TexelSize"
+						  ExpressionGUID=D38BD539411556DE289856853ECA5767
+						  MaterialExpressionEditorX=-1376
+						  MaterialExpressionEditorY=272
+						  MaterialExpressionGuid=DB480CAC4000EF2A7EB61BA94FE23AF6
+						  Material=PreviewMaterial'/Engine/Transient.PreviewMaterial_0'
+					   End Object
+					   MaterialExpression=MaterialExpressionScalarParameter'MaterialExpressionScalarParameter_27'
+					   NodePosX=-1376
+					   NodePosY=272
+					   bCanRenameNode=True
+					   NodeGuid=5D2EE70E424878583314F9BCB0CB30CD
+					   CustomProperties Pin (PinId=FBF4D0A74307C4900B1C89B367B4809D,PinName="Output",PinFriendlyName=" ",Direction="EGPD_Output",PinType.PinSubCategoryObject=None,PinType.PinSubCategoryMemberReference=(),PinType.bIsArray=False,PinType.bIsReference=False,PinType.bIsConst=False,PinType.bIsWeakPointer=False,LinkedTo=(MaterialGraphNode_137 F32BB01B4749F32C7A1BDEBBCFCF6045,),PersistentGuid=00000000000000000000000000000000,bHidden=False,bNotConnectable=False,bDefaultValueIsReadOnly=False,bDefaultValueIsIgnored=False,bAdvancedView=False,)
+					End Object
+					Begin Object Class=MaterialGraphNode Name="MaterialGraphNode_137"
+					   Begin Object Class=MaterialExpressionMultiply Name="MaterialExpressionMultiply_23"
+					   End Object
+					   Begin Object Name="MaterialExpressionMultiply_23"
+						  A=(Expression=MaterialExpressionScalarParameter'MaterialGraphNode_136.MaterialExpressionScalarParameter_27')
+						  ConstB=2.000000
+						  MaterialExpressionEditorX=-1168
+						  MaterialExpressionEditorY=272
+						  MaterialExpressionGuid=7E91936049081EBE9FF3E6BC12584011
+						  Material=PreviewMaterial'/Engine/Transient.PreviewMaterial_0'
+					   End Object
+					   MaterialExpression=MaterialExpressionMultiply'MaterialExpressionMultiply_23'
+					   NodePosX=-1168
+					   NodePosY=272
+					   NodeGuid=9C7460BC48D1AC8D3DB3B8A79D4E1689
+					   CustomProperties Pin (PinId=F32BB01B4749F32C7A1BDEBBCFCF6045,PinName="A",PinType.PinCategory="optional",PinType.PinSubCategoryObject=None,PinType.PinSubCategoryMemberReference=(),PinType.bIsArray=False,PinType.bIsReference=False,PinType.bIsConst=False,PinType.bIsWeakPointer=False,LinkedTo=(MaterialGraphNode_136 FBF4D0A74307C4900B1C89B367B4809D,),PersistentGuid=00000000000000000000000000000000,bHidden=False,bNotConnectable=False,bDefaultValueIsReadOnly=False,bDefaultValueIsIgnored=False,bAdvancedView=False,)
+					   CustomProperties Pin (PinId=18671EF24ECD06B2B09022BD7D5416F3,PinName="B",PinType.PinCategory="optional",PinType.PinSubCategoryObject=None,PinType.PinSubCategoryMemberReference=(),PinType.bIsArray=False,PinType.bIsReference=False,PinType.bIsConst=False,PinType.bIsWeakPointer=False,PersistentGuid=00000000000000000000000000000000,bHidden=False,bNotConnectable=False,bDefaultValueIsReadOnly=False,bDefaultValueIsIgnored=False,bAdvancedView=False,)
+					   CustomProperties Pin (PinId=1366AD63419044CFE4751F89542243AC,PinName="Output",PinFriendlyName=" ",Direction="EGPD_Output",PinType.PinSubCategoryObject=None,PinType.PinSubCategoryMemberReference=(),PinType.bIsArray=False,PinType.bIsReference=False,PinType.bIsConst=False,PinType.bIsWeakPointer=False,LinkedTo=(MaterialGraphNode_138 25E20CAB4F8687642CA25FB4565B931A,),PersistentGuid=00000000000000000000000000000000,bHidden=False,bNotConnectable=False,bDefaultValueIsReadOnly=False,bDefaultValueIsIgnored=False,bAdvancedView=False,)
+					End Object
+					Begin Object Class=MaterialGraphNode Name="MaterialGraphNode_138"
+					   Begin Object Class=MaterialExpressionAppendVector Name="MaterialExpressionAppendVector_14"
+					   End Object
+					   Begin Object Name="MaterialExpressionAppendVector_14"
+						  A=(Expression=MaterialExpressionConstant'MaterialGraphNode_139.MaterialExpressionConstant_20')
+						  B=(Expression=MaterialExpressionMultiply'MaterialGraphNode_137.MaterialExpressionMultiply_23')
+						  MaterialExpressionEditorX=-960
+						  MaterialExpressionEditorY=160
+						  MaterialExpressionGuid=86FFB01642E4A0DBDAE1799F464C1B6A
+						  Material=PreviewMaterial'/Engine/Transient.PreviewMaterial_0'
+					   End Object
+					   MaterialExpression=MaterialExpressionAppendVector'MaterialExpressionAppendVector_14'
+					   NodePosX=-960
+					   NodePosY=160
+					   NodeGuid=69D387304D7D9710E6A596A5A2A4DF03
+					   CustomProperties Pin (PinId=E29FDA6541D7ABAD3BDE5DB71677C99E,PinName="A",PinType.PinCategory="required",PinType.PinSubCategoryObject=None,PinType.PinSubCategoryMemberReference=(),PinType.bIsArray=False,PinType.bIsReference=False,PinType.bIsConst=False,PinType.bIsWeakPointer=False,LinkedTo=(MaterialGraphNode_139 EA4C77E744F9676FA0ECE3B66CC16A37,),PersistentGuid=00000000000000000000000000000000,bHidden=False,bNotConnectable=False,bDefaultValueIsReadOnly=False,bDefaultValueIsIgnored=False,bAdvancedView=False,)
+					   CustomProperties Pin (PinId=25E20CAB4F8687642CA25FB4565B931A,PinName="B",PinType.PinCategory="required",PinType.PinSubCategoryObject=None,PinType.PinSubCategoryMemberReference=(),PinType.bIsArray=False,PinType.bIsReference=False,PinType.bIsConst=False,PinType.bIsWeakPointer=False,LinkedTo=(MaterialGraphNode_137 1366AD63419044CFE4751F89542243AC,),PersistentGuid=00000000000000000000000000000000,bHidden=False,bNotConnectable=False,bDefaultValueIsReadOnly=False,bDefaultValueIsIgnored=False,bAdvancedView=False,)
+					   CustomProperties Pin (PinId=06BADCDA40DA310616B908AE9F50F2E0,PinName="Output",PinFriendlyName=" ",Direction="EGPD_Output",PinType.PinSubCategoryObject=None,PinType.PinSubCategoryMemberReference=(),PinType.bIsArray=False,PinType.bIsReference=False,PinType.bIsConst=False,PinType.bIsWeakPointer=False,LinkedTo=(MaterialGraphNode_140 38C40A64432F41596E6A30A00FDAF9AA,),PersistentGuid=00000000000000000000000000000000,bHidden=False,bNotConnectable=False,bDefaultValueIsReadOnly=False,bDefaultValueIsIgnored=False,bAdvancedView=False,)
+					End Object
+					Begin Object Class=MaterialGraphNode Name="MaterialGraphNode_139"
+					   Begin Object Class=MaterialExpressionConstant Name="MaterialExpressionConstant_20"
+					   End Object
+					   Begin Object Name="MaterialExpressionConstant_20"
+						  MaterialExpressionEditorX=-1088
+						  MaterialExpressionEditorY=160
+						  MaterialExpressionGuid=955D3B774788894330F648A25500B89E
+						  Material=PreviewMaterial'/Engine/Transient.PreviewMaterial_0'
+					   End Object
+					   MaterialExpression=MaterialExpressionConstant'MaterialExpressionConstant_20'
+					   NodePosX=-1088
+					   NodePosY=160
+					   NodeGuid=301430E04BF2D8AE6884D182E1DB4225
+					   CustomProperties Pin (PinId=EA4C77E744F9676FA0ECE3B66CC16A37,PinName="Output",PinFriendlyName=" ",Direction="EGPD_Output",PinType.PinSubCategoryObject=None,PinType.PinSubCategoryMemberReference=(),PinType.bIsArray=False,PinType.bIsReference=False,PinType.bIsConst=False,PinType.bIsWeakPointer=False,LinkedTo=(MaterialGraphNode_138 E29FDA6541D7ABAD3BDE5DB71677C99E,),PersistentGuid=00000000000000000000000000000000,bHidden=False,bNotConnectable=False,bDefaultValueIsReadOnly=False,bDefaultValueIsIgnored=False,bAdvancedView=False,)
+					End Object
+					Begin Object Class=MaterialGraphNode Name="MaterialGraphNode_140"
+					   Begin Object Class=MaterialExpressionAppendVector Name="MaterialExpressionAppendVector_15"
+					   End Object
+					   Begin Object Name="MaterialExpressionAppendVector_15"
+						  A=(Expression=MaterialExpressionAppendVector'MaterialGraphNode_138.MaterialExpressionAppendVector_14')
+						  B=(Expression=MaterialExpressionComponentMask'MaterialGraphNode_141.MaterialExpressionComponentMask_0')
+						  MaterialExpressionEditorX=-848
+						  MaterialExpressionEditorY=256
+						  MaterialExpressionGuid=AA3134A144DA8A4AC9616C85D0E911C4
+						  Material=PreviewMaterial'/Engine/Transient.PreviewMaterial_0'
+					   End Object
+					   MaterialExpression=MaterialExpressionAppendVector'MaterialExpressionAppendVector_15'
+					   NodePosX=-848
+					   NodePosY=256
+					   NodeGuid=D57F2FEC48378D54765FC78EABA9D1BF
+					   CustomProperties Pin (PinId=38C40A64432F41596E6A30A00FDAF9AA,PinName="A",PinType.PinCategory="required",PinType.PinSubCategoryObject=None,PinType.PinSubCategoryMemberReference=(),PinType.bIsArray=False,PinType.bIsReference=False,PinType.bIsConst=False,PinType.bIsWeakPointer=False,LinkedTo=(MaterialGraphNode_138 06BADCDA40DA310616B908AE9F50F2E0,),PersistentGuid=00000000000000000000000000000000,bHidden=False,bNotConnectable=False,bDefaultValueIsReadOnly=False,bDefaultValueIsIgnored=False,bAdvancedView=False,)
+					   CustomProperties Pin (PinId=968B2AC44CF2B2AC05B1DEB8198FE713,PinName="B",PinType.PinCategory="required",PinType.PinSubCategoryObject=None,PinType.PinSubCategoryMemberReference=(),PinType.bIsArray=False,PinType.bIsReference=False,PinType.bIsConst=False,PinType.bIsWeakPointer=False,LinkedTo=(MaterialGraphNode_141 CB22E89A4727A6492C2DCABA61FA8DBA,),PersistentGuid=00000000000000000000000000000000,bHidden=False,bNotConnectable=False,bDefaultValueIsReadOnly=False,bDefaultValueIsIgnored=False,bAdvancedView=False,)
+					   CustomProperties Pin (PinId=B3A1305E470B1C6ADCF67B8791FC0FFA,PinName="Output",PinFriendlyName=" ",Direction="EGPD_Output",PinType.PinSubCategoryObject=None,PinType.PinSubCategoryMemberReference=(),PinType.bIsArray=False,PinType.bIsReference=False,PinType.bIsConst=False,PinType.bIsWeakPointer=False,LinkedTo=(MaterialGraphNode_148 77BB6B7D4E7D2E75F1162784A4584E08,),PersistentGuid=00000000000000000000000000000000,bHidden=False,bNotConnectable=False,bDefaultValueIsReadOnly=False,bDefaultValueIsIgnored=False,bAdvancedView=False,)
+					End Object
+					Begin Object Class=MaterialGraphNode Name="MaterialGraphNode_141"
+					   Begin Object Class=MaterialExpressionComponentMask Name="MaterialExpressionComponentMask_0"
+					   End Object
+					   Begin Object Name="MaterialExpressionComponentMask_0"
+						  Input=(Expression=MaterialExpressionMultiply'MaterialGraphNode_149.MaterialExpressionMultiply_25')
+						  G=True
+						  MaterialExpressionEditorX=-960
+						  MaterialExpressionEditorY=304
+						  MaterialExpressionGuid=250F7A8544CBB056A3784D99C8043903
+						  Material=PreviewMaterial'/Engine/Transient.PreviewMaterial_0'
+					   End Object
+					   MaterialExpression=MaterialExpressionComponentMask'MaterialExpressionComponentMask_0'
+					   NodePosX=-960
+					   NodePosY=304
+					   NodeGuid=0FEF39694E92B4CEE8108A93CCE204D7
+					   CustomProperties Pin (PinId=0DB19D2A441680E26DE24B9180FF3798,PinName="Input",PinFriendlyName=" ",PinType.PinCategory="required",PinType.PinSubCategoryObject=None,PinType.PinSubCategoryMemberReference=(),PinType.bIsArray=False,PinType.bIsReference=False,PinType.bIsConst=False,PinType.bIsWeakPointer=False,LinkedTo=(MaterialGraphNode_149 E22BE2BA4AE7FB9C2697659946A4BBA0,),PersistentGuid=00000000000000000000000000000000,bHidden=False,bNotConnectable=False,bDefaultValueIsReadOnly=False,bDefaultValueIsIgnored=False,bAdvancedView=False,)
+					   CustomProperties Pin (PinId=CB22E89A4727A6492C2DCABA61FA8DBA,PinName="Output",PinFriendlyName=" ",Direction="EGPD_Output",PinType.PinSubCategoryObject=None,PinType.PinSubCategoryMemberReference=(),PinType.bIsArray=False,PinType.bIsReference=False,PinType.bIsConst=False,PinType.bIsWeakPointer=False,LinkedTo=(MaterialGraphNode_140 968B2AC44CF2B2AC05B1DEB8198FE713,),PersistentGuid=00000000000000000000000000000000,bHidden=False,bNotConnectable=False,bDefaultValueIsReadOnly=False,bDefaultValueIsIgnored=False,bAdvancedView=False,)
+					End Object
+					Begin Object Class=MaterialGraphNode Name="MaterialGraphNode_142"
+					   Begin Object Class=MaterialExpressionComponentMask Name="MaterialExpressionComponentMask_1"
+					   End Object
+					   Begin Object Name="MaterialExpressionComponentMask_1"
+						  Input=(Expression=MaterialExpressionMultiply'MaterialGraphNode_149.MaterialExpressionMultiply_25')
+						  R=True
+						  MaterialExpressionEditorX=-624
+						  MaterialExpressionEditorY=224
+						  MaterialExpressionGuid=250F7A8544CBB056A3784D99C8043903
+						  Material=PreviewMaterial'/Engine/Transient.PreviewMaterial_0'
+					   End Object
+					   MaterialExpression=MaterialExpressionComponentMask'MaterialExpressionComponentMask_1'
+					   NodePosX=-624
+					   NodePosY=224
+					   NodeGuid=0DA4290D4C276F2E1EFDA7B32BBF3BD6
+					   CustomProperties Pin (PinId=EA292FA14BB421AA84BBB99F989779CE,PinName="Input",PinFriendlyName=" ",PinType.PinCategory="required",PinType.PinSubCategoryObject=None,PinType.PinSubCategoryMemberReference=(),PinType.bIsArray=False,PinType.bIsReference=False,PinType.bIsConst=False,PinType.bIsWeakPointer=False,LinkedTo=(MaterialGraphNode_149 E22BE2BA4AE7FB9C2697659946A4BBA0,),PersistentGuid=00000000000000000000000000000000,bHidden=False,bNotConnectable=False,bDefaultValueIsReadOnly=False,bDefaultValueIsIgnored=False,bAdvancedView=False,)
+					   CustomProperties Pin (PinId=1207151E4F34F13A2C00EBAE461F55E9,PinName="Output",PinFriendlyName=" ",Direction="EGPD_Output",PinType.PinSubCategoryObject=None,PinType.PinSubCategoryMemberReference=(),PinType.bIsArray=False,PinType.bIsReference=False,PinType.bIsConst=False,PinType.bIsWeakPointer=False,LinkedTo=(MaterialGraphNode_143 C2BCBA8B49F534514A6F4F82835D8A64,),PersistentGuid=00000000000000000000000000000000,bHidden=False,bNotConnectable=False,bDefaultValueIsReadOnly=False,bDefaultValueIsIgnored=False,bAdvancedView=False,)
+					End Object
+					Begin Object Class=MaterialGraphNode Name="MaterialGraphNode_143"
+					   Begin Object Class=MaterialExpressionAppendVector Name="MaterialExpressionAppendVector_16"
+					   End Object
+					   Begin Object Name="MaterialExpressionAppendVector_16"
+						  A=(Expression=MaterialExpressionAppendVector'MaterialGraphNode_145.MaterialExpressionAppendVector_17')
+						  B=(Expression=MaterialExpressionComponentMask'MaterialGraphNode_142.MaterialExpressionComponentMask_1')
+						  MaterialExpressionEditorX=-464
+						  MaterialExpressionEditorY=192
+						  MaterialExpressionGuid=AA3134A144DA8A4AC9616C85D0E911C4
+						  Material=PreviewMaterial'/Engine/Transient.PreviewMaterial_0'
+					   End Object
+					   MaterialExpression=MaterialExpressionAppendVector'MaterialExpressionAppendVector_16'
+					   NodePosX=-464
+					   NodePosY=192
+					   NodeGuid=2EDFC0224319194D5DC7FFBE0F9BBEB0
+					   CustomProperties Pin (PinId=A928152347F83A0A0C3A678D680E3396,PinName="A",PinType.PinCategory="required",PinType.PinSubCategoryObject=None,PinType.PinSubCategoryMemberReference=(),PinType.bIsArray=False,PinType.bIsReference=False,PinType.bIsConst=False,PinType.bIsWeakPointer=False,LinkedTo=(MaterialGraphNode_145 C71CC0A54118E1DEB14A5C85655CBA8D,),PersistentGuid=00000000000000000000000000000000,bHidden=False,bNotConnectable=False,bDefaultValueIsReadOnly=False,bDefaultValueIsIgnored=False,bAdvancedView=False,)
+					   CustomProperties Pin (PinId=C2BCBA8B49F534514A6F4F82835D8A64,PinName="B",PinType.PinCategory="required",PinType.PinSubCategoryObject=None,PinType.PinSubCategoryMemberReference=(),PinType.bIsArray=False,PinType.bIsReference=False,PinType.bIsConst=False,PinType.bIsWeakPointer=False,LinkedTo=(MaterialGraphNode_142 1207151E4F34F13A2C00EBAE461F55E9,),PersistentGuid=00000000000000000000000000000000,bHidden=False,bNotConnectable=False,bDefaultValueIsReadOnly=False,bDefaultValueIsIgnored=False,bAdvancedView=False,)
+					   CustomProperties Pin (PinId=1CAFCABB409797E87CC6229B27A80018,PinName="Output",PinFriendlyName=" ",Direction="EGPD_Output",PinType.PinSubCategoryObject=None,PinType.PinSubCategoryMemberReference=(),PinType.bIsArray=False,PinType.bIsReference=False,PinType.bIsConst=False,PinType.bIsWeakPointer=False,LinkedTo=(MaterialGraphNode_148 D3C502AC475CC415D029A7B5ED674369,),PersistentGuid=00000000000000000000000000000000,bHidden=False,bNotConnectable=False,bDefaultValueIsReadOnly=False,bDefaultValueIsIgnored=False,bAdvancedView=False,)
+					End Object
+					Begin Object Class=MaterialGraphNode Name="MaterialGraphNode_144"
+					   Begin Object Class=MaterialExpressionConstant Name="MaterialExpressionConstant_21"
+					   End Object
+					   Begin Object Name="MaterialExpressionConstant_21"
+						  MaterialExpressionEditorX=-768
+						  MaterialExpressionEditorY=144
+						  MaterialExpressionGuid=955D3B774788894330F648A25500B89E
+						  Material=PreviewMaterial'/Engine/Transient.PreviewMaterial_0'
+					   End Object
+					   MaterialExpression=MaterialExpressionConstant'MaterialExpressionConstant_21'
+					   NodePosX=-768
+					   NodePosY=144
+					   NodeGuid=490889714E29AA757600CC9F58197113
+					   CustomProperties Pin (PinId=E3F32004415FEF1894D377A7C460CB91,PinName="Output",PinFriendlyName=" ",Direction="EGPD_Output",PinType.PinSubCategoryObject=None,PinType.PinSubCategoryMemberReference=(),PinType.bIsArray=False,PinType.bIsReference=False,PinType.bIsConst=False,PinType.bIsWeakPointer=False,LinkedTo=(MaterialGraphNode_145 79E93BD64F8D255943629F8F0DFED3B9,),PersistentGuid=00000000000000000000000000000000,bHidden=False,bNotConnectable=False,bDefaultValueIsReadOnly=False,bDefaultValueIsIgnored=False,bAdvancedView=False,)
+					End Object
+					Begin Object Class=MaterialGraphNode Name="MaterialGraphNode_145"
+					   Begin Object Class=MaterialExpressionAppendVector Name="MaterialExpressionAppendVector_17"
+					   End Object
+					   Begin Object Name="MaterialExpressionAppendVector_17"
+						  A=(Expression=MaterialExpressionMultiply'MaterialGraphNode_146.MaterialExpressionMultiply_24')
+						  B=(Expression=MaterialExpressionConstant'MaterialGraphNode_144.MaterialExpressionConstant_21')
+						  MaterialExpressionEditorX=-656
+						  MaterialExpressionEditorY=112
+						  MaterialExpressionGuid=86FFB01642E4A0DBDAE1799F464C1B6A
+						  Material=PreviewMaterial'/Engine/Transient.PreviewMaterial_0'
+					   End Object
+					   MaterialExpression=MaterialExpressionAppendVector'MaterialExpressionAppendVector_17'
+					   NodePosX=-656
+					   NodePosY=112
+					   NodeGuid=A4E2289B4E3DAA2FD0F156B513349A4F
+					   CustomProperties Pin (PinId=30E9D22748201CC9A19467A59B5674CB,PinName="A",PinType.PinCategory="required",PinType.PinSubCategoryObject=None,PinType.PinSubCategoryMemberReference=(),PinType.bIsArray=False,PinType.bIsReference=False,PinType.bIsConst=False,PinType.bIsWeakPointer=False,LinkedTo=(MaterialGraphNode_146 5C4EDE1244426FB6012FD782F6B2A7FD,),PersistentGuid=00000000000000000000000000000000,bHidden=False,bNotConnectable=False,bDefaultValueIsReadOnly=False,bDefaultValueIsIgnored=False,bAdvancedView=False,)
+					   CustomProperties Pin (PinId=79E93BD64F8D255943629F8F0DFED3B9,PinName="B",PinType.PinCategory="required",PinType.PinSubCategoryObject=None,PinType.PinSubCategoryMemberReference=(),PinType.bIsArray=False,PinType.bIsReference=False,PinType.bIsConst=False,PinType.bIsWeakPointer=False,LinkedTo=(MaterialGraphNode_144 E3F32004415FEF1894D377A7C460CB91,),PersistentGuid=00000000000000000000000000000000,bHidden=False,bNotConnectable=False,bDefaultValueIsReadOnly=False,bDefaultValueIsIgnored=False,bAdvancedView=False,)
+					   CustomProperties Pin (PinId=C71CC0A54118E1DEB14A5C85655CBA8D,PinName="Output",PinFriendlyName=" ",Direction="EGPD_Output",PinType.PinSubCategoryObject=None,PinType.PinSubCategoryMemberReference=(),PinType.bIsArray=False,PinType.bIsReference=False,PinType.bIsConst=False,PinType.bIsWeakPointer=False,LinkedTo=(MaterialGraphNode_143 A928152347F83A0A0C3A678D680E3396,),PersistentGuid=00000000000000000000000000000000,bHidden=False,bNotConnectable=False,bDefaultValueIsReadOnly=False,bDefaultValueIsIgnored=False,bAdvancedView=False,)
+					End Object
+					Begin Object Class=MaterialGraphNode Name="MaterialGraphNode_146"
+					   Begin Object Class=MaterialExpressionMultiply Name="MaterialExpressionMultiply_24"
+					   End Object
+					   Begin Object Name="MaterialExpressionMultiply_24"
+						  A=(Expression=MaterialExpressionScalarParameter'MaterialGraphNode_147.MaterialExpressionScalarParameter_28')
+						  ConstB=2.000000
+						  MaterialExpressionEditorX=-800
+						  MaterialExpressionEditorY=64
+						  MaterialExpressionGuid=7E91936049081EBE9FF3E6BC12584011
+						  Material=PreviewMaterial'/Engine/Transient.PreviewMaterial_0'
+					   End Object
+					   MaterialExpression=MaterialExpressionMultiply'MaterialExpressionMultiply_24'
+					   NodePosX=-800
+					   NodePosY=64
+					   NodeGuid=EC9FEA964BA67ED5A7346583706E38DC
+					   CustomProperties Pin (PinId=AF201CA842E9F28547F91FBB59D52E9A,PinName="A",PinType.PinCategory="optional",PinType.PinSubCategoryObject=None,PinType.PinSubCategoryMemberReference=(),PinType.bIsArray=False,PinType.bIsReference=False,PinType.bIsConst=False,PinType.bIsWeakPointer=False,LinkedTo=(MaterialGraphNode_147 1B9ED5B848CE35F098E9DF99D893692E,),PersistentGuid=00000000000000000000000000000000,bHidden=False,bNotConnectable=False,bDefaultValueIsReadOnly=False,bDefaultValueIsIgnored=False,bAdvancedView=False,)
+					   CustomProperties Pin (PinId=685F775E4CAFE636E16FB1A5DD0F2DA9,PinName="B",PinType.PinCategory="optional",PinType.PinSubCategoryObject=None,PinType.PinSubCategoryMemberReference=(),PinType.bIsArray=False,PinType.bIsReference=False,PinType.bIsConst=False,PinType.bIsWeakPointer=False,PersistentGuid=00000000000000000000000000000000,bHidden=False,bNotConnectable=False,bDefaultValueIsReadOnly=False,bDefaultValueIsIgnored=False,bAdvancedView=False,)
+					   CustomProperties Pin (PinId=5C4EDE1244426FB6012FD782F6B2A7FD,PinName="Output",PinFriendlyName=" ",Direction="EGPD_Output",PinType.PinSubCategoryObject=None,PinType.PinSubCategoryMemberReference=(),PinType.bIsArray=False,PinType.bIsReference=False,PinType.bIsConst=False,PinType.bIsWeakPointer=False,LinkedTo=(MaterialGraphNode_145 30E9D22748201CC9A19467A59B5674CB,),PersistentGuid=00000000000000000000000000000000,bHidden=False,bNotConnectable=False,bDefaultValueIsReadOnly=False,bDefaultValueIsIgnored=False,bAdvancedView=False,)
+					End Object
+					Begin Object Class=MaterialGraphNode Name="MaterialGraphNode_147"
+					   Begin Object Class=MaterialExpressionScalarParameter Name="MaterialExpressionScalarParameter_28"
+					   End Object
+					   Begin Object Name="MaterialExpressionScalarParameter_28"
+						  DefaultValue=0.003906
+						  ParameterName="TexelSize"
+						  ExpressionGUID=8D2C13E549BDCCF8E765738509B119C6
+						  MaterialExpressionEditorX=-1008
+						  MaterialExpressionEditorY=48
+						  MaterialExpressionGuid=DB480CAC4000EF2A7EB61BA94FE23AF6
+						  Material=PreviewMaterial'/Engine/Transient.PreviewMaterial_0'
+					   End Object
+					   MaterialExpression=MaterialExpressionScalarParameter'MaterialExpressionScalarParameter_28'
+					   NodePosX=-1008
+					   NodePosY=48
+					   bCanRenameNode=True
+					   NodeGuid=EB2264BB41936E5A3B982E8298F50763
+					   CustomProperties Pin (PinId=1B9ED5B848CE35F098E9DF99D893692E,PinName="Output",PinFriendlyName=" ",Direction="EGPD_Output",PinType.PinSubCategoryObject=None,PinType.PinSubCategoryMemberReference=(),PinType.bIsArray=False,PinType.bIsReference=False,PinType.bIsConst=False,PinType.bIsWeakPointer=False,LinkedTo=(MaterialGraphNode_146 AF201CA842E9F28547F91FBB59D52E9A,),PersistentGuid=00000000000000000000000000000000,bHidden=False,bNotConnectable=False,bDefaultValueIsReadOnly=False,bDefaultValueIsIgnored=False,bAdvancedView=False,)
+					End Object
+					Begin Object Class=MaterialGraphNode Name="MaterialGraphNode_148"
+					   Begin Object Class=MaterialExpressionCrossProduct Name="MaterialExpressionCrossProduct_0"
+					   End Object
+					   Begin Object Name="MaterialExpressionCrossProduct_0"
+						  A=(Expression=MaterialExpressionAppendVector'MaterialGraphNode_143.MaterialExpressionAppendVector_16')
+						  B=(Expression=MaterialExpressionAppendVector'MaterialGraphNode_140.MaterialExpressionAppendVector_15')
+						  MaterialExpressionEditorX=-304
+						  MaterialExpressionEditorY=240
+						  MaterialExpressionGuid=14AC11F34E74A4458D9242B63B545C1C
+						  Material=PreviewMaterial'/Engine/Transient.PreviewMaterial_0'
+					   End Object
+					   MaterialExpression=MaterialExpressionCrossProduct'MaterialExpressionCrossProduct_0'
+					   NodePosX=-304
+					   NodePosY=240
+					   NodeGuid=25DF310F42D326B8395441BEB47D221C
+					   CustomProperties Pin (PinId=D3C502AC475CC415D029A7B5ED674369,PinName="A",PinType.PinCategory="required",PinType.PinSubCategoryObject=None,PinType.PinSubCategoryMemberReference=(),PinType.bIsArray=False,PinType.bIsReference=False,PinType.bIsConst=False,PinType.bIsWeakPointer=False,LinkedTo=(MaterialGraphNode_143 1CAFCABB409797E87CC6229B27A80018,),PersistentGuid=00000000000000000000000000000000,bHidden=False,bNotConnectable=False,bDefaultValueIsReadOnly=False,bDefaultValueIsIgnored=False,bAdvancedView=False,)
+					   CustomProperties Pin (PinId=77BB6B7D4E7D2E75F1162784A4584E08,PinName="B",PinType.PinCategory="required",PinType.PinSubCategoryObject=None,PinType.PinSubCategoryMemberReference=(),PinType.bIsArray=False,PinType.bIsReference=False,PinType.bIsConst=False,PinType.bIsWeakPointer=False,LinkedTo=(MaterialGraphNode_140 B3A1305E470B1C6ADCF67B8791FC0FFA,),PersistentGuid=00000000000000000000000000000000,bHidden=False,bNotConnectable=False,bDefaultValueIsReadOnly=False,bDefaultValueIsIgnored=False,bAdvancedView=False,)
+					   CustomProperties Pin (PinId=2059BC774C1CE4CC4936FCBB1A3288CC,PinName="Output",PinFriendlyName=" ",Direction="EGPD_Output",PinType.PinSubCategoryObject=None,PinType.PinSubCategoryMemberReference=(),PinType.bIsArray=False,PinType.bIsReference=False,PinType.bIsConst=False,PinType.bIsWeakPointer=False,LinkedTo=(MaterialGraphNode_135 AE4DD0784D7EA7E1197F9A95EB7A13B7,),PersistentGuid=00000000000000000000000000000000,bHidden=False,bNotConnectable=False,bDefaultValueIsReadOnly=False,bDefaultValueIsIgnored=False,bAdvancedView=False,)
+					End Object
+					Begin Object Class=MaterialGraphNode Name="MaterialGraphNode_149"
+					   Begin Object Class=MaterialExpressionMultiply Name="MaterialExpressionMultiply_25"
+					   End Object
+					   Begin Object Name="MaterialExpressionMultiply_25"
+						  A=(Expression=MaterialExpressionSubtract'MaterialGraphNode_151.MaterialExpressionSubtract_4')
+						  B=(Expression=MaterialExpressionScalarParameter'MaterialGraphNode_150.MaterialExpressionScalarParameter_29')
+						  MaterialExpressionEditorX=-1072
+						  MaterialExpressionEditorY=448
+						  MaterialExpressionGuid=61532B08426CDB78D2F2A5A086C4B891
+						  Material=PreviewMaterial'/Engine/Transient.PreviewMaterial_0'
+					   End Object
+					   MaterialExpression=MaterialExpressionMultiply'MaterialExpressionMultiply_25'
+					   NodePosX=-1072
+					   NodePosY=448
+					   NodeGuid=B616470042B6C6DA1A5ED0BAE71FED0E
+					   CustomProperties Pin (PinId=44B1A61845EDEAA9F3DE5194BF08C626,PinName="A",PinType.PinCategory="optional",PinType.PinSubCategoryObject=None,PinType.PinSubCategoryMemberReference=(),PinType.bIsArray=False,PinType.bIsReference=False,PinType.bIsConst=False,PinType.bIsWeakPointer=False,LinkedTo=(MaterialGraphNode_151 6A1140D54495DB4B88D3EC8F495C6849,),PersistentGuid=00000000000000000000000000000000,bHidden=False,bNotConnectable=False,bDefaultValueIsReadOnly=False,bDefaultValueIsIgnored=False,bAdvancedView=False,)
+					   CustomProperties Pin (PinId=F7261C9C49B6BCCCD78D26984D55F99E,PinName="B",PinType.PinCategory="optional",PinType.PinSubCategoryObject=None,PinType.PinSubCategoryMemberReference=(),PinType.bIsArray=False,PinType.bIsReference=False,PinType.bIsConst=False,PinType.bIsWeakPointer=False,LinkedTo=(MaterialGraphNode_150 A3C289BA4546D107F07A09A50E9F5787,),PersistentGuid=00000000000000000000000000000000,bHidden=False,bNotConnectable=False,bDefaultValueIsReadOnly=False,bDefaultValueIsIgnored=False,bAdvancedView=False,)
+					   CustomProperties Pin (PinId=E22BE2BA4AE7FB9C2697659946A4BBA0,PinName="Output",PinFriendlyName=" ",Direction="EGPD_Output",PinType.PinSubCategoryObject=None,PinType.PinSubCategoryMemberReference=(),PinType.bIsArray=False,PinType.bIsReference=False,PinType.bIsConst=False,PinType.bIsWeakPointer=False,LinkedTo=(MaterialGraphNode_141 0DB19D2A441680E26DE24B9180FF3798,MaterialGraphNode_142 EA292FA14BB421AA84BBB99F989779CE,),PersistentGuid=00000000000000000000000000000000,bHidden=False,bNotConnectable=False,bDefaultValueIsReadOnly=False,bDefaultValueIsIgnored=False,bAdvancedView=False,)
+					End Object
+					Begin Object Class=MaterialGraphNode Name="MaterialGraphNode_150"
+					   Begin Object Class=MaterialExpressionScalarParameter Name="MaterialExpressionScalarParameter_29"
+					   End Object
+					   Begin Object Name="MaterialExpressionScalarParameter_29"
+						  DefaultValue=0.010000
+						  ParameterName="HeightScale"
+						  ExpressionGUID=964B5AA543201FAC74500FA3C4C301F2
+						  MaterialExpressionEditorX=-1200
+						  MaterialExpressionEditorY=576
+						  MaterialExpressionGuid=CACB78A44B8FEAF436702596FEF81BB5
+						  Material=PreviewMaterial'/Engine/Transient.PreviewMaterial_0'
+					   End Object
+					   MaterialExpression=MaterialExpressionScalarParameter'MaterialExpressionScalarParameter_29'
+					   NodePosX=-1200
+					   NodePosY=576
+					   bCanRenameNode=True
+					   NodeGuid=FA85CD894C2A3C6ADADA5FA8ECDE051B
+					   CustomProperties Pin (PinId=A3C289BA4546D107F07A09A50E9F5787,PinName="Output",PinFriendlyName=" ",Direction="EGPD_Output",PinType.PinSubCategoryObject=None,PinType.PinSubCategoryMemberReference=(),PinType.bIsArray=False,PinType.bIsReference=False,PinType.bIsConst=False,PinType.bIsWeakPointer=False,LinkedTo=(MaterialGraphNode_149 F7261C9C49B6BCCCD78D26984D55F99E,),PersistentGuid=00000000000000000000000000000000,bHidden=False,bNotConnectable=False,bDefaultValueIsReadOnly=False,bDefaultValueIsIgnored=False,bAdvancedView=False,)
+					End Object
+					Begin Object Class=MaterialGraphNode Name="MaterialGraphNode_151"
+					   Begin Object Class=MaterialExpressionSubtract Name="MaterialExpressionSubtract_4"
+					   End Object
+					   Begin Object Name="MaterialExpressionSubtract_4"
+						  A=(Expression=MaterialExpressionAppendVector'MaterialGraphNode_153.MaterialExpressionAppendVector_19')
+						  B=(Expression=MaterialExpressionAppendVector'MaterialGraphNode_152.MaterialExpressionAppendVector_18')
+						  MaterialExpressionEditorX=-1216
+						  MaterialExpressionEditorY=448
+						  MaterialExpressionGuid=FC7A20EB4E3BCDB8D72890B03688E864
+						  Material=PreviewMaterial'/Engine/Transient.PreviewMaterial_0'
+					   End Object
+					   MaterialExpression=MaterialExpressionSubtract'MaterialExpressionSubtract_4'
+					   NodePosX=-1216
+					   NodePosY=448
+					   NodeGuid=69F2B62F45209330A5AC15B4FE051B55
+					   CustomProperties Pin (PinId=A63ADF3346BCE0F000CBF6953C03F30B,PinName="A",PinType.PinCategory="optional",PinType.PinSubCategoryObject=None,PinType.PinSubCategoryMemberReference=(),PinType.bIsArray=False,PinType.bIsReference=False,PinType.bIsConst=False,PinType.bIsWeakPointer=False,LinkedTo=(MaterialGraphNode_153 36C1EFC144F89E2A2196A0B46E8C84C5,),PersistentGuid=00000000000000000000000000000000,bHidden=False,bNotConnectable=False,bDefaultValueIsReadOnly=False,bDefaultValueIsIgnored=False,bAdvancedView=False,)
+					   CustomProperties Pin (PinId=24CFABAB4E85BD5022C5D5A8ED4D301E,PinName="B",PinType.PinCategory="optional",PinType.PinSubCategoryObject=None,PinType.PinSubCategoryMemberReference=(),PinType.bIsArray=False,PinType.bIsReference=False,PinType.bIsConst=False,PinType.bIsWeakPointer=False,LinkedTo=(MaterialGraphNode_152 77BC0D0E4C8B579E968258A5EFF1F62A,),PersistentGuid=00000000000000000000000000000000,bHidden=False,bNotConnectable=False,bDefaultValueIsReadOnly=False,bDefaultValueIsIgnored=False,bAdvancedView=False,)
+					   CustomProperties Pin (PinId=6A1140D54495DB4B88D3EC8F495C6849,PinName="Output",PinFriendlyName=" ",Direction="EGPD_Output",PinType.PinSubCategoryObject=None,PinType.PinSubCategoryMemberReference=(),PinType.bIsArray=False,PinType.bIsReference=False,PinType.bIsConst=False,PinType.bIsWeakPointer=False,LinkedTo=(MaterialGraphNode_149 44B1A61845EDEAA9F3DE5194BF08C626,),PersistentGuid=00000000000000000000000000000000,bHidden=False,bNotConnectable=False,bDefaultValueIsReadOnly=False,bDefaultValueIsIgnored=False,bAdvancedView=False,)
+					End Object
+					Begin Object Class=MaterialGraphNode Name="MaterialGraphNode_152"
+					   Begin Object Class=MaterialExpressionAppendVector Name="MaterialExpressionAppendVector_18"
+					   End Object
+					   Begin Object Name="MaterialExpressionAppendVector_18"
+						  A=(Expression=MaterialExpressionTextureSampleParameter2D'MaterialGraphNode_161.MaterialExpressionTextureSampleParameter2D_27',OutputIndex=1,Mask=1,MaskR=1)
+						  B=(Expression=MaterialExpressionTextureSampleParameter2D'MaterialGraphNode_169.MaterialExpressionTextureSampleParameter2D_29',OutputIndex=1,Mask=1,MaskR=1)
+						  MaterialExpressionEditorX=-1328
+						  MaterialExpressionEditorY=544
+						  MaterialExpressionGuid=D23FB7EF42CA6CD3E1A6DCB987C59C5F
+						  Material=PreviewMaterial'/Engine/Transient.PreviewMaterial_0'
+					   End Object
+					   MaterialExpression=MaterialExpressionAppendVector'MaterialExpressionAppendVector_18'
+					   NodePosX=-1328
+					   NodePosY=544
+					   NodeGuid=C3F366B84817E08277627CBCE29F330F
+					   CustomProperties Pin (PinId=9DA0FD994C0986C3C8A4E6B1AA61D5C9,PinName="A",PinType.PinCategory="required",PinType.PinSubCategoryObject=None,PinType.PinSubCategoryMemberReference=(),PinType.bIsArray=False,PinType.bIsReference=False,PinType.bIsConst=False,PinType.bIsWeakPointer=False,LinkedTo=(MaterialGraphNode_161 B7DF7ECD43804F6D666F12BB51CE3166,),PersistentGuid=00000000000000000000000000000000,bHidden=False,bNotConnectable=False,bDefaultValueIsReadOnly=False,bDefaultValueIsIgnored=False,bAdvancedView=False,)
+					   CustomProperties Pin (PinId=28DF7D474943CA7C66A104B7504A711B,PinName="B",PinType.PinCategory="required",PinType.PinSubCategoryObject=None,PinType.PinSubCategoryMemberReference=(),PinType.bIsArray=False,PinType.bIsReference=False,PinType.bIsConst=False,PinType.bIsWeakPointer=False,LinkedTo=(MaterialGraphNode_169 2E05221B4F9937E3F66D36AE8E69475C,),PersistentGuid=00000000000000000000000000000000,bHidden=False,bNotConnectable=False,bDefaultValueIsReadOnly=False,bDefaultValueIsIgnored=False,bAdvancedView=False,)
+					   CustomProperties Pin (PinId=77BC0D0E4C8B579E968258A5EFF1F62A,PinName="Output",PinFriendlyName=" ",Direction="EGPD_Output",PinType.PinSubCategoryObject=None,PinType.PinSubCategoryMemberReference=(),PinType.bIsArray=False,PinType.bIsReference=False,PinType.bIsConst=False,PinType.bIsWeakPointer=False,LinkedTo=(MaterialGraphNode_151 24CFABAB4E85BD5022C5D5A8ED4D301E,),PersistentGuid=00000000000000000000000000000000,bHidden=False,bNotConnectable=False,bDefaultValueIsReadOnly=False,bDefaultValueIsIgnored=False,bAdvancedView=False,)
+					End Object
+					Begin Object Class=MaterialGraphNode Name="MaterialGraphNode_153"
+					   Begin Object Class=MaterialExpressionAppendVector Name="MaterialExpressionAppendVector_19"
+					   End Object
+					   Begin Object Name="MaterialExpressionAppendVector_19"
+						  A=(Expression=MaterialExpressionTextureSampleParameter2D'MaterialGraphNode_154.MaterialExpressionTextureSampleParameter2D_26',OutputIndex=1,Mask=1,MaskR=1)
+						  B=(Expression=MaterialExpressionTextureSampleParameter2D'MaterialGraphNode_165.MaterialExpressionTextureSampleParameter2D_28',OutputIndex=1,Mask=1,MaskR=1)
+						  MaterialExpressionEditorX=-1344
+						  MaterialExpressionEditorY=448
+						  MaterialExpressionGuid=8B878FA743E28782F3B03C8439CD1CDC
+						  Material=PreviewMaterial'/Engine/Transient.PreviewMaterial_0'
+					   End Object
+					   MaterialExpression=MaterialExpressionAppendVector'MaterialExpressionAppendVector_19'
+					   NodePosX=-1344
+					   NodePosY=448
+					   NodeGuid=D70278844032BDE140FA389BC6F5808D
+					   CustomProperties Pin (PinId=214A46394DEED874C9F9BDADBA27E85E,PinName="A",PinType.PinCategory="required",PinType.PinSubCategoryObject=None,PinType.PinSubCategoryMemberReference=(),PinType.bIsArray=False,PinType.bIsReference=False,PinType.bIsConst=False,PinType.bIsWeakPointer=False,LinkedTo=(MaterialGraphNode_154 74A83272440F80D5E8B3AB8D38180712,),PersistentGuid=00000000000000000000000000000000,bHidden=False,bNotConnectable=False,bDefaultValueIsReadOnly=False,bDefaultValueIsIgnored=False,bAdvancedView=False,)
+					   CustomProperties Pin (PinId=87AEDFE8427805C6775CAFB43BBDF596,PinName="B",PinType.PinCategory="required",PinType.PinSubCategoryObject=None,PinType.PinSubCategoryMemberReference=(),PinType.bIsArray=False,PinType.bIsReference=False,PinType.bIsConst=False,PinType.bIsWeakPointer=False,LinkedTo=(MaterialGraphNode_165 87A942C9487C8643262DC48786DE2B90,),PersistentGuid=00000000000000000000000000000000,bHidden=False,bNotConnectable=False,bDefaultValueIsReadOnly=False,bDefaultValueIsIgnored=False,bAdvancedView=False,)
+					   CustomProperties Pin (PinId=36C1EFC144F89E2A2196A0B46E8C84C5,PinName="Output",PinFriendlyName=" ",Direction="EGPD_Output",PinType.PinSubCategoryObject=None,PinType.PinSubCategoryMemberReference=(),PinType.bIsArray=False,PinType.bIsReference=False,PinType.bIsConst=False,PinType.bIsWeakPointer=False,LinkedTo=(MaterialGraphNode_151 A63ADF3346BCE0F000CBF6953C03F30B,),PersistentGuid=00000000000000000000000000000000,bHidden=False,bNotConnectable=False,bDefaultValueIsReadOnly=False,bDefaultValueIsIgnored=False,bAdvancedView=False,)
+					End Object
+					Begin Object Class=MaterialGraphNode Name="MaterialGraphNode_154"
+					   Begin Object Class=MaterialExpressionTextureSampleParameter2D Name="MaterialExpressionTextureSampleParameter2D_26"
+					   End Object
+					   Begin Object Name="MaterialExpressionTextureSampleParameter2D_26"
+						  ParameterName="Heightfield"
+						  ExpressionGUID=4DBD3ED647977C9FCA9236B0F903BD3C
+						  Coordinates=(Expression=MaterialExpressionAdd'MaterialGraphNode_157.MaterialExpressionAdd_16')
+						  MaterialExpressionEditorX=-1616
+						  MaterialExpressionEditorY=208
+						  MaterialExpressionGuid=8F3166474EF9A6F6A6825185B78360BE
+						  Material=PreviewMaterial'/Engine/Transient.PreviewMaterial_0'
+					   End Object
+					   MaterialExpression=MaterialExpressionTextureSampleParameter2D'MaterialExpressionTextureSampleParameter2D_26'
+					   NodePosX=-1616
+					   NodePosY=208
+					   bCanRenameNode=True
+					   NodeGuid=E6DEEF84434DE864DAD8BB9EA253BBF0
+					   CustomProperties Pin (PinId=A25AC337495A10638C30FA93453959BE,PinName="UVs",PinType.PinCategory="optional",PinType.PinSubCategoryObject=None,PinType.PinSubCategoryMemberReference=(),PinType.bIsArray=False,PinType.bIsReference=False,PinType.bIsConst=False,PinType.bIsWeakPointer=False,LinkedTo=(MaterialGraphNode_157 801803B740F37EFDD82FE9B7F901CD44,),PersistentGuid=00000000000000000000000000000000,bHidden=False,bNotConnectable=False,bDefaultValueIsReadOnly=False,bDefaultValueIsIgnored=False,bAdvancedView=False,)
+					   CustomProperties Pin (PinId=7DCD7E454474CA1BBB6C999AC1F4F51A,PinName="Tex",PinType.PinCategory="optional",PinType.PinSubCategoryObject=None,PinType.PinSubCategoryMemberReference=(),PinType.bIsArray=False,PinType.bIsReference=False,PinType.bIsConst=False,PinType.bIsWeakPointer=False,PersistentGuid=00000000000000000000000000000000,bHidden=False,bNotConnectable=False,bDefaultValueIsReadOnly=False,bDefaultValueIsIgnored=False,bAdvancedView=False,)
+					   CustomProperties Pin (PinId=7C292F534D186C1DA3002FAF1C911417,PinName="Output",PinFriendlyName=" ",Direction="EGPD_Output",PinType.PinCategory="mask",PinType.PinSubCategoryObject=None,PinType.PinSubCategoryMemberReference=(),PinType.bIsArray=False,PinType.bIsReference=False,PinType.bIsConst=False,PinType.bIsWeakPointer=False,PersistentGuid=00000000000000000000000000000000,bHidden=False,bNotConnectable=False,bDefaultValueIsReadOnly=False,bDefaultValueIsIgnored=False,bAdvancedView=False,)
+					   CustomProperties Pin (PinId=74A83272440F80D5E8B3AB8D38180712,PinName="Output2",PinFriendlyName=" ",Direction="EGPD_Output",PinType.PinCategory="mask",PinType.PinSubCategory="red",PinType.PinSubCategoryObject=None,PinType.PinSubCategoryMemberReference=(),PinType.bIsArray=False,PinType.bIsReference=False,PinType.bIsConst=False,PinType.bIsWeakPointer=False,LinkedTo=(MaterialGraphNode_153 214A46394DEED874C9F9BDADBA27E85E,),PersistentGuid=00000000000000000000000000000000,bHidden=False,bNotConnectable=False,bDefaultValueIsReadOnly=False,bDefaultValueIsIgnored=False,bAdvancedView=False,)
+					   CustomProperties Pin (PinId=D7D9B09545FAB4A9EC939981C7A82F5E,PinName="Output3",PinFriendlyName=" ",Direction="EGPD_Output",PinType.PinCategory="mask",PinType.PinSubCategory="green",PinType.PinSubCategoryObject=None,PinType.PinSubCategoryMemberReference=(),PinType.bIsArray=False,PinType.bIsReference=False,PinType.bIsConst=False,PinType.bIsWeakPointer=False,PersistentGuid=00000000000000000000000000000000,bHidden=False,bNotConnectable=False,bDefaultValueIsReadOnly=False,bDefaultValueIsIgnored=False,bAdvancedView=False,)
+					   CustomProperties Pin (PinId=7D4779294DAA988312659DADF4725FD9,PinName="Output4",PinFriendlyName=" ",Direction="EGPD_Output",PinType.PinCategory="mask",PinType.PinSubCategory="blue",PinType.PinSubCategoryObject=None,PinType.PinSubCategoryMemberReference=(),PinType.bIsArray=False,PinType.bIsReference=False,PinType.bIsConst=False,PinType.bIsWeakPointer=False,PersistentGuid=00000000000000000000000000000000,bHidden=False,bNotConnectable=False,bDefaultValueIsReadOnly=False,bDefaultValueIsIgnored=False,bAdvancedView=False,)
+					   CustomProperties Pin (PinId=2B6591464A95EC66446C12B4FE5879A3,PinName="Output5",PinFriendlyName=" ",Direction="EGPD_Output",PinType.PinCategory="mask",PinType.PinSubCategory="alpha",PinType.PinSubCategoryObject=None,PinType.PinSubCategoryMemberReference=(),PinType.bIsArray=False,PinType.bIsReference=False,PinType.bIsConst=False,PinType.bIsWeakPointer=False,PersistentGuid=00000000000000000000000000000000,bHidden=False,bNotConnectable=False,bDefaultValueIsReadOnly=False,bDefaultValueIsIgnored=False,bAdvancedView=False,)
+					End Object
+					Begin Object Class=MaterialGraphNode Name="MaterialGraphNode_155"
+					   Begin Object Class=MaterialExpressionConstant2Vector Name="MaterialExpressionConstant2Vector_8"
+					   End Object
+					   Begin Object Name="MaterialExpressionConstant2Vector_8"
+						  R=1.000000
+						  MaterialExpressionEditorX=-1968
+						  MaterialExpressionEditorY=352
+						  MaterialExpressionGuid=F611381D424E60EA8824DA9E1A50EFEB
+						  Material=PreviewMaterial'/Engine/Transient.PreviewMaterial_0'
+					   End Object
+					   MaterialExpression=MaterialExpressionConstant2Vector'MaterialExpressionConstant2Vector_8'
+					   NodePosX=-1968
+					   NodePosY=352
+					   NodeGuid=1B904D4E4E6CD3F1EC4DB88E72EB3C89
+					   CustomProperties Pin (PinId=F0B2AB5E41B850C4568355B630B18984,PinName="Output",PinFriendlyName=" ",Direction="EGPD_Output",PinType.PinSubCategoryObject=None,PinType.PinSubCategoryMemberReference=(),PinType.bIsArray=False,PinType.bIsReference=False,PinType.bIsConst=False,PinType.bIsWeakPointer=False,LinkedTo=(MaterialGraphNode_171 B657664A4FCD1A73583857A43E4C54A6,),PersistentGuid=00000000000000000000000000000000,bHidden=False,bNotConnectable=False,bDefaultValueIsReadOnly=False,bDefaultValueIsIgnored=False,bAdvancedView=False,)
+					End Object
+					Begin Object Class=MaterialGraphNode Name="MaterialGraphNode_156"
+					   Begin Object Class=MaterialExpressionTextureCoordinate Name="MaterialExpressionTextureCoordinate_8"
+					   End Object
+					   Begin Object Name="MaterialExpressionTextureCoordinate_8"
+						  MaterialExpressionEditorX=-1904
+						  MaterialExpressionEditorY=192
+						  MaterialExpressionGuid=F1C1E2C24651DC43FD6AD5B0AF633576
+						  Material=PreviewMaterial'/Engine/Transient.PreviewMaterial_0'
+					   End Object
+					   MaterialExpression=MaterialExpressionTextureCoordinate'MaterialExpressionTextureCoordinate_8'
+					   NodePosX=-1904
+					   NodePosY=192
+					   NodeGuid=89DCD0934C3D5F4BFDE05EB0B6EE83A3
+					   CustomProperties Pin (PinId=4691EC594DDFCB84DEE6E080AEF820CB,PinName="Output",PinFriendlyName=" ",Direction="EGPD_Output",PinType.PinSubCategoryObject=None,PinType.PinSubCategoryMemberReference=(),PinType.bIsArray=False,PinType.bIsReference=False,PinType.bIsConst=False,PinType.bIsWeakPointer=False,LinkedTo=(MaterialGraphNode_157 B96D2712404FD8FCE7A0C3AF39882576,),PersistentGuid=00000000000000000000000000000000,bHidden=False,bNotConnectable=False,bDefaultValueIsReadOnly=False,bDefaultValueIsIgnored=False,bAdvancedView=False,)
+					End Object
+					Begin Object Class=MaterialGraphNode Name="MaterialGraphNode_157"
+					   Begin Object Class=MaterialExpressionAdd Name="MaterialExpressionAdd_16"
+					   End Object
+					   Begin Object Name="MaterialExpressionAdd_16"
+						  A=(Expression=MaterialExpressionTextureCoordinate'MaterialGraphNode_156.MaterialExpressionTextureCoordinate_8')
+						  B=(Expression=MaterialExpressionMultiply'MaterialGraphNode_171.MaterialExpressionMultiply_26')
+						  MaterialExpressionEditorX=-1744
+						  MaterialExpressionEditorY=224
+						  MaterialExpressionGuid=EF5DFD284086D3D0FDC39A86FEE5E825
+						  Material=PreviewMaterial'/Engine/Transient.PreviewMaterial_0'
+					   End Object
+					   MaterialExpression=MaterialExpressionAdd'MaterialExpressionAdd_16'
+					   NodePosX=-1744
+					   NodePosY=224
+					   NodeGuid=7FFA5A3F446ADB3C85AF0895EA84FB0C
+					   CustomProperties Pin (PinId=B96D2712404FD8FCE7A0C3AF39882576,PinName="A",PinType.PinCategory="optional",PinType.PinSubCategoryObject=None,PinType.PinSubCategoryMemberReference=(),PinType.bIsArray=False,PinType.bIsReference=False,PinType.bIsConst=False,PinType.bIsWeakPointer=False,LinkedTo=(MaterialGraphNode_156 4691EC594DDFCB84DEE6E080AEF820CB,),PersistentGuid=00000000000000000000000000000000,bHidden=False,bNotConnectable=False,bDefaultValueIsReadOnly=False,bDefaultValueIsIgnored=False,bAdvancedView=False,)
+					   CustomProperties Pin (PinId=920356B749BA9152B229F4B3EB889E76,PinName="B",PinType.PinCategory="optional",PinType.PinSubCategoryObject=None,PinType.PinSubCategoryMemberReference=(),PinType.bIsArray=False,PinType.bIsReference=False,PinType.bIsConst=False,PinType.bIsWeakPointer=False,LinkedTo=(MaterialGraphNode_171 4516072E4B787E4FE0FC078D5185246E,),PersistentGuid=00000000000000000000000000000000,bHidden=False,bNotConnectable=False,bDefaultValueIsReadOnly=False,bDefaultValueIsIgnored=False,bAdvancedView=False,)
+					   CustomProperties Pin (PinId=801803B740F37EFDD82FE9B7F901CD44,PinName="Output",PinFriendlyName=" ",Direction="EGPD_Output",PinType.PinSubCategoryObject=None,PinType.PinSubCategoryMemberReference=(),PinType.bIsArray=False,PinType.bIsReference=False,PinType.bIsConst=False,PinType.bIsWeakPointer=False,LinkedTo=(MaterialGraphNode_154 A25AC337495A10638C30FA93453959BE,),PersistentGuid=00000000000000000000000000000000,bHidden=False,bNotConnectable=False,bDefaultValueIsReadOnly=False,bDefaultValueIsIgnored=False,bAdvancedView=False,)
+					End Object
+					Begin Object Class=MaterialGraphNode Name="MaterialGraphNode_158"
+					   Begin Object Class=MaterialExpressionConstant2Vector Name="MaterialExpressionConstant2Vector_9"
+					   End Object
+					   Begin Object Name="MaterialExpressionConstant2Vector_9"
+						  R=-1.000000
+						  MaterialExpressionEditorX=-1968
+						  MaterialExpressionEditorY=544
+						  MaterialExpressionGuid=F611381D424E60EA8824DA9E1A50EFEB
+						  Material=PreviewMaterial'/Engine/Transient.PreviewMaterial_0'
+					   End Object
+					   MaterialExpression=MaterialExpressionConstant2Vector'MaterialExpressionConstant2Vector_9'
+					   NodePosX=-1968
+					   NodePosY=544
+					   NodeGuid=7D963DCD4DCA2088147BA99B97A42E21
+					   CustomProperties Pin (PinId=2FE8D18F4EC0AB1771E2F9938D9A4BB2,PinName="Output",PinFriendlyName=" ",Direction="EGPD_Output",PinType.PinSubCategoryObject=None,PinType.PinSubCategoryMemberReference=(),PinType.bIsArray=False,PinType.bIsReference=False,PinType.bIsConst=False,PinType.bIsWeakPointer=False,LinkedTo=(MaterialGraphNode_172 39F3EBC345BB2DF4A95A3CAD26D568F8,),PersistentGuid=00000000000000000000000000000000,bHidden=False,bNotConnectable=False,bDefaultValueIsReadOnly=False,bDefaultValueIsIgnored=False,bAdvancedView=False,)
+					End Object
+					Begin Object Class=MaterialGraphNode Name="MaterialGraphNode_159"
+					   Begin Object Class=MaterialExpressionTextureCoordinate Name="MaterialExpressionTextureCoordinate_9"
+					   End Object
+					   Begin Object Name="MaterialExpressionTextureCoordinate_9"
+						  MaterialExpressionEditorX=-1904
+						  MaterialExpressionEditorY=400
+						  MaterialExpressionGuid=F1C1E2C24651DC43FD6AD5B0AF633576
+						  Material=PreviewMaterial'/Engine/Transient.PreviewMaterial_0'
+					   End Object
+					   MaterialExpression=MaterialExpressionTextureCoordinate'MaterialExpressionTextureCoordinate_9'
+					   NodePosX=-1904
+					   NodePosY=400
+					   NodeGuid=90240D0F40C546A3AE43C7B6AA26F585
+					   CustomProperties Pin (PinId=BD6BE4D348F587EE716D8FAD9DD4FA2D,PinName="Output",PinFriendlyName=" ",Direction="EGPD_Output",PinType.PinSubCategoryObject=None,PinType.PinSubCategoryMemberReference=(),PinType.bIsArray=False,PinType.bIsReference=False,PinType.bIsConst=False,PinType.bIsWeakPointer=False,LinkedTo=(MaterialGraphNode_160 91911EC94FF3140208CD8AAEC7D62AAD,),PersistentGuid=00000000000000000000000000000000,bHidden=False,bNotConnectable=False,bDefaultValueIsReadOnly=False,bDefaultValueIsIgnored=False,bAdvancedView=False,)
+					End Object
+					Begin Object Class=MaterialGraphNode Name="MaterialGraphNode_160"
+					   Begin Object Class=MaterialExpressionAdd Name="MaterialExpressionAdd_17"
+					   End Object
+					   Begin Object Name="MaterialExpressionAdd_17"
+						  A=(Expression=MaterialExpressionTextureCoordinate'MaterialGraphNode_159.MaterialExpressionTextureCoordinate_9')
+						  B=(Expression=MaterialExpressionMultiply'MaterialGraphNode_172.MaterialExpressionMultiply_27')
+						  MaterialExpressionEditorX=-1744
+						  MaterialExpressionEditorY=432
+						  MaterialExpressionGuid=EF5DFD284086D3D0FDC39A86FEE5E825
+						  Material=PreviewMaterial'/Engine/Transient.PreviewMaterial_0'
+					   End Object
+					   MaterialExpression=MaterialExpressionAdd'MaterialExpressionAdd_17'
+					   NodePosX=-1744
+					   NodePosY=432
+					   NodeGuid=507C2BDF4DDFE4FB648173826DDF6D7B
+					   CustomProperties Pin (PinId=91911EC94FF3140208CD8AAEC7D62AAD,PinName="A",PinType.PinCategory="optional",PinType.PinSubCategoryObject=None,PinType.PinSubCategoryMemberReference=(),PinType.bIsArray=False,PinType.bIsReference=False,PinType.bIsConst=False,PinType.bIsWeakPointer=False,LinkedTo=(MaterialGraphNode_159 BD6BE4D348F587EE716D8FAD9DD4FA2D,),PersistentGuid=00000000000000000000000000000000,bHidden=False,bNotConnectable=False,bDefaultValueIsReadOnly=False,bDefaultValueIsIgnored=False,bAdvancedView=False,)
+					   CustomProperties Pin (PinId=503B4DE54CF04B50933C5A86F2BA3B00,PinName="B",PinType.PinCategory="optional",PinType.PinSubCategoryObject=None,PinType.PinSubCategoryMemberReference=(),PinType.bIsArray=False,PinType.bIsReference=False,PinType.bIsConst=False,PinType.bIsWeakPointer=False,LinkedTo=(MaterialGraphNode_172 7FA80168416592A75B02EB955B17DF36,),PersistentGuid=00000000000000000000000000000000,bHidden=False,bNotConnectable=False,bDefaultValueIsReadOnly=False,bDefaultValueIsIgnored=False,bAdvancedView=False,)
+					   CustomProperties Pin (PinId=FD7A842143B98083CB4034B6187E59F5,PinName="Output",PinFriendlyName=" ",Direction="EGPD_Output",PinType.PinSubCategoryObject=None,PinType.PinSubCategoryMemberReference=(),PinType.bIsArray=False,PinType.bIsReference=False,PinType.bIsConst=False,PinType.bIsWeakPointer=False,LinkedTo=(MaterialGraphNode_161 9043324B49F7D3C18D9422B6D9BBA0F8,),PersistentGuid=00000000000000000000000000000000,bHidden=False,bNotConnectable=False,bDefaultValueIsReadOnly=False,bDefaultValueIsIgnored=False,bAdvancedView=False,)
+					End Object
+					Begin Object Class=MaterialGraphNode Name="MaterialGraphNode_161"
+					   Begin Object Class=MaterialExpressionTextureSampleParameter2D Name="MaterialExpressionTextureSampleParameter2D_27"
+					   End Object
+					   Begin Object Name="MaterialExpressionTextureSampleParameter2D_27"
+						  ParameterName="Heightfield"
+						  ExpressionGUID=D709CDCF470864444AED4A8957D26C55
+						  Coordinates=(Expression=MaterialExpressionAdd'MaterialGraphNode_160.MaterialExpressionAdd_17')
+						  MaterialExpressionEditorX=-1616
+						  MaterialExpressionEditorY=416
+						  MaterialExpressionGuid=8F3166474EF9A6F6A6825185B78360BE
+						  Material=PreviewMaterial'/Engine/Transient.PreviewMaterial_0'
+					   End Object
+					   MaterialExpression=MaterialExpressionTextureSampleParameter2D'MaterialExpressionTextureSampleParameter2D_27'
+					   NodePosX=-1616
+					   NodePosY=416
+					   bCanRenameNode=True
+					   NodeGuid=7C2B776B473BF2FA90C10DABC896AFC6
+					   CustomProperties Pin (PinId=9043324B49F7D3C18D9422B6D9BBA0F8,PinName="UVs",PinType.PinCategory="optional",PinType.PinSubCategoryObject=None,PinType.PinSubCategoryMemberReference=(),PinType.bIsArray=False,PinType.bIsReference=False,PinType.bIsConst=False,PinType.bIsWeakPointer=False,LinkedTo=(MaterialGraphNode_160 FD7A842143B98083CB4034B6187E59F5,),PersistentGuid=00000000000000000000000000000000,bHidden=False,bNotConnectable=False,bDefaultValueIsReadOnly=False,bDefaultValueIsIgnored=False,bAdvancedView=False,)
+					   CustomProperties Pin (PinId=9202400D4B70EDB2A0E25993FE1673FB,PinName="Tex",PinType.PinCategory="optional",PinType.PinSubCategoryObject=None,PinType.PinSubCategoryMemberReference=(),PinType.bIsArray=False,PinType.bIsReference=False,PinType.bIsConst=False,PinType.bIsWeakPointer=False,PersistentGuid=00000000000000000000000000000000,bHidden=False,bNotConnectable=False,bDefaultValueIsReadOnly=False,bDefaultValueIsIgnored=False,bAdvancedView=False,)
+					   CustomProperties Pin (PinId=EB170C1241A51163FC3F27992F0D0076,PinName="Output",PinFriendlyName=" ",Direction="EGPD_Output",PinType.PinCategory="mask",PinType.PinSubCategoryObject=None,PinType.PinSubCategoryMemberReference=(),PinType.bIsArray=False,PinType.bIsReference=False,PinType.bIsConst=False,PinType.bIsWeakPointer=False,PersistentGuid=00000000000000000000000000000000,bHidden=False,bNotConnectable=False,bDefaultValueIsReadOnly=False,bDefaultValueIsIgnored=False,bAdvancedView=False,)
+					   CustomProperties Pin (PinId=B7DF7ECD43804F6D666F12BB51CE3166,PinName="Output2",PinFriendlyName=" ",Direction="EGPD_Output",PinType.PinCategory="mask",PinType.PinSubCategory="red",PinType.PinSubCategoryObject=None,PinType.PinSubCategoryMemberReference=(),PinType.bIsArray=False,PinType.bIsReference=False,PinType.bIsConst=False,PinType.bIsWeakPointer=False,LinkedTo=(MaterialGraphNode_152 9DA0FD994C0986C3C8A4E6B1AA61D5C9,),PersistentGuid=00000000000000000000000000000000,bHidden=False,bNotConnectable=False,bDefaultValueIsReadOnly=False,bDefaultValueIsIgnored=False,bAdvancedView=False,)
+					   CustomProperties Pin (PinId=527DE468409730B09E41ABBDCC592452,PinName="Output3",PinFriendlyName=" ",Direction="EGPD_Output",PinType.PinCategory="mask",PinType.PinSubCategory="green",PinType.PinSubCategoryObject=None,PinType.PinSubCategoryMemberReference=(),PinType.bIsArray=False,PinType.bIsReference=False,PinType.bIsConst=False,PinType.bIsWeakPointer=False,PersistentGuid=00000000000000000000000000000000,bHidden=False,bNotConnectable=False,bDefaultValueIsReadOnly=False,bDefaultValueIsIgnored=False,bAdvancedView=False,)
+					   CustomProperties Pin (PinId=CC9BA78E479839306AD0F4A6B7D7F98F,PinName="Output4",PinFriendlyName=" ",Direction="EGPD_Output",PinType.PinCategory="mask",PinType.PinSubCategory="blue",PinType.PinSubCategoryObject=None,PinType.PinSubCategoryMemberReference=(),PinType.bIsArray=False,PinType.bIsReference=False,PinType.bIsConst=False,PinType.bIsWeakPointer=False,PersistentGuid=00000000000000000000000000000000,bHidden=False,bNotConnectable=False,bDefaultValueIsReadOnly=False,bDefaultValueIsIgnored=False,bAdvancedView=False,)
+					   CustomProperties Pin (PinId=D17A070C401E9F40C3E2EAA3DF23ACCC,PinName="Output5",PinFriendlyName=" ",Direction="EGPD_Output",PinType.PinCategory="mask",PinType.PinSubCategory="alpha",PinType.PinSubCategoryObject=None,PinType.PinSubCategoryMemberReference=(),PinType.bIsArray=False,PinType.bIsReference=False,PinType.bIsConst=False,PinType.bIsWeakPointer=False,PersistentGuid=00000000000000000000000000000000,bHidden=False,bNotConnectable=False,bDefaultValueIsReadOnly=False,bDefaultValueIsIgnored=False,bAdvancedView=False,)
+					End Object
+					Begin Object Class=MaterialGraphNode Name="MaterialGraphNode_162"
+					   Begin Object Class=MaterialExpressionConstant2Vector Name="MaterialExpressionConstant2Vector_10"
+					   End Object
+					   Begin Object Name="MaterialExpressionConstant2Vector_10"
+						  G=1.000000
+						  MaterialExpressionEditorX=-1968
+						  MaterialExpressionEditorY=736
+						  MaterialExpressionGuid=F611381D424E60EA8824DA9E1A50EFEB
+						  Material=PreviewMaterial'/Engine/Transient.PreviewMaterial_0'
+					   End Object
+					   MaterialExpression=MaterialExpressionConstant2Vector'MaterialExpressionConstant2Vector_10'
+					   NodePosX=-1968
+					   NodePosY=736
+					   NodeGuid=8357582B415C69932F34B280C0889236
+					   CustomProperties Pin (PinId=306FBF144C276D38E8314AA4DC278F13,PinName="Output",PinFriendlyName=" ",Direction="EGPD_Output",PinType.PinSubCategoryObject=None,PinType.PinSubCategoryMemberReference=(),PinType.bIsArray=False,PinType.bIsReference=False,PinType.bIsConst=False,PinType.bIsWeakPointer=False,LinkedTo=(MaterialGraphNode_173 E2742465403470C0EBFA3388ABB8C357,),PersistentGuid=00000000000000000000000000000000,bHidden=False,bNotConnectable=False,bDefaultValueIsReadOnly=False,bDefaultValueIsIgnored=False,bAdvancedView=False,)
+					End Object
+					Begin Object Class=MaterialGraphNode Name="MaterialGraphNode_163"
+					   Begin Object Class=MaterialExpressionTextureCoordinate Name="MaterialExpressionTextureCoordinate_10"
+					   End Object
+					   Begin Object Name="MaterialExpressionTextureCoordinate_10"
+						  MaterialExpressionEditorX=-1920
+						  MaterialExpressionEditorY=608
+						  MaterialExpressionGuid=F1C1E2C24651DC43FD6AD5B0AF633576
+						  Material=PreviewMaterial'/Engine/Transient.PreviewMaterial_0'
+					   End Object
+					   MaterialExpression=MaterialExpressionTextureCoordinate'MaterialExpressionTextureCoordinate_10'
+					   NodePosX=-1920
+					   NodePosY=608
+					   NodeGuid=ECD35065477EFD7C2A4FAAA10F6B5075
+					   CustomProperties Pin (PinId=84D091414557B99F4C647C888CD0AA0E,PinName="Output",PinFriendlyName=" ",Direction="EGPD_Output",PinType.PinSubCategoryObject=None,PinType.PinSubCategoryMemberReference=(),PinType.bIsArray=False,PinType.bIsReference=False,PinType.bIsConst=False,PinType.bIsWeakPointer=False,LinkedTo=(MaterialGraphNode_164 794F0E814B6186DBACB3E2B4AB4F44AD,),PersistentGuid=00000000000000000000000000000000,bHidden=False,bNotConnectable=False,bDefaultValueIsReadOnly=False,bDefaultValueIsIgnored=False,bAdvancedView=False,)
+					End Object
+					Begin Object Class=MaterialGraphNode Name="MaterialGraphNode_164"
+					   Begin Object Class=MaterialExpressionAdd Name="MaterialExpressionAdd_18"
+					   End Object
+					   Begin Object Name="MaterialExpressionAdd_18"
+						  A=(Expression=MaterialExpressionTextureCoordinate'MaterialGraphNode_163.MaterialExpressionTextureCoordinate_10')
+						  B=(Expression=MaterialExpressionMultiply'MaterialGraphNode_173.MaterialExpressionMultiply_28')
+						  MaterialExpressionEditorX=-1728
+						  MaterialExpressionEditorY=656
+						  MaterialExpressionGuid=EF5DFD284086D3D0FDC39A86FEE5E825
+						  Material=PreviewMaterial'/Engine/Transient.PreviewMaterial_0'
+					   End Object
+					   MaterialExpression=MaterialExpressionAdd'MaterialExpressionAdd_18'
+					   NodePosX=-1728
+					   NodePosY=656
+					   NodeGuid=38C43FD94E45400DAF8CCEA10F8615BF
+					   CustomProperties Pin (PinId=794F0E814B6186DBACB3E2B4AB4F44AD,PinName="A",PinType.PinCategory="optional",PinType.PinSubCategoryObject=None,PinType.PinSubCategoryMemberReference=(),PinType.bIsArray=False,PinType.bIsReference=False,PinType.bIsConst=False,PinType.bIsWeakPointer=False,LinkedTo=(MaterialGraphNode_163 84D091414557B99F4C647C888CD0AA0E,),PersistentGuid=00000000000000000000000000000000,bHidden=False,bNotConnectable=False,bDefaultValueIsReadOnly=False,bDefaultValueIsIgnored=False,bAdvancedView=False,)
+					   CustomProperties Pin (PinId=D62F1B144C195039083297AC1B4ED0A0,PinName="B",PinType.PinCategory="optional",PinType.PinSubCategoryObject=None,PinType.PinSubCategoryMemberReference=(),PinType.bIsArray=False,PinType.bIsReference=False,PinType.bIsConst=False,PinType.bIsWeakPointer=False,LinkedTo=(MaterialGraphNode_173 B728F1CE40E6AF8FB5DDB9BB562EA784,),PersistentGuid=00000000000000000000000000000000,bHidden=False,bNotConnectable=False,bDefaultValueIsReadOnly=False,bDefaultValueIsIgnored=False,bAdvancedView=False,)
+					   CustomProperties Pin (PinId=9CAD60C744EE671EBF684EA662DDE05A,PinName="Output",PinFriendlyName=" ",Direction="EGPD_Output",PinType.PinSubCategoryObject=None,PinType.PinSubCategoryMemberReference=(),PinType.bIsArray=False,PinType.bIsReference=False,PinType.bIsConst=False,PinType.bIsWeakPointer=False,LinkedTo=(MaterialGraphNode_165 3290366A4FD5AAD90F4CC5B390FC4CDD,),PersistentGuid=00000000000000000000000000000000,bHidden=False,bNotConnectable=False,bDefaultValueIsReadOnly=False,bDefaultValueIsIgnored=False,bAdvancedView=False,)
+					End Object
+					Begin Object Class=MaterialGraphNode Name="MaterialGraphNode_165"
+					   Begin Object Class=MaterialExpressionTextureSampleParameter2D Name="MaterialExpressionTextureSampleParameter2D_28"
+					   End Object
+					   Begin Object Name="MaterialExpressionTextureSampleParameter2D_28"
+						  ParameterName="Heightfield"
+						  ExpressionGUID=89BD9DA146CE12B48308C9BF04CB4CD8
+						  Coordinates=(Expression=MaterialExpressionAdd'MaterialGraphNode_164.MaterialExpressionAdd_18')
+						  MaterialExpressionEditorX=-1616
+						  MaterialExpressionEditorY=656
+						  MaterialExpressionGuid=8F3166474EF9A6F6A6825185B78360BE
+						  Material=PreviewMaterial'/Engine/Transient.PreviewMaterial_0'
+					   End Object
+					   MaterialExpression=MaterialExpressionTextureSampleParameter2D'MaterialExpressionTextureSampleParameter2D_28'
+					   NodePosX=-1616
+					   NodePosY=656
+					   bCanRenameNode=True
+					   NodeGuid=B97DE2824077531967C973BF29C2A188
+					   CustomProperties Pin (PinId=3290366A4FD5AAD90F4CC5B390FC4CDD,PinName="UVs",PinType.PinCategory="optional",PinType.PinSubCategoryObject=None,PinType.PinSubCategoryMemberReference=(),PinType.bIsArray=False,PinType.bIsReference=False,PinType.bIsConst=False,PinType.bIsWeakPointer=False,LinkedTo=(MaterialGraphNode_164 9CAD60C744EE671EBF684EA662DDE05A,),PersistentGuid=00000000000000000000000000000000,bHidden=False,bNotConnectable=False,bDefaultValueIsReadOnly=False,bDefaultValueIsIgnored=False,bAdvancedView=False,)
+					   CustomProperties Pin (PinId=FD1DE5D04CED6C1E572E6099D910DD11,PinName="Tex",PinType.PinCategory="optional",PinType.PinSubCategoryObject=None,PinType.PinSubCategoryMemberReference=(),PinType.bIsArray=False,PinType.bIsReference=False,PinType.bIsConst=False,PinType.bIsWeakPointer=False,PersistentGuid=00000000000000000000000000000000,bHidden=False,bNotConnectable=False,bDefaultValueIsReadOnly=False,bDefaultValueIsIgnored=False,bAdvancedView=False,)
+					   CustomProperties Pin (PinId=7A21F305411F8D7DA96FDFA91AA4A294,PinName="Output",PinFriendlyName=" ",Direction="EGPD_Output",PinType.PinCategory="mask",PinType.PinSubCategoryObject=None,PinType.PinSubCategoryMemberReference=(),PinType.bIsArray=False,PinType.bIsReference=False,PinType.bIsConst=False,PinType.bIsWeakPointer=False,PersistentGuid=00000000000000000000000000000000,bHidden=False,bNotConnectable=False,bDefaultValueIsReadOnly=False,bDefaultValueIsIgnored=False,bAdvancedView=False,)
+					   CustomProperties Pin (PinId=87A942C9487C8643262DC48786DE2B90,PinName="Output2",PinFriendlyName=" ",Direction="EGPD_Output",PinType.PinCategory="mask",PinType.PinSubCategory="red",PinType.PinSubCategoryObject=None,PinType.PinSubCategoryMemberReference=(),PinType.bIsArray=False,PinType.bIsReference=False,PinType.bIsConst=False,PinType.bIsWeakPointer=False,LinkedTo=(MaterialGraphNode_153 87AEDFE8427805C6775CAFB43BBDF596,),PersistentGuid=00000000000000000000000000000000,bHidden=False,bNotConnectable=False,bDefaultValueIsReadOnly=False,bDefaultValueIsIgnored=False,bAdvancedView=False,)
+					   CustomProperties Pin (PinId=E5EBF53E497E100909F297B585DD6771,PinName="Output3",PinFriendlyName=" ",Direction="EGPD_Output",PinType.PinCategory="mask",PinType.PinSubCategory="green",PinType.PinSubCategoryObject=None,PinType.PinSubCategoryMemberReference=(),PinType.bIsArray=False,PinType.bIsReference=False,PinType.bIsConst=False,PinType.bIsWeakPointer=False,PersistentGuid=00000000000000000000000000000000,bHidden=False,bNotConnectable=False,bDefaultValueIsReadOnly=False,bDefaultValueIsIgnored=False,bAdvancedView=False,)
+					   CustomProperties Pin (PinId=63433C0444AB2E2D9964CAAF2BB37C55,PinName="Output4",PinFriendlyName=" ",Direction="EGPD_Output",PinType.PinCategory="mask",PinType.PinSubCategory="blue",PinType.PinSubCategoryObject=None,PinType.PinSubCategoryMemberReference=(),PinType.bIsArray=False,PinType.bIsReference=False,PinType.bIsConst=False,PinType.bIsWeakPointer=False,PersistentGuid=00000000000000000000000000000000,bHidden=False,bNotConnectable=False,bDefaultValueIsReadOnly=False,bDefaultValueIsIgnored=False,bAdvancedView=False,)
+					   CustomProperties Pin (PinId=03BCD8C4403DF4C42142BBB87101A20D,PinName="Output5",PinFriendlyName=" ",Direction="EGPD_Output",PinType.PinCategory="mask",PinType.PinSubCategory="alpha",PinType.PinSubCategoryObject=None,PinType.PinSubCategoryMemberReference=(),PinType.bIsArray=False,PinType.bIsReference=False,PinType.bIsConst=False,PinType.bIsWeakPointer=False,PersistentGuid=00000000000000000000000000000000,bHidden=False,bNotConnectable=False,bDefaultValueIsReadOnly=False,bDefaultValueIsIgnored=False,bAdvancedView=False,)
+					End Object
+					Begin Object Class=MaterialGraphNode Name="MaterialGraphNode_166"
+					   Begin Object Class=MaterialExpressionConstant2Vector Name="MaterialExpressionConstant2Vector_11"
+					   End Object
+					   Begin Object Name="MaterialExpressionConstant2Vector_11"
+						  G=-1.000000
+						  MaterialExpressionEditorX=-1968
+						  MaterialExpressionEditorY=960
+						  MaterialExpressionGuid=F611381D424E60EA8824DA9E1A50EFEB
+						  Material=PreviewMaterial'/Engine/Transient.PreviewMaterial_0'
+					   End Object
+					   MaterialExpression=MaterialExpressionConstant2Vector'MaterialExpressionConstant2Vector_11'
+					   NodePosX=-1968
+					   NodePosY=960
+					   NodeGuid=EB780D944B98BA4568626BB551BC0718
+					   CustomProperties Pin (PinId=968DDAD440FA87BAD6CB61B0295E30D5,PinName="Output",PinFriendlyName=" ",Direction="EGPD_Output",PinType.PinSubCategoryObject=None,PinType.PinSubCategoryMemberReference=(),PinType.bIsArray=False,PinType.bIsReference=False,PinType.bIsConst=False,PinType.bIsWeakPointer=False,LinkedTo=(MaterialGraphNode_174 2AC39AE04630D9978B16708425590D8F,),PersistentGuid=00000000000000000000000000000000,bHidden=False,bNotConnectable=False,bDefaultValueIsReadOnly=False,bDefaultValueIsIgnored=False,bAdvancedView=False,)
+					End Object
+					Begin Object Class=MaterialGraphNode Name="MaterialGraphNode_167"
+					   Begin Object Class=MaterialExpressionTextureCoordinate Name="MaterialExpressionTextureCoordinate_11"
+					   End Object
+					   Begin Object Name="MaterialExpressionTextureCoordinate_11"
+						  MaterialExpressionEditorX=-1904
+						  MaterialExpressionEditorY=832
+						  MaterialExpressionGuid=F1C1E2C24651DC43FD6AD5B0AF633576
+						  Material=PreviewMaterial'/Engine/Transient.PreviewMaterial_0'
+					   End Object
+					   MaterialExpression=MaterialExpressionTextureCoordinate'MaterialExpressionTextureCoordinate_11'
+					   NodePosX=-1904
+					   NodePosY=832
+					   NodeGuid=FDD695BB488300D0DA65C790474B13AA
+					   CustomProperties Pin (PinId=97BEF0FA41E3349A65E39792512AD9CD,PinName="Output",PinFriendlyName=" ",Direction="EGPD_Output",PinType.PinSubCategoryObject=None,PinType.PinSubCategoryMemberReference=(),PinType.bIsArray=False,PinType.bIsReference=False,PinType.bIsConst=False,PinType.bIsWeakPointer=False,LinkedTo=(MaterialGraphNode_168 21881DAA482C5E00B1739285251D9221,),PersistentGuid=00000000000000000000000000000000,bHidden=False,bNotConnectable=False,bDefaultValueIsReadOnly=False,bDefaultValueIsIgnored=False,bAdvancedView=False,)
+					End Object
+					Begin Object Class=MaterialGraphNode Name="MaterialGraphNode_168"
+					   Begin Object Class=MaterialExpressionAdd Name="MaterialExpressionAdd_19"
+					   End Object
+					   Begin Object Name="MaterialExpressionAdd_19"
+						  A=(Expression=MaterialExpressionTextureCoordinate'MaterialGraphNode_167.MaterialExpressionTextureCoordinate_11')
+						  B=(Expression=MaterialExpressionMultiply'MaterialGraphNode_174.MaterialExpressionMultiply_29')
+						  MaterialExpressionEditorX=-1728
+						  MaterialExpressionEditorY=880
+						  MaterialExpressionGuid=EF5DFD284086D3D0FDC39A86FEE5E825
+						  Material=PreviewMaterial'/Engine/Transient.PreviewMaterial_0'
+					   End Object
+					   MaterialExpression=MaterialExpressionAdd'MaterialExpressionAdd_19'
+					   NodePosX=-1728
+					   NodePosY=880
+					   NodeGuid=7CA13B8F472D507B545DE180C0F5EFAC
+					   CustomProperties Pin (PinId=21881DAA482C5E00B1739285251D9221,PinName="A",PinType.PinCategory="optional",PinType.PinSubCategoryObject=None,PinType.PinSubCategoryMemberReference=(),PinType.bIsArray=False,PinType.bIsReference=False,PinType.bIsConst=False,PinType.bIsWeakPointer=False,LinkedTo=(MaterialGraphNode_167 97BEF0FA41E3349A65E39792512AD9CD,),PersistentGuid=00000000000000000000000000000000,bHidden=False,bNotConnectable=False,bDefaultValueIsReadOnly=False,bDefaultValueIsIgnored=False,bAdvancedView=False,)
+					   CustomProperties Pin (PinId=2095D10640015CA7AFDC8DBAC2280461,PinName="B",PinType.PinCategory="optional",PinType.PinSubCategoryObject=None,PinType.PinSubCategoryMemberReference=(),PinType.bIsArray=False,PinType.bIsReference=False,PinType.bIsConst=False,PinType.bIsWeakPointer=False,LinkedTo=(MaterialGraphNode_174 D57C4D6343636919355B4D81457BA1B8,),PersistentGuid=00000000000000000000000000000000,bHidden=False,bNotConnectable=False,bDefaultValueIsReadOnly=False,bDefaultValueIsIgnored=False,bAdvancedView=False,)
+					   CustomProperties Pin (PinId=70CA0E58454989F89650219BCCB205AE,PinName="Output",PinFriendlyName=" ",Direction="EGPD_Output",PinType.PinSubCategoryObject=None,PinType.PinSubCategoryMemberReference=(),PinType.bIsArray=False,PinType.bIsReference=False,PinType.bIsConst=False,PinType.bIsWeakPointer=False,LinkedTo=(MaterialGraphNode_169 569D1FCE4678A58BB9D1628EF6924D03,),PersistentGuid=00000000000000000000000000000000,bHidden=False,bNotConnectable=False,bDefaultValueIsReadOnly=False,bDefaultValueIsIgnored=False,bAdvancedView=False,)
+					End Object
+					Begin Object Class=MaterialGraphNode Name="MaterialGraphNode_169"
+					   Begin Object Class=MaterialExpressionTextureSampleParameter2D Name="MaterialExpressionTextureSampleParameter2D_29"
+					   End Object
+					   Begin Object Name="MaterialExpressionTextureSampleParameter2D_29"
+						  ParameterName="Heightfield"
+						  ExpressionGUID=FABEB0C24397290EDB8155BCC8F387AF
+						  Coordinates=(Expression=MaterialExpressionAdd'MaterialGraphNode_168.MaterialExpressionAdd_19')
+						  MaterialExpressionEditorX=-1616
+						  MaterialExpressionEditorY=864
+						  MaterialExpressionGuid=8F3166474EF9A6F6A6825185B78360BE
+						  Material=PreviewMaterial'/Engine/Transient.PreviewMaterial_0'
+					   End Object
+					   MaterialExpression=MaterialExpressionTextureSampleParameter2D'MaterialExpressionTextureSampleParameter2D_29'
+					   NodePosX=-1616
+					   NodePosY=864
+					   bCanRenameNode=True
+					   NodeGuid=9430D7F54F7642635CCDAE8B482A4D38
+					   CustomProperties Pin (PinId=569D1FCE4678A58BB9D1628EF6924D03,PinName="UVs",PinType.PinCategory="optional",PinType.PinSubCategoryObject=None,PinType.PinSubCategoryMemberReference=(),PinType.bIsArray=False,PinType.bIsReference=False,PinType.bIsConst=False,PinType.bIsWeakPointer=False,LinkedTo=(MaterialGraphNode_168 70CA0E58454989F89650219BCCB205AE,),PersistentGuid=00000000000000000000000000000000,bHidden=False,bNotConnectable=False,bDefaultValueIsReadOnly=False,bDefaultValueIsIgnored=False,bAdvancedView=False,)
+					   CustomProperties Pin (PinId=C92C9B694B2DB51771F9FDBA2896F72A,PinName="Tex",PinType.PinCategory="optional",PinType.PinSubCategoryObject=None,PinType.PinSubCategoryMemberReference=(),PinType.bIsArray=False,PinType.bIsReference=False,PinType.bIsConst=False,PinType.bIsWeakPointer=False,PersistentGuid=00000000000000000000000000000000,bHidden=False,bNotConnectable=False,bDefaultValueIsReadOnly=False,bDefaultValueIsIgnored=False,bAdvancedView=False,)
+					   CustomProperties Pin (PinId=6185E726433BACB3572C6E8A65DD5A16,PinName="Output",PinFriendlyName=" ",Direction="EGPD_Output",PinType.PinCategory="mask",PinType.PinSubCategoryObject=None,PinType.PinSubCategoryMemberReference=(),PinType.bIsArray=False,PinType.bIsReference=False,PinType.bIsConst=False,PinType.bIsWeakPointer=False,PersistentGuid=00000000000000000000000000000000,bHidden=False,bNotConnectable=False,bDefaultValueIsReadOnly=False,bDefaultValueIsIgnored=False,bAdvancedView=False,)
+					   CustomProperties Pin (PinId=2E05221B4F9937E3F66D36AE8E69475C,PinName="Output2",PinFriendlyName=" ",Direction="EGPD_Output",PinType.PinCategory="mask",PinType.PinSubCategory="red",PinType.PinSubCategoryObject=None,PinType.PinSubCategoryMemberReference=(),PinType.bIsArray=False,PinType.bIsReference=False,PinType.bIsConst=False,PinType.bIsWeakPointer=False,LinkedTo=(MaterialGraphNode_152 28DF7D474943CA7C66A104B7504A711B,),PersistentGuid=00000000000000000000000000000000,bHidden=False,bNotConnectable=False,bDefaultValueIsReadOnly=False,bDefaultValueIsIgnored=False,bAdvancedView=False,)
+					   CustomProperties Pin (PinId=2FA0BDD442C6EA58AA2749B308DADAF9,PinName="Output3",PinFriendlyName=" ",Direction="EGPD_Output",PinType.PinCategory="mask",PinType.PinSubCategory="green",PinType.PinSubCategoryObject=None,PinType.PinSubCategoryMemberReference=(),PinType.bIsArray=False,PinType.bIsReference=False,PinType.bIsConst=False,PinType.bIsWeakPointer=False,PersistentGuid=00000000000000000000000000000000,bHidden=False,bNotConnectable=False,bDefaultValueIsReadOnly=False,bDefaultValueIsIgnored=False,bAdvancedView=False,)
+					   CustomProperties Pin (PinId=1AB88BCE425055CCF1902D947C826BCB,PinName="Output4",PinFriendlyName=" ",Direction="EGPD_Output",PinType.PinCategory="mask",PinType.PinSubCategory="blue",PinType.PinSubCategoryObject=None,PinType.PinSubCategoryMemberReference=(),PinType.bIsArray=False,PinType.bIsReference=False,PinType.bIsConst=False,PinType.bIsWeakPointer=False,PersistentGuid=00000000000000000000000000000000,bHidden=False,bNotConnectable=False,bDefaultValueIsReadOnly=False,bDefaultValueIsIgnored=False,bAdvancedView=False,)
+					   CustomProperties Pin (PinId=77E65D0345E493751334D5BA2072A164,PinName="Output5",PinFriendlyName=" ",Direction="EGPD_Output",PinType.PinCategory="mask",PinType.PinSubCategory="alpha",PinType.PinSubCategoryObject=None,PinType.PinSubCategoryMemberReference=(),PinType.bIsArray=False,PinType.bIsReference=False,PinType.bIsConst=False,PinType.bIsWeakPointer=False,PersistentGuid=00000000000000000000000000000000,bHidden=False,bNotConnectable=False,bDefaultValueIsReadOnly=False,bDefaultValueIsIgnored=False,bAdvancedView=False,)
+					End Object
+					Begin Object Class=MaterialGraphNode Name="MaterialGraphNode_170"
+					   Begin Object Class=MaterialExpressionScalarParameter Name="MaterialExpressionScalarParameter_30"
+					   End Object
+					   Begin Object Name="MaterialExpressionScalarParameter_30"
+						  DefaultValue=0.003906
+						  ParameterName="TexelSize"
+						  ExpressionGUID=5D1179884D7449F6061BA6B5C72C2D17
+						  MaterialExpressionEditorX=-2256
+						  MaterialExpressionEditorY=496
+						  MaterialExpressionGuid=DB480CAC4000EF2A7EB61BA94FE23AF6
+						  Material=PreviewMaterial'/Engine/Transient.PreviewMaterial_0'
+					   End Object
+					   MaterialExpression=MaterialExpressionScalarParameter'MaterialExpressionScalarParameter_30'
+					   NodePosX=-2256
+					   NodePosY=496
+					   bCanRenameNode=True
+					   NodeGuid=EE95AA6E4EAAE2956A4321A8C849E4B0
+					   CustomProperties Pin (PinId=B0EA06BD43EF842B983A64954D2BCCAF,PinName="Output",PinFriendlyName=" ",Direction="EGPD_Output",PinType.PinSubCategoryObject=None,PinType.PinSubCategoryMemberReference=(),PinType.bIsArray=False,PinType.bIsReference=False,PinType.bIsConst=False,PinType.bIsWeakPointer=False,LinkedTo=(MaterialGraphNode_171 53E4A1C1439C6C9EEDCC349EA9B608AA,MaterialGraphNode_172 E015B2F749B73C1C841782BE0C03C33C,MaterialGraphNode_173 58AFC7A04310051EA26F35838A7A1FA0,MaterialGraphNode_174 3E0154D9473A9A9B3BD955AB552AD7D8,),PersistentGuid=00000000000000000000000000000000,bHidden=False,bNotConnectable=False,bDefaultValueIsReadOnly=False,bDefaultValueIsIgnored=False,bAdvancedView=False,)
+					End Object
+					Begin Object Class=MaterialGraphNode Name="MaterialGraphNode_171"
+					   Begin Object Class=MaterialExpressionMultiply Name="MaterialExpressionMultiply_26"
+					   End Object
+					   Begin Object Name="MaterialExpressionMultiply_26"
+						  A=(Expression=MaterialExpressionScalarParameter'MaterialGraphNode_170.MaterialExpressionScalarParameter_30')
+						  B=(Expression=MaterialExpressionConstant2Vector'MaterialGraphNode_155.MaterialExpressionConstant2Vector_8')
+						  MaterialExpressionEditorX=-1904
+						  MaterialExpressionEditorY=272
+						  MaterialExpressionGuid=9BB2CF3945F3869BD091C7AFFEBBD2A1
+						  Material=PreviewMaterial'/Engine/Transient.PreviewMaterial_0'
+					   End Object
+					   MaterialExpression=MaterialExpressionMultiply'MaterialExpressionMultiply_26'
+					   NodePosX=-1904
+					   NodePosY=272
+					   NodeGuid=442E0A584F5D3B711B831898FF8BEE0C
+					   CustomProperties Pin (PinId=53E4A1C1439C6C9EEDCC349EA9B608AA,PinName="A",PinType.PinCategory="optional",PinType.PinSubCategoryObject=None,PinType.PinSubCategoryMemberReference=(),PinType.bIsArray=False,PinType.bIsReference=False,PinType.bIsConst=False,PinType.bIsWeakPointer=False,LinkedTo=(MaterialGraphNode_170 B0EA06BD43EF842B983A64954D2BCCAF,),PersistentGuid=00000000000000000000000000000000,bHidden=False,bNotConnectable=False,bDefaultValueIsReadOnly=False,bDefaultValueIsIgnored=False,bAdvancedView=False,)
+					   CustomProperties Pin (PinId=B657664A4FCD1A73583857A43E4C54A6,PinName="B",PinType.PinCategory="optional",PinType.PinSubCategoryObject=None,PinType.PinSubCategoryMemberReference=(),PinType.bIsArray=False,PinType.bIsReference=False,PinType.bIsConst=False,PinType.bIsWeakPointer=False,LinkedTo=(MaterialGraphNode_155 F0B2AB5E41B850C4568355B630B18984,),PersistentGuid=00000000000000000000000000000000,bHidden=False,bNotConnectable=False,bDefaultValueIsReadOnly=False,bDefaultValueIsIgnored=False,bAdvancedView=False,)
+					   CustomProperties Pin (PinId=4516072E4B787E4FE0FC078D5185246E,PinName="Output",PinFriendlyName=" ",Direction="EGPD_Output",PinType.PinSubCategoryObject=None,PinType.PinSubCategoryMemberReference=(),PinType.bIsArray=False,PinType.bIsReference=False,PinType.bIsConst=False,PinType.bIsWeakPointer=False,LinkedTo=(MaterialGraphNode_157 920356B749BA9152B229F4B3EB889E76,),PersistentGuid=00000000000000000000000000000000,bHidden=False,bNotConnectable=False,bDefaultValueIsReadOnly=False,bDefaultValueIsIgnored=False,bAdvancedView=False,)
+					End Object
+					Begin Object Class=MaterialGraphNode Name="MaterialGraphNode_172"
+					   Begin Object Class=MaterialExpressionMultiply Name="MaterialExpressionMultiply_27"
+					   End Object
+					   Begin Object Name="MaterialExpressionMultiply_27"
+						  A=(Expression=MaterialExpressionScalarParameter'MaterialGraphNode_170.MaterialExpressionScalarParameter_30')
+						  B=(Expression=MaterialExpressionConstant2Vector'MaterialGraphNode_158.MaterialExpressionConstant2Vector_9')
+						  MaterialExpressionEditorX=-1888
+						  MaterialExpressionEditorY=480
+						  MaterialExpressionGuid=9BB2CF3945F3869BD091C7AFFEBBD2A1
+						  Material=PreviewMaterial'/Engine/Transient.PreviewMaterial_0'
+					   End Object
+					   MaterialExpression=MaterialExpressionMultiply'MaterialExpressionMultiply_27'
+					   NodePosX=-1888
+					   NodePosY=480
+					   NodeGuid=BC48D9CE4202EB72DB48DAB1536A3622
+					   CustomProperties Pin (PinId=E015B2F749B73C1C841782BE0C03C33C,PinName="A",PinType.PinCategory="optional",PinType.PinSubCategoryObject=None,PinType.PinSubCategoryMemberReference=(),PinType.bIsArray=False,PinType.bIsReference=False,PinType.bIsConst=False,PinType.bIsWeakPointer=False,LinkedTo=(MaterialGraphNode_170 B0EA06BD43EF842B983A64954D2BCCAF,),PersistentGuid=00000000000000000000000000000000,bHidden=False,bNotConnectable=False,bDefaultValueIsReadOnly=False,bDefaultValueIsIgnored=False,bAdvancedView=False,)
+					   CustomProperties Pin (PinId=39F3EBC345BB2DF4A95A3CAD26D568F8,PinName="B",PinType.PinCategory="optional",PinType.PinSubCategoryObject=None,PinType.PinSubCategoryMemberReference=(),PinType.bIsArray=False,PinType.bIsReference=False,PinType.bIsConst=False,PinType.bIsWeakPointer=False,LinkedTo=(MaterialGraphNode_158 2FE8D18F4EC0AB1771E2F9938D9A4BB2,),PersistentGuid=00000000000000000000000000000000,bHidden=False,bNotConnectable=False,bDefaultValueIsReadOnly=False,bDefaultValueIsIgnored=False,bAdvancedView=False,)
+					   CustomProperties Pin (PinId=7FA80168416592A75B02EB955B17DF36,PinName="Output",PinFriendlyName=" ",Direction="EGPD_Output",PinType.PinSubCategoryObject=None,PinType.PinSubCategoryMemberReference=(),PinType.bIsArray=False,PinType.bIsReference=False,PinType.bIsConst=False,PinType.bIsWeakPointer=False,LinkedTo=(MaterialGraphNode_160 503B4DE54CF04B50933C5A86F2BA3B00,),PersistentGuid=00000000000000000000000000000000,bHidden=False,bNotConnectable=False,bDefaultValueIsReadOnly=False,bDefaultValueIsIgnored=False,bAdvancedView=False,)
+					End Object
+					Begin Object Class=MaterialGraphNode Name="MaterialGraphNode_173"
+					   Begin Object Class=MaterialExpressionMultiply Name="MaterialExpressionMultiply_28"
+					   End Object
+					   Begin Object Name="MaterialExpressionMultiply_28"
+						  A=(Expression=MaterialExpressionScalarParameter'MaterialGraphNode_170.MaterialExpressionScalarParameter_30')
+						  B=(Expression=MaterialExpressionConstant2Vector'MaterialGraphNode_162.MaterialExpressionConstant2Vector_10')
+						  MaterialExpressionEditorX=-1888
+						  MaterialExpressionEditorY=688
+						  MaterialExpressionGuid=9BB2CF3945F3869BD091C7AFFEBBD2A1
+						  Material=PreviewMaterial'/Engine/Transient.PreviewMaterial_0'
+					   End Object
+					   MaterialExpression=MaterialExpressionMultiply'MaterialExpressionMultiply_28'
+					   NodePosX=-1888
+					   NodePosY=688
+					   NodeGuid=E0D312AD47D5133CCEBE4F89A24E74BE
+					   CustomProperties Pin (PinId=58AFC7A04310051EA26F35838A7A1FA0,PinName="A",PinType.PinCategory="optional",PinType.PinSubCategoryObject=None,PinType.PinSubCategoryMemberReference=(),PinType.bIsArray=False,PinType.bIsReference=False,PinType.bIsConst=False,PinType.bIsWeakPointer=False,LinkedTo=(MaterialGraphNode_170 B0EA06BD43EF842B983A64954D2BCCAF,),PersistentGuid=00000000000000000000000000000000,bHidden=False,bNotConnectable=False,bDefaultValueIsReadOnly=False,bDefaultValueIsIgnored=False,bAdvancedView=False,)
+					   CustomProperties Pin (PinId=E2742465403470C0EBFA3388ABB8C357,PinName="B",PinType.PinCategory="optional",PinType.PinSubCategoryObject=None,PinType.PinSubCategoryMemberReference=(),PinType.bIsArray=False,PinType.bIsReference=False,PinType.bIsConst=False,PinType.bIsWeakPointer=False,LinkedTo=(MaterialGraphNode_162 306FBF144C276D38E8314AA4DC278F13,),PersistentGuid=00000000000000000000000000000000,bHidden=False,bNotConnectable=False,bDefaultValueIsReadOnly=False,bDefaultValueIsIgnored=False,bAdvancedView=False,)
+					   CustomProperties Pin (PinId=B728F1CE40E6AF8FB5DDB9BB562EA784,PinName="Output",PinFriendlyName=" ",Direction="EGPD_Output",PinType.PinSubCategoryObject=None,PinType.PinSubCategoryMemberReference=(),PinType.bIsArray=False,PinType.bIsReference=False,PinType.bIsConst=False,PinType.bIsWeakPointer=False,LinkedTo=(MaterialGraphNode_164 D62F1B144C195039083297AC1B4ED0A0,),PersistentGuid=00000000000000000000000000000000,bHidden=False,bNotConnectable=False,bDefaultValueIsReadOnly=False,bDefaultValueIsIgnored=False,bAdvancedView=False,)
+					End Object
+					Begin Object Class=MaterialGraphNode Name="MaterialGraphNode_174"
+					   Begin Object Class=MaterialExpressionMultiply Name="MaterialExpressionMultiply_29"
+					   End Object
+					   Begin Object Name="MaterialExpressionMultiply_29"
+						  A=(Expression=MaterialExpressionScalarParameter'MaterialGraphNode_170.MaterialExpressionScalarParameter_30')
+						  B=(Expression=MaterialExpressionConstant2Vector'MaterialGraphNode_166.MaterialExpressionConstant2Vector_11')
+						  MaterialExpressionEditorX=-1856
+						  MaterialExpressionEditorY=928
+						  MaterialExpressionGuid=9BB2CF3945F3869BD091C7AFFEBBD2A1
+						  Material=PreviewMaterial'/Engine/Transient.PreviewMaterial_0'
+					   End Object
+					   MaterialExpression=MaterialExpressionMultiply'MaterialExpressionMultiply_29'
+					   NodePosX=-1856
+					   NodePosY=928
+					   NodeGuid=860F5B4041EC5A6B9631AA9B2AB155F8
+					   CustomProperties Pin (PinId=3E0154D9473A9A9B3BD955AB552AD7D8,PinName="A",PinType.PinCategory="optional",PinType.PinSubCategoryObject=None,PinType.PinSubCategoryMemberReference=(),PinType.bIsArray=False,PinType.bIsReference=False,PinType.bIsConst=False,PinType.bIsWeakPointer=False,LinkedTo=(MaterialGraphNode_170 B0EA06BD43EF842B983A64954D2BCCAF,),PersistentGuid=00000000000000000000000000000000,bHidden=False,bNotConnectable=False,bDefaultValueIsReadOnly=False,bDefaultValueIsIgnored=False,bAdvancedView=False,)
+					   CustomProperties Pin (PinId=2AC39AE04630D9978B16708425590D8F,PinName="B",PinType.PinCategory="optional",PinType.PinSubCategoryObject=None,PinType.PinSubCategoryMemberReference=(),PinType.bIsArray=False,PinType.bIsReference=False,PinType.bIsConst=False,PinType.bIsWeakPointer=False,LinkedTo=(MaterialGraphNode_166 968DDAD440FA87BAD6CB61B0295E30D5,),PersistentGuid=00000000000000000000000000000000,bHidden=False,bNotConnectable=False,bDefaultValueIsReadOnly=False,bDefaultValueIsIgnored=False,bAdvancedView=False,)
+					   CustomProperties Pin (PinId=D57C4D6343636919355B4D81457BA1B8,PinName="Output",PinFriendlyName=" ",Direction="EGPD_Output",PinType.PinSubCategoryObject=None,PinType.PinSubCategoryMemberReference=(),PinType.bIsArray=False,PinType.bIsReference=False,PinType.bIsConst=False,PinType.bIsWeakPointer=False,LinkedTo=(MaterialGraphNode_168 2095D10640015CA7AFDC8DBAC2280461,),PersistentGuid=00000000000000000000000000000000,bHidden=False,bNotConnectable=False,bDefaultValueIsReadOnly=False,bDefaultValueIsIgnored=False,bAdvancedView=False,)
+					End Object
+			END OBJECT
+	[/PARAM]
+[/OBJECT]
+[/EXCERPT:CompNormal]
+
+[EXCERPT:ForceSplat]	
+[OBJECT:Snippet]
+	[PARAMLITERAL:id]
+		MAT_ForceSplat
+	[/PARAMLITERAL]
+	[PARAMLITERAL:type]
+		material
+	[/PARAMLITERAL]
+	[PARAM:image]
+		![](Engine\Rendering\RenderTargets\BlueprintRenderTargets\HowTo\HeightFieldPainter\HFP_Material_Creation_03.png)
+	[/PARAM]
+	[PARAM:snippet]
+			Begin Object Class=MaterialGraphNode Name="MaterialGraphNode_0"
+			   Begin Object Class=MaterialExpressionMultiply Name="MaterialExpressionMultiply_2"
+			   End Object
+			   Begin Object Name="MaterialExpressionMultiply_2"
+				  A=(Expression=MaterialExpressionDivide'MaterialGraphNode_13.MaterialExpressionDivide_2')
+				  B=(Expression=MaterialExpressionScalarParameter'MaterialGraphNode_1.MaterialExpressionScalarParameter_4')
+				  MaterialExpressionEditorX=112
+				  MaterialExpressionEditorY=80
+				  MaterialExpressionGuid=3025F98C46D3F1D0C59CD387F8C33F29
+				  Material=PreviewMaterial'/Engine/Transient.ForceSplat'
+			   End Object
+			   MaterialExpression=MaterialExpressionMultiply'MaterialExpressionMultiply_2'
+			   NodePosX=112
+			   NodePosY=80
+			   NodeGuid=5AD1F5CD44FA09FD0FDAA1AFA22A4AA9
+			   CustomProperties Pin (PinId=B7D0CF28498836478F9C8987312501BF,PinName="A",PinType.PinCategory="optional",PinType.PinSubCategoryObject=None,PinType.PinSubCategoryMemberReference=(),PinType.bIsArray=False,PinType.bIsReference=False,PinType.bIsConst=False,PinType.bIsWeakPointer=False,LinkedTo=(MaterialGraphNode_13 D0CE6B2B4A6886F59532E4B83367D7BD,),PersistentGuid=00000000000000000000000000000000,bHidden=False,bNotConnectable=False,bDefaultValueIsReadOnly=False,bDefaultValueIsIgnored=False,bAdvancedView=False,)
+			   CustomProperties Pin (PinId=362F295247D0A2B0935DB9889A9995EB,PinName="B",PinType.PinCategory="optional",PinType.PinSubCategoryObject=None,PinType.PinSubCategoryMemberReference=(),PinType.bIsArray=False,PinType.bIsReference=False,PinType.bIsConst=False,PinType.bIsWeakPointer=False,LinkedTo=(MaterialGraphNode_1 751E250A4F77873070E2D29BD2471826,),PersistentGuid=00000000000000000000000000000000,bHidden=False,bNotConnectable=False,bDefaultValueIsReadOnly=False,bDefaultValueIsIgnored=False,bAdvancedView=False,)
+			   CustomProperties Pin (PinId=1F55ED9148C9F02DB35DA78AA20C458D,PinName="Output",PinFriendlyName=" ",Direction="EGPD_Output",PinType.PinSubCategoryObject=None,PinType.PinSubCategoryMemberReference=(),PinType.bIsArray=False,PinType.bIsReference=False,PinType.bIsConst=False,PinType.bIsWeakPointer=False,LinkedTo=(MaterialGraphNode_Root_0 3890D74A48618B9B44F8AA8F876CD0C8,),PersistentGuid=00000000000000000000000000000000,bHidden=False,bNotConnectable=False,bDefaultValueIsReadOnly=False,bDefaultValueIsIgnored=False,bAdvancedView=False,)
+			End Object
+			Begin Object Class=MaterialGraphNode Name="MaterialGraphNode_1"
+			   Begin Object Class=MaterialExpressionScalarParameter Name="MaterialExpressionScalarParameter_4"
+			   End Object
+			   Begin Object Name="MaterialExpressionScalarParameter_4"
+				  DefaultValue=1.000000
+				  ParameterName="ForceStrength"
+				  ExpressionGUID=EF5F282F41F66D3274BC6FA7B987DCC9
+				  MaterialExpressionEditorX=-64
+				  MaterialExpressionEditorY=192
+				  MaterialExpressionGuid=9EF14C914FADADA4DBDE7B9C45E23576
+				  Material=PreviewMaterial'/Engine/Transient.ForceSplat'
+			   End Object
+			   MaterialExpression=MaterialExpressionScalarParameter'MaterialExpressionScalarParameter_4'
+			   NodePosX=-64
+			   NodePosY=192
+			   bCanRenameNode=True
+			   NodeGuid=E2EF5596438E8DA9CBD9509FBD3BF3F5
+			   CustomProperties Pin (PinId=751E250A4F77873070E2D29BD2471826,PinName="Output",PinFriendlyName=" ",Direction="EGPD_Output",PinType.PinSubCategoryObject=None,PinType.PinSubCategoryMemberReference=(),PinType.bIsArray=False,PinType.bIsReference=False,PinType.bIsConst=False,PinType.bIsWeakPointer=False,LinkedTo=(MaterialGraphNode_0 362F295247D0A2B0935DB9889A9995EB,),PersistentGuid=00000000000000000000000000000000,bHidden=False,bNotConnectable=False,bDefaultValueIsReadOnly=False,bDefaultValueIsIgnored=False,bAdvancedView=False,)
+			End Object
+			Begin Object Class=MaterialGraphNode Name="MaterialGraphNode_2"
+			   Begin Object Class=MaterialExpressionSubtract Name="MaterialExpressionSubtract_6"
+			   End Object
+			   Begin Object Name="MaterialExpressionSubtract_6"
+				  A=(Expression=MaterialExpressionSubtract'MaterialGraphNode_8.MaterialExpressionSubtract_7')
+				  B=(Expression=MaterialExpressionOneMinus'MaterialGraphNode_3.MaterialExpressionOneMinus_2')
+				  MaterialExpressionEditorX=-496
+				  MaterialExpressionEditorY=32
+				  MaterialExpressionGuid=0F3C46254F00CCEB93571093119ECB9D
+				  Material=PreviewMaterial'/Engine/Transient.ForceSplat'
+			   End Object
+			   MaterialExpression=MaterialExpressionSubtract'MaterialExpressionSubtract_6'
+			   NodePosX=-496
+			   NodePosY=32
+			   NodeGuid=478BC67945A3DC7AAD1D518FEE303739
+			   CustomProperties Pin (PinId=C5F733EF466287E9E4EB35ACE6B65E84,PinName="A",PinType.PinCategory="optional",PinType.PinSubCategoryObject=None,PinType.PinSubCategoryMemberReference=(),PinType.bIsArray=False,PinType.bIsReference=False,PinType.bIsConst=False,PinType.bIsWeakPointer=False,LinkedTo=(MaterialGraphNode_8 38FFC1874298113EDA7D9DBC71D7C22F,),PersistentGuid=00000000000000000000000000000000,bHidden=False,bNotConnectable=False,bDefaultValueIsReadOnly=False,bDefaultValueIsIgnored=False,bAdvancedView=False,)
+			   CustomProperties Pin (PinId=6AADA53C49FE3593A2AB3D9626E59461,PinName="B",PinType.PinCategory="optional",PinType.PinSubCategoryObject=None,PinType.PinSubCategoryMemberReference=(),PinType.bIsArray=False,PinType.bIsReference=False,PinType.bIsConst=False,PinType.bIsWeakPointer=False,LinkedTo=(MaterialGraphNode_3 368EE4A347A49750642B86858743DC92,),PersistentGuid=00000000000000000000000000000000,bHidden=False,bNotConnectable=False,bDefaultValueIsReadOnly=False,bDefaultValueIsIgnored=False,bAdvancedView=False,)
+			   CustomProperties Pin (PinId=25D2D0404B46764F28521E9B20B14B64,PinName="Output",PinFriendlyName=" ",Direction="EGPD_Output",PinType.PinSubCategoryObject=None,PinType.PinSubCategoryMemberReference=(),PinType.bIsArray=False,PinType.bIsReference=False,PinType.bIsConst=False,PinType.bIsWeakPointer=False,LinkedTo=(MaterialGraphNode_7 42C937124B3EBFB7D0BFD98373FA7A0B,),PersistentGuid=00000000000000000000000000000000,bHidden=False,bNotConnectable=False,bDefaultValueIsReadOnly=False,bDefaultValueIsIgnored=False,bAdvancedView=False,)
+			End Object
+			Begin Object Class=MaterialGraphNode Name="MaterialGraphNode_3"
+			   Begin Object Class=MaterialExpressionOneMinus Name="MaterialExpressionOneMinus_2"
+			   End Object
+			   Begin Object Name="MaterialExpressionOneMinus_2"
+				  Input=(Expression=MaterialExpressionScalarParameter'MaterialGraphNode_4.MaterialExpressionScalarParameter_5')
+				  MaterialExpressionEditorX=-656
+				  MaterialExpressionEditorY=64
+				  MaterialExpressionGuid=0FC613604260BF8A67765EB5C5556AE0
+				  Material=PreviewMaterial'/Engine/Transient.ForceSplat'
+			   End Object
+			   MaterialExpression=MaterialExpressionOneMinus'MaterialExpressionOneMinus_2'
+			   NodePosX=-656
+			   NodePosY=64
+			   NodeGuid=1530583D4196B0C70F53138D51ED4AF4
+			   CustomProperties Pin (PinId=98DDAD8444773CE90CF776BF47876227,PinName="Input",PinFriendlyName=" ",PinType.PinCategory="required",PinType.PinSubCategoryObject=None,PinType.PinSubCategoryMemberReference=(),PinType.bIsArray=False,PinType.bIsReference=False,PinType.bIsConst=False,PinType.bIsWeakPointer=False,LinkedTo=(MaterialGraphNode_4 243D5A9F4A177E21AD5B0C9A9EE13B6C,),PersistentGuid=00000000000000000000000000000000,bHidden=False,bNotConnectable=False,bDefaultValueIsReadOnly=False,bDefaultValueIsIgnored=False,bAdvancedView=False,)
+			   CustomProperties Pin (PinId=368EE4A347A49750642B86858743DC92,PinName="Output",PinFriendlyName=" ",Direction="EGPD_Output",PinType.PinSubCategoryObject=None,PinType.PinSubCategoryMemberReference=(),PinType.bIsArray=False,PinType.bIsReference=False,PinType.bIsConst=False,PinType.bIsWeakPointer=False,LinkedTo=(MaterialGraphNode_2 6AADA53C49FE3593A2AB3D9626E59461,),PersistentGuid=00000000000000000000000000000000,bHidden=False,bNotConnectable=False,bDefaultValueIsReadOnly=False,bDefaultValueIsIgnored=False,bAdvancedView=False,)
+			End Object
+			Begin Object Class=MaterialGraphNode Name="MaterialGraphNode_4"
+			   Begin Object Class=MaterialExpressionScalarParameter Name="MaterialExpressionScalarParameter_5"
+			   End Object
+			   Begin Object Name="MaterialExpressionScalarParameter_5"
+				  DefaultValue=0.020000
+				  ParameterName="ForceSize"
+				  ExpressionGUID=22FEC1D646BC1B274C6F6B8AF07EC523
+				  MaterialExpressionEditorX=-880
+				  MaterialExpressionEditorY=224
+				  MaterialExpressionGuid=08209E994F6A7C27F8AF1B987C8BA87A
+				  Material=PreviewMaterial'/Engine/Transient.ForceSplat'
+			   End Object
+			   MaterialExpression=MaterialExpressionScalarParameter'MaterialExpressionScalarParameter_5'
+			   NodePosX=-880
+			   NodePosY=224
+			   bCanRenameNode=True
+			   NodeGuid=3C5CD52A4E4EACE650CA47871ED15868
+			   CustomProperties Pin (PinId=243D5A9F4A177E21AD5B0C9A9EE13B6C,PinName="Output",PinFriendlyName=" ",Direction="EGPD_Output",PinType.PinSubCategoryObject=None,PinType.PinSubCategoryMemberReference=(),PinType.bIsArray=False,PinType.bIsReference=False,PinType.bIsConst=False,PinType.bIsWeakPointer=False,LinkedTo=(MaterialGraphNode_3 98DDAD8444773CE90CF776BF47876227,MaterialGraphNode_13 3501ECFD4967CB5C6FA143BB453E5636,),PersistentGuid=00000000000000000000000000000000,bHidden=False,bNotConnectable=False,bDefaultValueIsReadOnly=False,bDefaultValueIsIgnored=False,bAdvancedView=False,)
+			End Object
+			Begin Object Class=MaterialGraphNode Name="MaterialGraphNode_5"
+			   Begin Object Class=MaterialExpressionComponentMask Name="MaterialExpressionComponentMask_2"
+			   End Object
+			   Begin Object Name="MaterialExpressionComponentMask_2"
+				  Input=(Expression=MaterialExpressionVectorParameter'MaterialGraphNode_6.MaterialExpressionVectorParameter_2',Mask=1,MaskR=1,MaskG=1,MaskB=1)
+				  R=True
+				  G=True
+				  MaterialExpressionEditorX=-1280
+				  MaterialExpressionEditorY=48
+				  MaterialExpressionGuid=FF2B6BA24182D92569D49491035FFA18
+				  Material=PreviewMaterial'/Engine/Transient.ForceSplat'
+			   End Object
+			   MaterialExpression=MaterialExpressionComponentMask'MaterialExpressionComponentMask_2'
+			   NodePosX=-1280
+			   NodePosY=48
+			   NodeGuid=FC6C573F4C2431384FF35987120D12B5
+			   CustomProperties Pin (PinId=EC552D08411CF9140F7BCDA9D904E890,PinName="Input",PinFriendlyName=" ",PinType.PinCategory="required",PinType.PinSubCategoryObject=None,PinType.PinSubCategoryMemberReference=(),PinType.bIsArray=False,PinType.bIsReference=False,PinType.bIsConst=False,PinType.bIsWeakPointer=False,LinkedTo=(MaterialGraphNode_6 7B0A20374241C6FD036BDFADE36D1DD7,),PersistentGuid=00000000000000000000000000000000,bHidden=False,bNotConnectable=False,bDefaultValueIsReadOnly=False,bDefaultValueIsIgnored=False,bAdvancedView=False,)
+			   CustomProperties Pin (PinId=2A7AE85A415E54152002F9BA674B7E41,PinName="Output",PinFriendlyName=" ",Direction="EGPD_Output",PinType.PinSubCategoryObject=None,PinType.PinSubCategoryMemberReference=(),PinType.bIsArray=False,PinType.bIsReference=False,PinType.bIsConst=False,PinType.bIsWeakPointer=False,LinkedTo=(MaterialGraphNode_10 26815BDA44BF1141EF0D0084C2E3D993,),PersistentGuid=00000000000000000000000000000000,bHidden=False,bNotConnectable=False,bDefaultValueIsReadOnly=False,bDefaultValueIsIgnored=False,bAdvancedView=False,)
+			End Object
+			Begin Object Class=MaterialGraphNode Name="MaterialGraphNode_6"
+			   Begin Object Class=MaterialExpressionVectorParameter Name="MaterialExpressionVectorParameter_2"
+			   End Object
+			   Begin Object Name="MaterialExpressionVectorParameter_2"
+				  DefaultValue=(R=0.500000,G=0.500000,B=0.000000,A=0.000000)
+				  ParameterName="ForcePosition"
+				  ExpressionGUID=EB3126774C539A8048339782BDD5212C
+				  MaterialExpressionEditorX=-1472
+				  MaterialExpressionEditorY=32
+				  MaterialExpressionGuid=75C595794CE24CE6008419B632EE3760
+				  Material=PreviewMaterial'/Engine/Transient.ForceSplat'
+			   End Object
+			   MaterialExpression=MaterialExpressionVectorParameter'MaterialExpressionVectorParameter_2'
+			   NodePosX=-1472
+			   NodePosY=32
+			   bCanRenameNode=True
+			   NodeGuid=3869F5FA462D6E66ABABC89D92F2BF20
+			   CustomProperties Pin (PinId=7B0A20374241C6FD036BDFADE36D1DD7,PinName="Output",PinFriendlyName=" ",Direction="EGPD_Output",PinType.PinCategory="mask",PinType.PinSubCategoryObject=None,PinType.PinSubCategoryMemberReference=(),PinType.bIsArray=False,PinType.bIsReference=False,PinType.bIsConst=False,PinType.bIsWeakPointer=False,LinkedTo=(MaterialGraphNode_5 EC552D08411CF9140F7BCDA9D904E890,),PersistentGuid=00000000000000000000000000000000,bHidden=False,bNotConnectable=False,bDefaultValueIsReadOnly=False,bDefaultValueIsIgnored=False,bAdvancedView=False,)
+			   CustomProperties Pin (PinId=DA45DDA34307453D2D4A169D672CADA4,PinName="Output2",PinFriendlyName=" ",Direction="EGPD_Output",PinType.PinCategory="mask",PinType.PinSubCategory="red",PinType.PinSubCategoryObject=None,PinType.PinSubCategoryMemberReference=(),PinType.bIsArray=False,PinType.bIsReference=False,PinType.bIsConst=False,PinType.bIsWeakPointer=False,PersistentGuid=00000000000000000000000000000000,bHidden=False,bNotConnectable=False,bDefaultValueIsReadOnly=False,bDefaultValueIsIgnored=False,bAdvancedView=False,)
+			   CustomProperties Pin (PinId=D93CF20D46F990734F44869983F64E57,PinName="Output3",PinFriendlyName=" ",Direction="EGPD_Output",PinType.PinCategory="mask",PinType.PinSubCategory="green",PinType.PinSubCategoryObject=None,PinType.PinSubCategoryMemberReference=(),PinType.bIsArray=False,PinType.bIsReference=False,PinType.bIsConst=False,PinType.bIsWeakPointer=False,PersistentGuid=00000000000000000000000000000000,bHidden=False,bNotConnectable=False,bDefaultValueIsReadOnly=False,bDefaultValueIsIgnored=False,bAdvancedView=False,)
+			   CustomProperties Pin (PinId=7930B3334341172972107B8C418D144C,PinName="Output4",PinFriendlyName=" ",Direction="EGPD_Output",PinType.PinCategory="mask",PinType.PinSubCategory="blue",PinType.PinSubCategoryObject=None,PinType.PinSubCategoryMemberReference=(),PinType.bIsArray=False,PinType.bIsReference=False,PinType.bIsConst=False,PinType.bIsWeakPointer=False,PersistentGuid=00000000000000000000000000000000,bHidden=False,bNotConnectable=False,bDefaultValueIsReadOnly=False,bDefaultValueIsIgnored=False,bAdvancedView=False,)
+			   CustomProperties Pin (PinId=76733A1049495B6147E21FB88B151EF3,PinName="Output5",PinFriendlyName=" ",Direction="EGPD_Output",PinType.PinCategory="mask",PinType.PinSubCategory="alpha",PinType.PinSubCategoryObject=None,PinType.PinSubCategoryMemberReference=(),PinType.bIsArray=False,PinType.bIsReference=False,PinType.bIsConst=False,PinType.bIsWeakPointer=False,PersistentGuid=00000000000000000000000000000000,bHidden=False,bNotConnectable=False,bDefaultValueIsReadOnly=False,bDefaultValueIsIgnored=False,bAdvancedView=False,)
+			End Object
+			Begin Object Class=MaterialGraphNode Name="MaterialGraphNode_7"
+			   Begin Object Class=MaterialExpressionClamp Name="MaterialExpressionClamp_2"
+			   End Object
+			   Begin Object Name="MaterialExpressionClamp_2"
+				  Input=(Expression=MaterialExpressionSubtract'MaterialGraphNode_2.MaterialExpressionSubtract_6')
+				  MaterialExpressionEditorX=-368
+				  MaterialExpressionEditorY=48
+				  MaterialExpressionGuid=A23A3F7746E7BD37E9FD7180212F059B
+				  Material=PreviewMaterial'/Engine/Transient.ForceSplat'
+			   End Object
+			   MaterialExpression=MaterialExpressionClamp'MaterialExpressionClamp_2'
+			   NodePosX=-368
+			   NodePosY=48
+			   NodeGuid=A0E540AE43E3B6A3F381579913CF988D
+			   CustomProperties Pin (PinId=42C937124B3EBFB7D0BFD98373FA7A0B,PinName="Input",PinFriendlyName=" ",PinType.PinCategory="required",PinType.PinSubCategoryObject=None,PinType.PinSubCategoryMemberReference=(),PinType.bIsArray=False,PinType.bIsReference=False,PinType.bIsConst=False,PinType.bIsWeakPointer=False,LinkedTo=(MaterialGraphNode_2 25D2D0404B46764F28521E9B20B14B64,),PersistentGuid=00000000000000000000000000000000,bHidden=False,bNotConnectable=False,bDefaultValueIsReadOnly=False,bDefaultValueIsIgnored=False,bAdvancedView=False,)
+			   CustomProperties Pin (PinId=566760C74D635A14E498668CDEBA3396,PinName="Min",PinType.PinCategory="optional",PinType.PinSubCategoryObject=None,PinType.PinSubCategoryMemberReference=(),PinType.bIsArray=False,PinType.bIsReference=False,PinType.bIsConst=False,PinType.bIsWeakPointer=False,PersistentGuid=00000000000000000000000000000000,bHidden=False,bNotConnectable=False,bDefaultValueIsReadOnly=False,bDefaultValueIsIgnored=False,bAdvancedView=False,)
+			   CustomProperties Pin (PinId=5B6B6C3B48C819D0607B028ED1143585,PinName="Max",PinType.PinCategory="optional",PinType.PinSubCategoryObject=None,PinType.PinSubCategoryMemberReference=(),PinType.bIsArray=False,PinType.bIsReference=False,PinType.bIsConst=False,PinType.bIsWeakPointer=False,PersistentGuid=00000000000000000000000000000000,bHidden=False,bNotConnectable=False,bDefaultValueIsReadOnly=False,bDefaultValueIsIgnored=False,bAdvancedView=False,)
+			   CustomProperties Pin (PinId=8D4598BA460B77BFB2B4F0BA844C1932,PinName="Output",PinFriendlyName=" ",Direction="EGPD_Output",PinType.PinSubCategoryObject=None,PinType.PinSubCategoryMemberReference=(),PinType.bIsArray=False,PinType.bIsReference=False,PinType.bIsConst=False,PinType.bIsWeakPointer=False,LinkedTo=(MaterialGraphNode_13 13CBDB574A6785183151A39E8853B26E,),PersistentGuid=00000000000000000000000000000000,bHidden=False,bNotConnectable=False,bDefaultValueIsReadOnly=False,bDefaultValueIsIgnored=False,bAdvancedView=False,)
+			End Object
+			Begin Object Class=MaterialGraphNode Name="MaterialGraphNode_8"
+			   Begin Object Class=MaterialExpressionSubtract Name="MaterialExpressionSubtract_7"
+			   End Object
+			   Begin Object Name="MaterialExpressionSubtract_7"
+				  A=(Expression=MaterialExpressionConstant'MaterialGraphNode_9.MaterialExpressionConstant_2')
+				  B=(Expression=MaterialExpressionMaterialFunctionCall'MaterialGraphNode_12.MaterialExpressionMaterialFunctionCall_2')
+				  MaterialExpressionEditorX=-672
+				  MaterialExpressionEditorY=-32
+				  MaterialExpressionGuid=263301C645665622D4AF209DC4DFC1D4
+				  Material=PreviewMaterial'/Engine/Transient.ForceSplat'
+			   End Object
+			   MaterialExpression=MaterialExpressionSubtract'MaterialExpressionSubtract_7'
+			   NodePosX=-672
+			   NodePosY=-32
+			   NodeGuid=FCC50A2447D0671839BBFDAD0038A57D
+			   CustomProperties Pin (PinId=F06C53144F1B362FBAEC0AB0C16081E1,PinName="A",PinType.PinCategory="optional",PinType.PinSubCategoryObject=None,PinType.PinSubCategoryMemberReference=(),PinType.bIsArray=False,PinType.bIsReference=False,PinType.bIsConst=False,PinType.bIsWeakPointer=False,LinkedTo=(MaterialGraphNode_9 C21918BC40A135B5C229D1845162D312,),PersistentGuid=00000000000000000000000000000000,bHidden=False,bNotConnectable=False,bDefaultValueIsReadOnly=False,bDefaultValueIsIgnored=False,bAdvancedView=False,)
+			   CustomProperties Pin (PinId=CCDF13A940A6E80D85B59299C96B714F,PinName="B",PinType.PinCategory="optional",PinType.PinSubCategoryObject=None,PinType.PinSubCategoryMemberReference=(),PinType.bIsArray=False,PinType.bIsReference=False,PinType.bIsConst=False,PinType.bIsWeakPointer=False,LinkedTo=(MaterialGraphNode_12 DBCC0DEA4DCEE420BA5862AAC582B0C0,),PersistentGuid=00000000000000000000000000000000,bHidden=False,bNotConnectable=False,bDefaultValueIsReadOnly=False,bDefaultValueIsIgnored=False,bAdvancedView=False,)
+			   CustomProperties Pin (PinId=38FFC1874298113EDA7D9DBC71D7C22F,PinName="Output",PinFriendlyName=" ",Direction="EGPD_Output",PinType.PinSubCategoryObject=None,PinType.PinSubCategoryMemberReference=(),PinType.bIsArray=False,PinType.bIsReference=False,PinType.bIsConst=False,PinType.bIsWeakPointer=False,LinkedTo=(MaterialGraphNode_2 C5F733EF466287E9E4EB35ACE6B65E84,),PersistentGuid=00000000000000000000000000000000,bHidden=False,bNotConnectable=False,bDefaultValueIsReadOnly=False,bDefaultValueIsIgnored=False,bAdvancedView=False,)
+			End Object
+			Begin Object Class=MaterialGraphNode Name="MaterialGraphNode_9"
+			   Begin Object Class=MaterialExpressionConstant Name="MaterialExpressionConstant_2"
+			   End Object
+			   Begin Object Name="MaterialExpressionConstant_2"
+				  R=1.000000
+				  MaterialExpressionEditorX=-768
+				  MaterialExpressionEditorY=-64
+				  MaterialExpressionGuid=4D20573E49F05F155BF9099DE5AB86DB
+				  Material=PreviewMaterial'/Engine/Transient.ForceSplat'
+			   End Object
+			   MaterialExpression=MaterialExpressionConstant'MaterialExpressionConstant_2'
+			   NodePosX=-768
+			   NodePosY=-64
+			   NodeGuid=417D7D11440191D4AA9C28B3871ACF2E
+			   CustomProperties Pin (PinId=C21918BC40A135B5C229D1845162D312,PinName="Output",PinFriendlyName=" ",Direction="EGPD_Output",PinType.PinSubCategoryObject=None,PinType.PinSubCategoryMemberReference=(),PinType.bIsArray=False,PinType.bIsReference=False,PinType.bIsConst=False,PinType.bIsWeakPointer=False,LinkedTo=(MaterialGraphNode_8 F06C53144F1B362FBAEC0AB0C16081E1,),PersistentGuid=00000000000000000000000000000000,bHidden=False,bNotConnectable=False,bDefaultValueIsReadOnly=False,bDefaultValueIsIgnored=False,bAdvancedView=False,)
+			End Object
+			Begin Object Class=MaterialGraphNode Name="MaterialGraphNode_10"
+			   Begin Object Class=MaterialExpressionSubtract Name="MaterialExpressionSubtract_8"
+			   End Object
+			   Begin Object Name="MaterialExpressionSubtract_8"
+				  A=(Expression=MaterialExpressionTextureCoordinate'MaterialGraphNode_11.MaterialExpressionTextureCoordinate_2')
+				  B=(Expression=MaterialExpressionComponentMask'MaterialGraphNode_5.MaterialExpressionComponentMask_2')
+				  MaterialExpressionEditorX=-1088
+				  MaterialExpressionGuid=68FF50AA4D025CB5B12184BDE48D55D5
+				  Material=PreviewMaterial'/Engine/Transient.ForceSplat'
+			   End Object
+			   MaterialExpression=MaterialExpressionSubtract'MaterialExpressionSubtract_8'
+			   NodePosX=-1088
+			   NodeGuid=990564B140370AE959137EABE556BEAC
+			   CustomProperties Pin (PinId=1D0B51A44832EA486A38FBABDE408C34,PinName="A",PinType.PinCategory="optional",PinType.PinSubCategoryObject=None,PinType.PinSubCategoryMemberReference=(),PinType.bIsArray=False,PinType.bIsReference=False,PinType.bIsConst=False,PinType.bIsWeakPointer=False,LinkedTo=(MaterialGraphNode_11 8D89F6834B42C8AADDD283818E4EDC94,),PersistentGuid=00000000000000000000000000000000,bHidden=False,bNotConnectable=False,bDefaultValueIsReadOnly=False,bDefaultValueIsIgnored=False,bAdvancedView=False,)
+			   CustomProperties Pin (PinId=26815BDA44BF1141EF0D0084C2E3D993,PinName="B",PinType.PinCategory="optional",PinType.PinSubCategoryObject=None,PinType.PinSubCategoryMemberReference=(),PinType.bIsArray=False,PinType.bIsReference=False,PinType.bIsConst=False,PinType.bIsWeakPointer=False,LinkedTo=(MaterialGraphNode_5 2A7AE85A415E54152002F9BA674B7E41,),PersistentGuid=00000000000000000000000000000000,bHidden=False,bNotConnectable=False,bDefaultValueIsReadOnly=False,bDefaultValueIsIgnored=False,bAdvancedView=False,)
+			   CustomProperties Pin (PinId=CE7F469F4980CA4BA970EDA2C41EAB70,PinName="Output",PinFriendlyName=" ",Direction="EGPD_Output",PinType.PinSubCategoryObject=None,PinType.PinSubCategoryMemberReference=(),PinType.bIsArray=False,PinType.bIsReference=False,PinType.bIsConst=False,PinType.bIsWeakPointer=False,LinkedTo=(MaterialGraphNode_12 4D386898434293CC7337C5B2EADB4BF2,),PersistentGuid=00000000000000000000000000000000,bHidden=False,bNotConnectable=False,bDefaultValueIsReadOnly=False,bDefaultValueIsIgnored=False,bAdvancedView=False,)
+			End Object
+			Begin Object Class=MaterialGraphNode Name="MaterialGraphNode_11"
+			   Begin Object Class=MaterialExpressionTextureCoordinate Name="MaterialExpressionTextureCoordinate_2"
+			   End Object
+			   Begin Object Name="MaterialExpressionTextureCoordinate_2"
+				  MaterialExpressionEditorX=-1280
+				  MaterialExpressionEditorY=-48
+				  MaterialExpressionGuid=28E452A8400D1753908CE0BE4D090D8A
+				  Material=PreviewMaterial'/Engine/Transient.ForceSplat'
+			   End Object
+			   MaterialExpression=MaterialExpressionTextureCoordinate'MaterialExpressionTextureCoordinate_2'
+			   NodePosX=-1280
+			   NodePosY=-48
+			   NodeGuid=0F745D7B4DD6F4646535A8A53DE84B49
+			   CustomProperties Pin (PinId=8D89F6834B42C8AADDD283818E4EDC94,PinName="Output",PinFriendlyName=" ",Direction="EGPD_Output",PinType.PinSubCategoryObject=None,PinType.PinSubCategoryMemberReference=(),PinType.bIsArray=False,PinType.bIsReference=False,PinType.bIsConst=False,PinType.bIsWeakPointer=False,LinkedTo=(MaterialGraphNode_10 1D0B51A44832EA486A38FBABDE408C34,),PersistentGuid=00000000000000000000000000000000,bHidden=False,bNotConnectable=False,bDefaultValueIsReadOnly=False,bDefaultValueIsIgnored=False,bAdvancedView=False,)
+			End Object
+			Begin Object Class=MaterialGraphNode Name="MaterialGraphNode_12"
+			   Begin Object Class=MaterialExpressionMaterialFunctionCall Name="MaterialExpressionMaterialFunctionCall_2"
+			   End Object
+			   Begin Object Name="MaterialExpressionMaterialFunctionCall_2"
+				  MaterialFunction=MaterialFunction'/Engine/Functions/Engine_MaterialFunctions02/Utility/VectorLength.VectorLength'
+				  FunctionInputs(0)=(ExpressionInputId=E06081AD4255F67C6D398889E3CE73BA,Input=(Expression=MaterialExpressionSubtract'MaterialGraphNode_10.MaterialExpressionSubtract_8',InputName="Vector 2"))
+				  FunctionInputs(1)=(ExpressionInputId=1F411C544AE4A9ACC5F5F8A2A099B309,Input=(OutputIndex=-1,InputName="Vector 3"))
+				  FunctionOutputs(0)=(ExpressionOutputId=3C53F2D54FE1029372151E82C85F7755,Output=(OutputName="V2 Length"))
+				  FunctionOutputs(1)=(ExpressionOutputId=F186F43445EF47104CD650A242DEAECE,Output=(OutputName="V3 Length"))
+				  MaterialExpressionEditorX=-912
+				  MaterialExpressionEditorY=16
+				  MaterialExpressionGuid=CA6CCA5B4A71415750F226A88DE0B584
+				  Material=PreviewMaterial'/Engine/Transient.ForceSplat'
+				  Outputs(0)=(OutputName="V2 Length")
+				  Outputs(1)=(OutputName="V3 Length")
+			   End Object
+			   MaterialExpression=MaterialExpressionMaterialFunctionCall'MaterialExpressionMaterialFunctionCall_2'
+			   NodePosX=-912
+			   NodePosY=16
+			   NodeGuid=FFAF3FEE489F8E401BF0D3BA11C26C1D
+			   CustomProperties Pin (PinId=4D386898434293CC7337C5B2EADB4BF2,PinName="Vector 2 (V2)",PinType.PinCategory="optional",PinType.PinSubCategoryObject=None,PinType.PinSubCategoryMemberReference=(),PinType.bIsArray=False,PinType.bIsReference=False,PinType.bIsConst=False,PinType.bIsWeakPointer=False,LinkedTo=(MaterialGraphNode_10 CE7F469F4980CA4BA970EDA2C41EAB70,),PersistentGuid=00000000000000000000000000000000,bHidden=False,bNotConnectable=False,bDefaultValueIsReadOnly=False,bDefaultValueIsIgnored=False,bAdvancedView=False,)
+			   CustomProperties Pin (PinId=51F10BD04E8FF243B7963F95D51E7A0E,PinName="Vector 3 (V3)",PinType.PinCategory="optional",PinType.PinSubCategoryObject=None,PinType.PinSubCategoryMemberReference=(),PinType.bIsArray=False,PinType.bIsReference=False,PinType.bIsConst=False,PinType.bIsWeakPointer=False,PersistentGuid=00000000000000000000000000000000,bHidden=False,bNotConnectable=False,bDefaultValueIsReadOnly=False,bDefaultValueIsIgnored=False,bAdvancedView=False,)
+			   CustomProperties Pin (PinId=DBCC0DEA4DCEE420BA5862AAC582B0C0,PinName="V2 Length",Direction="EGPD_Output",PinType.PinSubCategoryObject=None,PinType.PinSubCategoryMemberReference=(),PinType.bIsArray=False,PinType.bIsReference=False,PinType.bIsConst=False,PinType.bIsWeakPointer=False,LinkedTo=(MaterialGraphNode_8 CCDF13A940A6E80D85B59299C96B714F,),PersistentGuid=00000000000000000000000000000000,bHidden=False,bNotConnectable=False,bDefaultValueIsReadOnly=False,bDefaultValueIsIgnored=False,bAdvancedView=False,)
+			   CustomProperties Pin (PinId=1EDAC5CB45C35155B74AC9AC7C968F52,PinName="V3 Length",Direction="EGPD_Output",PinType.PinSubCategoryObject=None,PinType.PinSubCategoryMemberReference=(),PinType.bIsArray=False,PinType.bIsReference=False,PinType.bIsConst=False,PinType.bIsWeakPointer=False,PersistentGuid=00000000000000000000000000000000,bHidden=False,bNotConnectable=False,bDefaultValueIsReadOnly=False,bDefaultValueIsIgnored=False,bAdvancedView=False,)
+			End Object
+			Begin Object Class=MaterialGraphNode Name="MaterialGraphNode_13"
+			   Begin Object Class=MaterialExpressionDivide Name="MaterialExpressionDivide_2"
+			   End Object
+			   Begin Object Name="MaterialExpressionDivide_2"
+				  A=(Expression=MaterialExpressionClamp'MaterialGraphNode_7.MaterialExpressionClamp_2')
+				  B=(Expression=MaterialExpressionScalarParameter'MaterialGraphNode_4.MaterialExpressionScalarParameter_5')
+				  MaterialExpressionEditorX=-80
+				  MaterialExpressionEditorY=80
+				  MaterialExpressionGuid=8730CE28486A7A975F06AAAF76360BDE
+				  Material=PreviewMaterial'/Engine/Transient.ForceSplat'
+			   End Object
+			   MaterialExpression=MaterialExpressionDivide'MaterialExpressionDivide_2'
+			   NodePosX=-80
+			   NodePosY=80
+			   NodeGuid=77DB1B5D4C0BCEC4D991518C010F05E1
+			   CustomProperties Pin (PinId=13CBDB574A6785183151A39E8853B26E,PinName="A",PinType.PinCategory="optional",PinType.PinSubCategoryObject=None,PinType.PinSubCategoryMemberReference=(),PinType.bIsArray=False,PinType.bIsReference=False,PinType.bIsConst=False,PinType.bIsWeakPointer=False,LinkedTo=(MaterialGraphNode_7 8D4598BA460B77BFB2B4F0BA844C1932,),PersistentGuid=00000000000000000000000000000000,bHidden=False,bNotConnectable=False,bDefaultValueIsReadOnly=False,bDefaultValueIsIgnored=False,bAdvancedView=False,)
+			   CustomProperties Pin (PinId=3501ECFD4967CB5C6FA143BB453E5636,PinName="B",PinType.PinCategory="optional",PinType.PinSubCategoryObject=None,PinType.PinSubCategoryMemberReference=(),PinType.bIsArray=False,PinType.bIsReference=False,PinType.bIsConst=False,PinType.bIsWeakPointer=False,LinkedTo=(MaterialGraphNode_4 243D5A9F4A177E21AD5B0C9A9EE13B6C,),PersistentGuid=00000000000000000000000000000000,bHidden=False,bNotConnectable=False,bDefaultValueIsReadOnly=False,bDefaultValueIsIgnored=False,bAdvancedView=False,)
+			   CustomProperties Pin (PinId=D0CE6B2B4A6886F59532E4B83367D7BD,PinName="Output",PinFriendlyName=" ",Direction="EGPD_Output",PinType.PinSubCategoryObject=None,PinType.PinSubCategoryMemberReference=(),PinType.bIsArray=False,PinType.bIsReference=False,PinType.bIsConst=False,PinType.bIsWeakPointer=False,LinkedTo=(MaterialGraphNode_0 B7D0CF28498836478F9C8987312501BF,),PersistentGuid=00000000000000000000000000000000,bHidden=False,bNotConnectable=False,bDefaultValueIsReadOnly=False,bDefaultValueIsIgnored=False,bAdvancedView=False,)
+			End Object
+
+	[/PARAM]
+[/OBJECT]
+[/EXCERPT:ForceSplat]
+
+-->
+
