Copying //UE4/Dev-Documentation to Samples-Main (//UE4/Samples-Main)

diff --git a/Engine/Documentation/Source/Engine/Rendering/RenderTargets/BlueprintRenderTargets/HowTo/FluidSurface/3/FSBP3.INT.udn b/Engine/Documentation/Source/Engine/Rendering/RenderTargets/BlueprintRenderTargets/HowTo/FluidSurface/3/FSBP3.INT.udn
new file mode 100644
--- /dev/null
+++ b/Engine/Documentation/Source/Engine/Rendering/RenderTargets/BlueprintRenderTargets/HowTo/FluidSurface/3/FSBP3.INT.udn
@@ -0,0 +1,2005 @@
+availability:docs
+title:3. Blueprint Setup
+description:Going over setting up the Fluid Surface Blueprint.
+crumbs:%ROOT%, Engine, Engine\Rendering\RenderTargets
+version:4.13
+type:multi-step
+SkillLevel:Advanced
+parent:Engine/Rendering/RenderTargets/BlueprintRenderTargets/HowTo/FluidSurface
+Navigation:multi-step
+tags: Render Targets
+order:3
+checkpoint:FSB
+
+[NAV]
+
+
+Now that the required Materials and Render Targets have been created and setup, it is now time to create a new Blueprint and configure it to work as the fluid surface Actor. 
+
+## Blueprint Creation & Variable Setup 
+1. Right-click in the Content Browser and from the menu in the **Create Basic Asset** section, click on the **Blueprint Class** option.
+
+	![](Engine\Rendering\RenderTargets\BlueprintRenderTargets\HowTo\HeightFieldPainter\HFP_BP_Setup_00.png)
+
+1. From the **Pick Parent Class** window, select **Actor**, naming it `WaterSurface_BP`.
+
+	![](Engine\Rendering\RenderTargets\BlueprintRenderTargets\HowTo\FluidSurface\FS_Water_BP_Creation.png)
+
+1. Double-click on the WaterSurface_BP to open it up and then add the following variable types.
+
+	![](Engine\Rendering\RenderTargets\BlueprintRenderTargets\HowTo\FluidSurface\FS_Var_Creation.png)
+
+1. Now, set the variables to the following default values:
+
+	|Variable Name| Variable Type|Default Value|
+	|--------|------------|---|
+	|WaterMeshAsset|Static Mesh|SM_Plane1000_512 (Included in the Content Example Maps)|
+	|WaterMaterial|Material Interface|WaterMaterial|
+	|HeightState|Integer|0|
+	|TimeAccumulator|Float|0.0|
+	|WaterMID|MaterialInstanceDynamic|N/A|
+	|UpdateRate|Float|60.0|
+	|TouchingCharacter|Actor|N/A|
+	|LastTouchingActorPosition|Vector|0, 0, 0|
+	|InteractionDistance|Float|1.0|
+	|MouseDown|Bool|Un-Checked|
+
+## Blueprint Function Creation & Setup
+Now that the Blueprint has been created with its required variables, it is now time to create some functions. In the following section, we will go over what the three functions are and how you go about setting them up.
+
+[region:note]
+Do not forget to Compile and Save your Blueprint so that it can be used.
+[/region]
+
+* The **GetHeightRT** function will get and store the height of the render target so that it can be used later. Before you start to create this function, make sure you setup the function's Inputs and Outputs with the following information.
+	* **Inputs**:
+		* Name: Index 
+		* Type:Integer
+	* **Outputs**:
+		* Name: Heigh RT 
+		* Type:Texture Render Target 2D
+
+		[REGION:fullwidth]
+		[INCLUDE:Engine\Rendering\RenderTargets\BlueprintRenderTargets\HowTo\FluidSurface\3\#HeightFunction]
+		[/REGION]
+
+		[region:note]
+		Do not forget to add each of the three render targets to the **Height RT** inputs on the Return Node.
+		[/region]
+
+* The **GetLastHeight** function will store the last known height for the fluid surface for later use. Before you start to create this function, make sure you setup the function's Inputs and Outputs with the following information.
+	* **Inputs**:
+		* Name: CurrentHeightIndex 
+		* Type:Integer
+		* Name: NumFramesOld
+		* Type:Integer
+	* **Outputs**:
+		* Name: Heigh RT 
+		* Type:Texture Render Target 2D
+
+		[REGION:fullwidth]
+		[INCLUDE:Engine\Rendering\RenderTargets\BlueprintRenderTargets\HowTo\FluidSurface\3\#GetLastHeight]
+		[/REGION]
+
+* **Trace From Camera** is the same trace that was used in the Height Field Painter Blueprint. For this function, there are no input or output variables that need to be setup.
+
+	[REGION:fullwidth]
+	[INCLUDE:Engine\Rendering\RenderTargets\BlueprintRenderTargets\HowTo\FluidSurface\3\#CameraTrace]
+	[/REGION]
+
+## Blueprint Construction Graph Setup
+The Construction Graph is used to create and set both of the Static Meshes and Materials that fluid simulation needs to function correctly. In the following section, we will go over what the three functions are, and how you go about setting them up.
+
+[region:note]
+Do not forget to Compile and Save your Blueprint so that it can be used.
+[/region]
+
+* In the Construction Graph, we are going to create a Static Mesh that is needed for the water surface and then create and apply a Dynamic Material Instance to that Static Mesh, which will simulate the water's surface.
+	
+	[REGION:fullwidth]
+	[INCLUDE:Engine\Rendering\RenderTargets\BlueprintRenderTargets\HowTo\FluidSurface\3\#ConScript]
+	[/REGION]
+
+## Blueprint Event Graph Setup
+The Event graph is where everything that is needed for this effect to work correctly is added and connected. In the following section, we will go over what each section of the Event Graph does and how it affects the final Blueprint that is placed in the world.
+
+[region:note]
+Do not forget to Compile and Save your Blueprint so that it can be used.
+[/region]
+
+* **Clear all render targets of data** is the part of the Blueprint that clears all of the render targets of any previous data that they might contain in order to ensure that no artifacts appear the first time the Blueprint is run.  
+
+	[REGION:fullwidth]
+	[INCLUDE:Engine\Rendering\RenderTargets\BlueprintRenderTargets\HowTo\FluidSurface\3\#ClearRT]
+	[/REGION]
+
+	[region:note]
+	Do not forget to add each of the three render targets to the *Height RT** inputs on the Return Node.
+	[/region]
+
+* **Check to see if Left Mouse Button is held down** is the part of the Blueprint that checks to see if the user has the Left Mouse Button held down or not. If the user has the Left Mouse button held down, the **MouseDown** is set to true; otherwise, it is set to false.
+
+	[REGION:fullwidth]
+	[INCLUDE:Engine\Rendering\RenderTargets\BlueprintRenderTargets\HowTo\FluidSurface\3\#MouseDown]
+	[/REGION]
+
+* The **Take Damage Force** section finds the location on the Static Mesh that was hit and applies the **ForceSplat** Material to that section so that it can offset the surface that is simulating an impact event.
+
+	[REGION:fullwidth]
+	[INCLUDE:Engine\Rendering\RenderTargets\BlueprintRenderTargets\HowTo\FluidSurface\3\#TakeDamage]
+	[/REGION]
+
+* The **Begin & End Overlap** section checks to see if the player is currently touching the fluid surface Actor that has been placed in the world or not. 
+	
+	[REGION:fullwidth]
+	[INCLUDE:Engine\Rendering\RenderTargets\BlueprintRenderTargets\HowTo\FluidSurface\3\#Overlap]
+	[/REGION]
+
+* **Shoot when mouse is down** is the part of the Blueprint that checks each Tick to see if the Left Mouse button is held down. When the Left Mouse button is held down, the **Trace From Camera** function is called, which traces from the center of the camera into the world.  
+
+	[REGION:fullwidth]
+	[INCLUDE:Engine\Rendering\RenderTargets\BlueprintRenderTargets\HowTo\FluidSurface\3\#ShootWhenMouseDown]
+	[/REGION]
+
+* **Touching Character Force** is the part of the Blueprint that is responsible for applying and updating the force that is applied to the fluid surface Actor when a player is touching it.  
+
+	[REGION:fullwidth]
+	[INCLUDE:Engine\Rendering\RenderTargets\BlueprintRenderTargets\HowTo\FluidSurface\3\#TouchingCarForce]
+	[/REGION]
+
+* **Fixed timestep, so fluid speed is framerate independent** is the part of the Blueprint that makes sure that the propagation of the waves through the fluid surface has a frame rate that is not tied to the frame rate of the project. Essentially, the waves in the fluid surface will not slow down or speed up based on what is currently happening in the level.
+
+	[REGION:fullwidth]
+	[INCLUDE:Engine\Rendering\RenderTargets\BlueprintRenderTargets\HowTo\FluidSurface\3\#FixedFrameRate]
+	[/REGION]
+
+* **Cycle Active heightfield** is the part of the Blueprint that cycles through the various render targets for the fluid surface so that the fluid surface has a more dynamic look when shot or touched.
+
+	[REGION:fullwidth]
+	[INCLUDE:Engine\Rendering\RenderTargets\BlueprintRenderTargets\HowTo\FluidSurface\3\#CycleActiveHeightfield]
+	[/REGION]
+
+* **Apply fluid surface simulation kernel, which propagates waves** is the part of the Blueprint that handles the propagation of the waves through the fluid surface Actor making it look like the waves die out over time.
+
+	[REGION:fullwidth]
+	[INCLUDE:Engine\Rendering\RenderTargets\BlueprintRenderTargets\HowTo\FluidSurface\3\#ApplySim]
+	[/REGION]
+
+* **Compute Surface Normal** is a section of the Blueprint that applies a normal to the surface of the fluid which helps make the ripples more pronounced and easier to see. 
+
+	[REGION:fullwidth]
+	[INCLUDE:Engine\Rendering\RenderTargets\BlueprintRenderTargets\HowTo\FluidSurface\3\#CompSurfNorm]
+	[/REGION]
+
+Now that the Blueprints have been setup in the next section we will take a look at how to set all of this up for use in UE4.
+[NAV]
+
+<!--
+[EXCERPT:CompSurfNorm]	
+[OBJECT:Snippet]
+	[PARAMLITERAL:id]
+		CompSurfNorm
+	[/PARAMLITERAL]
+	[PARAMLITERAL:type]
+		blueprint
+	[/PARAMLITERAL]
+	[PARAM:image]
+		![](Engine\Rendering\RenderTargets\BlueprintRenderTargets\HowTo\FluidSurface\FS_Compute_Surface_Normal.png)
+	[/PARAM]
+	[PARAM:snippet]
+			BEGIN OBJECT
+							Begin Object Class=K2Node_VariableGet Name="K2Node_VariableGet_193"
+			   VariableReference=(MemberName="WaterMID",MemberGuid=900EBFD242C0177F9BFCF7A541E03E45,bSelfContext=True)
+			   NodePosX=2176
+			   NodePosY=992
+			   NodeGuid=C3E0710D4C9160E3BF2C9EAF3E268CD7
+			   CustomProperties Pin (PinId=E8C406E34369A7F135491B8C94106CB0,PinName="WaterMID",Direction="EGPD_Output",PinType.PinCategory="object",PinType.PinSubCategoryObject=Class'/Script/Engine.MaterialInstanceDynamic',PinType.PinSubCategoryMemberReference=(),PinType.bIsArray=False,PinType.bIsReference=False,PinType.bIsConst=False,PinType.bIsWeakPointer=False,LinkedTo=(K2Node_CallFunction_636 5DD3E3F246B683E0C500FE881FDE590E,),PersistentGuid=00000000000000000000000000000000,bHidden=False,bNotConnectable=False,bDefaultValueIsReadOnly=False,bDefaultValueIsIgnored=False,bAdvancedView=False,)
+			   CustomProperties Pin (PinId=DB1C8246407F68DF33F9BAAC1FF607EB,PinName="self",PinFriendlyName=NSLOCTEXT("K2Node", "Target", "Target"),PinType.PinCategory="object",PinType.PinSubCategoryObject=BlueprintGeneratedClass'/Game/ExampleContent/BlueprintFluidSurface/WaterSurface_BP.WaterSurface_BP_C',PinType.PinSubCategoryMemberReference=(),PinType.bIsArray=False,PinType.bIsReference=False,PinType.bIsConst=False,PinType.bIsWeakPointer=False,PersistentGuid=00000000000000000000000000000000,bHidden=True,bNotConnectable=False,bDefaultValueIsReadOnly=False,bDefaultValueIsIgnored=False,bAdvancedView=False,)
+			End Object
+			Begin Object Class=K2Node_CallFunction Name="K2Node_CallFunction_636"
+			   FunctionReference=(MemberParent=Class'/Script/Engine.MaterialInstanceDynamic',MemberName="SetTextureParameterValue")
+			   NodePosX=2384
+			   NodePosY=736
+			   NodeGuid=16E4EFC043E75DD6033A16830004D89E
+			   CustomProperties Pin (PinId=6464966E4B5F49CCDCC4C0A0073C2352,PinName="execute",PinToolTip="\nExec",PinType.PinCategory="exec",PinType.PinSubCategoryObject=None,PinType.PinSubCategoryMemberReference=(),PinType.bIsArray=False,PinType.bIsReference=False,PinType.bIsConst=False,PinType.bIsWeakPointer=False,LinkedTo=(K2Node_CallFunction_150 3518A0B14FC034E0ABEF289465A629BF,),PersistentGuid=00000000000000000000000000000000,bHidden=False,bNotConnectable=False,bDefaultValueIsReadOnly=False,bDefaultValueIsIgnored=False,bAdvancedView=False,)
+			   CustomProperties Pin (PinId=A82196CF4996E41CA17E4786E8F8FCCA,PinName="then",PinToolTip="\nExec",Direction="EGPD_Output",PinType.PinCategory="exec",PinType.PinSubCategoryObject=None,PinType.PinSubCategoryMemberReference=(),PinType.bIsArray=False,PinType.bIsReference=False,PinType.bIsConst=False,PinType.bIsWeakPointer=False,PersistentGuid=00000000000000000000000000000000,bHidden=False,bNotConnectable=False,bDefaultValueIsReadOnly=False,bDefaultValueIsIgnored=False,bAdvancedView=False,)
+			   CustomProperties Pin (PinId=5DD3E3F246B683E0C500FE881FDE590E,PinName="self",PinFriendlyName=NSLOCTEXT("K2Node", "Target", "Target"),PinToolTip="Target\nMaterial Instance Dynamic Reference",PinType.PinCategory="object",PinType.PinSubCategoryObject=Class'/Script/Engine.MaterialInstanceDynamic',PinType.PinSubCategoryMemberReference=(),PinType.bIsArray=False,PinType.bIsReference=False,PinType.bIsConst=False,PinType.bIsWeakPointer=False,LinkedTo=(K2Node_VariableGet_193 E8C406E34369A7F135491B8C94106CB0,),PersistentGuid=00000000000000000000000000000000,bHidden=False,bNotConnectable=False,bDefaultValueIsReadOnly=False,bDefaultValueIsIgnored=False,bAdvancedView=False,)
+			   CustomProperties Pin (PinId=C38C3A93459AC64E47B5C8A4E1BCA25F,PinName="ParameterName",PinToolTip="Parameter Name\nName",PinType.PinCategory="name",PinType.PinSubCategoryObject=None,PinType.PinSubCategoryMemberReference=(),PinType.bIsArray=False,PinType.bIsReference=False,PinType.bIsConst=False,PinType.bIsWeakPointer=False,DefaultValue="Heightfield",AutogeneratedDefaultValue="None",PersistentGuid=00000000000000000000000000000000,bHidden=False,bNotConnectable=False,bDefaultValueIsReadOnly=False,bDefaultValueIsIgnored=False,bAdvancedView=False,)
+			   CustomProperties Pin (PinId=F159709749EC1DFC6F2A7CAD2034EDE0,PinName="Value",PinToolTip="Value\nTexture Reference",PinType.PinCategory="object",PinType.PinSubCategoryObject=Class'/Script/Engine.Texture',PinType.PinSubCategoryMemberReference=(),PinType.bIsArray=False,PinType.bIsReference=False,PinType.bIsConst=False,PinType.bIsWeakPointer=False,LinkedTo=(K2Node_CallFunction_175 17ED0A19416B8925E2303CB3AE6D3CA9,),PersistentGuid=00000000000000000000000000000000,bHidden=False,bNotConnectable=False,bDefaultValueIsReadOnly=False,bDefaultValueIsIgnored=False,bAdvancedView=False,)
+			End Object
+			Begin Object Class=K2Node_CallFunction Name="K2Node_CallFunction_175"
+			   FunctionReference=(MemberName="GetHeightRT",MemberGuid=2EEA338F4CA878E7228A10BB2DD63E00,bSelfContext=True)
+			   NodePosX=1568
+			   NodePosY=736
+			   NodeGuid=8B5B0A394E725836DB140A8AFEF6ABFF
+			   CustomProperties Pin (PinId=6953BBD0466BEF70ADA7EB944311C773,PinName="execute",PinToolTip="\nExec",PinType.PinCategory="exec",PinType.PinSubCategoryObject=None,PinType.PinSubCategoryMemberReference=(),PinType.bIsArray=False,PinType.bIsReference=False,PinType.bIsConst=False,PinType.bIsWeakPointer=False,LinkedTo=(K2Node_CallMaterialParameterCollectionFunction_112 D1DC644F437101011A73719D87003CF3,),PersistentGuid=00000000000000000000000000000000,bHidden=False,bNotConnectable=False,bDefaultValueIsReadOnly=False,bDefaultValueIsIgnored=False,bAdvancedView=False,)
+			   CustomProperties Pin (PinId=D03AF0CE4F40BC2B72AB9093E454AF7C,PinName="then",PinToolTip="\nExec",Direction="EGPD_Output",PinType.PinCategory="exec",PinType.PinSubCategoryObject=None,PinType.PinSubCategoryMemberReference=(),PinType.bIsArray=False,PinType.bIsReference=False,PinType.bIsConst=False,PinType.bIsWeakPointer=False,LinkedTo=(K2Node_CallFunction_2070 B4EF9FDF44A64EAFA4BC079CAE9DD7EA,),PersistentGuid=00000000000000000000000000000000,bHidden=False,bNotConnectable=False,bDefaultValueIsReadOnly=False,bDefaultValueIsIgnored=False,bAdvancedView=False,)
+			   CustomProperties Pin (PinId=7DF1FA7F425A5456A71932B56B94CA0F,PinName="self",PinFriendlyName=NSLOCTEXT("K2Node", "Target", "Target"),PinToolTip="Target\nSelf Reference",PinType.PinCategory="object",PinType.PinSubCategory="self",PinType.PinSubCategoryObject=None,PinType.PinSubCategoryMemberReference=(),PinType.bIsArray=False,PinType.bIsReference=False,PinType.bIsConst=False,PinType.bIsWeakPointer=False,PersistentGuid=00000000000000000000000000000000,bHidden=False,bNotConnectable=False,bDefaultValueIsReadOnly=False,bDefaultValueIsIgnored=False,bAdvancedView=False,)
+			   CustomProperties Pin (PinId=5F1BB9ED4029783BBDF5F791A0608883,PinName="Index",PinToolTip="Index\nInteger",PinType.PinCategory="int",PinType.PinSubCategoryObject=None,PinType.PinSubCategoryMemberReference=(),PinType.bIsArray=False,PinType.bIsReference=False,PinType.bIsConst=False,PinType.bIsWeakPointer=False,DefaultValue="0",AutogeneratedDefaultValue="0",LinkedTo=(K2Node_VariableGet_130 5C098F58454D29E4C04F05B60CCF8D98,),PersistentGuid=00000000000000000000000000000000,bHidden=False,bNotConnectable=False,bDefaultValueIsReadOnly=False,bDefaultValueIsIgnored=False,bAdvancedView=False,)
+			   CustomProperties Pin (PinId=17ED0A19416B8925E2303CB3AE6D3CA9,PinName="HeightRT",PinToolTip="Height RT\nTexture Render Target 2D Reference",Direction="EGPD_Output",PinType.PinCategory="object",PinType.PinSubCategoryObject=Class'/Script/Engine.TextureRenderTarget2D',PinType.PinSubCategoryMemberReference=(),PinType.bIsArray=False,PinType.bIsReference=False,PinType.bIsConst=False,PinType.bIsWeakPointer=False,LinkedTo=(K2Node_CallFunction_636 F159709749EC1DFC6F2A7CAD2034EDE0,K2Node_CallFunction_2070 A9EF7DBB418EB0C7514E928F8881B5BA,),PersistentGuid=00000000000000000000000000000000,bHidden=False,bNotConnectable=False,bDefaultValueIsReadOnly=False,bDefaultValueIsIgnored=False,bAdvancedView=False,)
+			End Object
+			Begin Object Class=K2Node_VariableGet Name="K2Node_VariableGet_130"
+			   VariableReference=(MemberName="HeightState",MemberGuid=8E10E24F4154A224C2065A89B367BD62,bSelfContext=True)
+			   NodePosX=1408
+			   NodePosY=976
+			   NodeGuid=DB3D6C794610439084D4FF8C12011B40
+			   CustomProperties Pin (PinId=5C098F58454D29E4C04F05B60CCF8D98,PinName="HeightState",Direction="EGPD_Output",PinType.PinCategory="int",PinType.PinSubCategoryObject=None,PinType.PinSubCategoryMemberReference=(),PinType.bIsArray=False,PinType.bIsReference=False,PinType.bIsConst=False,PinType.bIsWeakPointer=False,AutogeneratedDefaultValue="0",LinkedTo=(K2Node_CallFunction_175 5F1BB9ED4029783BBDF5F791A0608883,),PersistentGuid=00000000000000000000000000000000,bHidden=False,bNotConnectable=False,bDefaultValueIsReadOnly=False,bDefaultValueIsIgnored=False,bAdvancedView=False,)
+			   CustomProperties Pin (PinId=E8B6761B4C7B415363F85B9485BDA925,PinName="self",PinFriendlyName=NSLOCTEXT("K2Node", "Target", "Target"),PinType.PinCategory="object",PinType.PinSubCategoryObject=BlueprintGeneratedClass'/Game/ExampleContent/BlueprintFluidSurface/WaterSurface_BP.WaterSurface_BP_C',PinType.PinSubCategoryMemberReference=(),PinType.bIsArray=False,PinType.bIsReference=False,PinType.bIsConst=False,PinType.bIsWeakPointer=False,PersistentGuid=00000000000000000000000000000000,bHidden=True,bNotConnectable=False,bDefaultValueIsReadOnly=False,bDefaultValueIsIgnored=False,bAdvancedView=False,)
+			End Object
+			Begin Object Class=K2Node_CallFunction Name="K2Node_CallFunction_150"
+			   FunctionReference=(MemberParent=Class'/Script/Engine.KismetRenderingLibrary',MemberName="DrawMaterialToRenderTarget")
+			   NodePosX=2096
+			   NodePosY=752
+			   NodeGuid=FA4BD327416F773E7ED2338E8AE6E30B
+			   CustomProperties Pin (PinId=A7CDD48D45940F5F062D73A8005F2D43,PinName="execute",PinToolTip="\nExec",PinType.PinCategory="exec",PinType.PinSubCategoryObject=None,PinType.PinSubCategoryMemberReference=(),PinType.bIsArray=False,PinType.bIsReference=False,PinType.bIsConst=False,PinType.bIsWeakPointer=False,LinkedTo=(K2Node_CallFunction_2070 2F6CF49D48269B0073CB76ADDE38337B,),PersistentGuid=00000000000000000000000000000000,bHidden=False,bNotConnectable=False,bDefaultValueIsReadOnly=False,bDefaultValueIsIgnored=False,bAdvancedView=False,)
+			   CustomProperties Pin (PinId=3518A0B14FC034E0ABEF289465A629BF,PinName="then",PinToolTip="\nExec",Direction="EGPD_Output",PinType.PinCategory="exec",PinType.PinSubCategoryObject=None,PinType.PinSubCategoryMemberReference=(),PinType.bIsArray=False,PinType.bIsReference=False,PinType.bIsConst=False,PinType.bIsWeakPointer=False,LinkedTo=(K2Node_CallFunction_636 6464966E4B5F49CCDCC4C0A0073C2352,),PersistentGuid=00000000000000000000000000000000,bHidden=False,bNotConnectable=False,bDefaultValueIsReadOnly=False,bDefaultValueIsIgnored=False,bAdvancedView=False,)
+			   CustomProperties Pin (PinId=A5DEAA20404B5A924CA1F1A5E9020524,PinName="self",PinFriendlyName=NSLOCTEXT("K2Node", "Target", "Target"),PinToolTip="Target\nKismet Rendering Library Reference",PinType.PinCategory="object",PinType.PinSubCategoryObject=Class'/Script/Engine.KismetRenderingLibrary',PinType.PinSubCategoryMemberReference=(),PinType.bIsArray=False,PinType.bIsReference=False,PinType.bIsConst=False,PinType.bIsWeakPointer=False,DefaultObject="/Script/Engine.Default__KismetRenderingLibrary",PersistentGuid=00000000000000000000000000000000,bHidden=True,bNotConnectable=False,bDefaultValueIsReadOnly=False,bDefaultValueIsIgnored=False,bAdvancedView=False,)
+			   CustomProperties Pin (PinId=DF8D40B24A1F6CBB6E2FF4A731F50DE2,PinName="WorldContextObject",PinToolTip="World Context Object\nObject Reference",PinType.PinCategory="object",PinType.PinSubCategoryObject=Class'/Script/CoreUObject.Object',PinType.PinSubCategoryMemberReference=(),PinType.bIsArray=False,PinType.bIsReference=False,PinType.bIsConst=False,PinType.bIsWeakPointer=False,PersistentGuid=00000000000000000000000000000000,bHidden=True,bNotConnectable=False,bDefaultValueIsReadOnly=False,bDefaultValueIsIgnored=False,bAdvancedView=False,)
+			   CustomProperties Pin (PinId=94D570EB4003AA7318611780BFD94FE8,PinName="TextureRenderTarget",PinToolTip="Texture Render Target\nTexture Render Target 2D Reference",PinType.PinCategory="object",PinType.PinSubCategoryObject=Class'/Script/Engine.TextureRenderTarget2D',PinType.PinSubCategoryMemberReference=(),PinType.bIsArray=False,PinType.bIsReference=False,PinType.bIsConst=False,PinType.bIsWeakPointer=False,DefaultObject="/Game/ExampleContent/BlueprintFluidSurface/HeightfieldNormal.HeightfieldNormal",PersistentGuid=00000000000000000000000000000000,bHidden=False,bNotConnectable=False,bDefaultValueIsReadOnly=False,bDefaultValueIsIgnored=False,bAdvancedView=False,)
+			   CustomProperties Pin (PinId=72BFBBD84DEDB3B6EC38639048AB6A45,PinName="Material",PinToolTip="Material\nMaterial Interface Reference",PinType.PinCategory="object",PinType.PinSubCategoryObject=Class'/Script/Engine.MaterialInterface',PinType.PinSubCategoryMemberReference=(),PinType.bIsArray=False,PinType.bIsReference=False,PinType.bIsConst=False,PinType.bIsWeakPointer=False,LinkedTo=(K2Node_CallMaterialParameterCollectionFunction_112 EAC78C6D415430B012961DB7D2A9E3FE,),PersistentGuid=00000000000000000000000000000000,bHidden=False,bNotConnectable=False,bDefaultValueIsReadOnly=False,bDefaultValueIsIgnored=False,bAdvancedView=False,)
+			End Object
+			Begin Object Class=K2Node_CallMaterialParameterCollectionFunction Name="K2Node_CallMaterialParameterCollectionFunction_112"
+			   FunctionReference=(MemberParent=Class'/Script/Engine.KismetMaterialLibrary',MemberName="CreateDynamicMaterialInstance")
+			   NodePosX=1232
+			   NodePosY=752
+			   NodeGuid=7BA9E4734F0F9F8BDEA82A99E3543AF6
+			   CustomProperties Pin (PinId=7481116A43E6DECBD5AE68A3BAA018B4,PinName="execute",PinToolTip="\nExec",PinType.PinCategory="exec",PinType.PinSubCategoryObject=None,PinType.PinSubCategoryMemberReference=(),PinType.bIsArray=False,PinType.bIsReference=False,PinType.bIsConst=False,PinType.bIsWeakPointer=False,LinkedTo=(K2Node_MacroInstance_0 3F8D0C52472E83BA32ACAABBA33F8AD7,),PersistentGuid=00000000000000000000000000000000,bHidden=False,bNotConnectable=False,bDefaultValueIsReadOnly=False,bDefaultValueIsIgnored=False,bAdvancedView=False,)
+			   CustomProperties Pin (PinId=D1DC644F437101011A73719D87003CF3,PinName="then",PinToolTip="\nExec",Direction="EGPD_Output",PinType.PinCategory="exec",PinType.PinSubCategoryObject=None,PinType.PinSubCategoryMemberReference=(),PinType.bIsArray=False,PinType.bIsReference=False,PinType.bIsConst=False,PinType.bIsWeakPointer=False,LinkedTo=(K2Node_CallFunction_175 6953BBD0466BEF70ADA7EB944311C773,),PersistentGuid=00000000000000000000000000000000,bHidden=False,bNotConnectable=False,bDefaultValueIsReadOnly=False,bDefaultValueIsIgnored=False,bAdvancedView=False,)
+			   CustomProperties Pin (PinId=719668CD4164D7091AA630883C52340F,PinName="self",PinFriendlyName=NSLOCTEXT("K2Node", "Target", "Target"),PinToolTip="Target\nKismet Material Library Reference",PinType.PinCategory="object",PinType.PinSubCategoryObject=Class'/Script/Engine.KismetMaterialLibrary',PinType.PinSubCategoryMemberReference=(),PinType.bIsArray=False,PinType.bIsReference=False,PinType.bIsConst=False,PinType.bIsWeakPointer=False,DefaultObject="/Script/Engine.Default__KismetMaterialLibrary",PersistentGuid=00000000000000000000000000000000,bHidden=True,bNotConnectable=False,bDefaultValueIsReadOnly=False,bDefaultValueIsIgnored=False,bAdvancedView=False,)
+			   CustomProperties Pin (PinId=AD1AEA064BBC69AA124829A70FDB037A,PinName="WorldContextObject",PinToolTip="World Context Object\nObject Reference",PinType.PinCategory="object",PinType.PinSubCategoryObject=Class'/Script/CoreUObject.Object',PinType.PinSubCategoryMemberReference=(),PinType.bIsArray=False,PinType.bIsReference=False,PinType.bIsConst=False,PinType.bIsWeakPointer=False,PersistentGuid=00000000000000000000000000000000,bHidden=True,bNotConnectable=False,bDefaultValueIsReadOnly=False,bDefaultValueIsIgnored=False,bAdvancedView=False,)
+			   CustomProperties Pin (PinId=73B85BC745FE8F707AEB438631585139,PinName="Parent",PinToolTip="Parent\nMaterial Interface Reference",PinType.PinCategory="object",PinType.PinSubCategoryObject=Class'/Script/Engine.MaterialInterface',PinType.PinSubCategoryMemberReference=(),PinType.bIsArray=False,PinType.bIsReference=False,PinType.bIsConst=False,PinType.bIsWeakPointer=False,DefaultObject="/Game/ExampleContent/BlueprintFluidSurface/ComputeNormal.ComputeNormal",PersistentGuid=00000000000000000000000000000000,bHidden=False,bNotConnectable=False,bDefaultValueIsReadOnly=False,bDefaultValueIsIgnored=False,bAdvancedView=False,)
+			   CustomProperties Pin (PinId=EAC78C6D415430B012961DB7D2A9E3FE,PinName="ReturnValue",PinToolTip="Return Value\nMaterial Instance Dynamic Reference",Direction="EGPD_Output",PinType.PinCategory="object",PinType.PinSubCategoryObject=Class'/Script/Engine.MaterialInstanceDynamic',PinType.PinSubCategoryMemberReference=(),PinType.bIsArray=False,PinType.bIsReference=False,PinType.bIsConst=False,PinType.bIsWeakPointer=False,LinkedTo=(K2Node_CallFunction_2070 058B7D14466F1D7BAD9496ACFD44563E,K2Node_CallFunction_150 72BFBBD84DEDB3B6EC38639048AB6A45,),PersistentGuid=00000000000000000000000000000000,bHidden=False,bNotConnectable=False,bDefaultValueIsReadOnly=False,bDefaultValueIsIgnored=False,bAdvancedView=False,)
+			End Object
+			Begin Object Class=K2Node_CallFunction Name="K2Node_CallFunction_2070"
+			   FunctionReference=(MemberParent=Class'/Script/Engine.MaterialInstanceDynamic',MemberName="SetTextureParameterValue")
+			   NodePosX=1808
+			   NodePosY=736
+			   NodeGuid=469CD7A2451ED9B272F746A1AB5CF15D
+			   CustomProperties Pin (PinId=B4EF9FDF44A64EAFA4BC079CAE9DD7EA,PinName="execute",PinToolTip="\nExec",PinType.PinCategory="exec",PinType.PinSubCategoryObject=None,PinType.PinSubCategoryMemberReference=(),PinType.bIsArray=False,PinType.bIsReference=False,PinType.bIsConst=False,PinType.bIsWeakPointer=False,LinkedTo=(K2Node_CallFunction_175 D03AF0CE4F40BC2B72AB9093E454AF7C,),PersistentGuid=00000000000000000000000000000000,bHidden=False,bNotConnectable=False,bDefaultValueIsReadOnly=False,bDefaultValueIsIgnored=False,bAdvancedView=False,)
+			   CustomProperties Pin (PinId=2F6CF49D48269B0073CB76ADDE38337B,PinName="then",PinToolTip="\nExec",Direction="EGPD_Output",PinType.PinCategory="exec",PinType.PinSubCategoryObject=None,PinType.PinSubCategoryMemberReference=(),PinType.bIsArray=False,PinType.bIsReference=False,PinType.bIsConst=False,PinType.bIsWeakPointer=False,LinkedTo=(K2Node_CallFunction_150 A7CDD48D45940F5F062D73A8005F2D43,),PersistentGuid=00000000000000000000000000000000,bHidden=False,bNotConnectable=False,bDefaultValueIsReadOnly=False,bDefaultValueIsIgnored=False,bAdvancedView=False,)
+			   CustomProperties Pin (PinId=058B7D14466F1D7BAD9496ACFD44563E,PinName="self",PinFriendlyName=NSLOCTEXT("K2Node", "Target", "Target"),PinToolTip="Target\nMaterial Instance Dynamic Reference",PinType.PinCategory="object",PinType.PinSubCategoryObject=Class'/Script/Engine.MaterialInstanceDynamic',PinType.PinSubCategoryMemberReference=(),PinType.bIsArray=False,PinType.bIsReference=False,PinType.bIsConst=False,PinType.bIsWeakPointer=False,LinkedTo=(K2Node_CallMaterialParameterCollectionFunction_112 EAC78C6D415430B012961DB7D2A9E3FE,),PersistentGuid=00000000000000000000000000000000,bHidden=False,bNotConnectable=False,bDefaultValueIsReadOnly=False,bDefaultValueIsIgnored=False,bAdvancedView=False,)
+			   CustomProperties Pin (PinId=239A55914B640A967CE0FF8958477331,PinName="ParameterName",PinToolTip="Parameter Name\nName",PinType.PinCategory="name",PinType.PinSubCategoryObject=None,PinType.PinSubCategoryMemberReference=(),PinType.bIsArray=False,PinType.bIsReference=False,PinType.bIsConst=False,PinType.bIsWeakPointer=False,DefaultValue="Heightfield",AutogeneratedDefaultValue="None",PersistentGuid=00000000000000000000000000000000,bHidden=False,bNotConnectable=False,bDefaultValueIsReadOnly=False,bDefaultValueIsIgnored=False,bAdvancedView=False,)
+			   CustomProperties Pin (PinId=A9EF7DBB418EB0C7514E928F8881B5BA,PinName="Value",PinToolTip="Value\nTexture Reference",PinType.PinCategory="object",PinType.PinSubCategoryObject=Class'/Script/Engine.Texture',PinType.PinSubCategoryMemberReference=(),PinType.bIsArray=False,PinType.bIsReference=False,PinType.bIsConst=False,PinType.bIsWeakPointer=False,LinkedTo=(K2Node_CallFunction_175 17ED0A19416B8925E2303CB3AE6D3CA9,),PersistentGuid=00000000000000000000000000000000,bHidden=False,bNotConnectable=False,bDefaultValueIsReadOnly=False,bDefaultValueIsIgnored=False,bAdvancedView=False,)
+			End Object
+			Begin Object Class=EdGraphNode_Comment Name="EdGraphNode_Comment_14"
+			   NodePosX=1168
+			   NodePosY=656
+			   NodeWidth=1504
+			   NodeHeight=432
+			   NodeComment="Compute surface normal"
+			   NodeGuid=1BB83BA04AA12E8E7A3A42BDAEB22E2B
+			End Object
+			END OBJECT
+	[/PARAM]
+[/OBJECT]
+[/EXCERPT:CompSurfNorm]
+
+[EXCERPT:HeightFunction]	
+[OBJECT:Snippet]
+	[PARAMLITERAL:id]
+		HeightFunction
+	[/PARAMLITERAL]
+	[PARAMLITERAL:type]
+		blueprint
+	[/PARAMLITERAL]
+	[PARAM:image]
+		![](Engine\Rendering\RenderTargets\BlueprintRenderTargets\HowTo\FluidSurface\FS_Get_Height_Function.png)
+	[/PARAM]
+	[PARAM:snippet]
+			BEGIN OBJECT
+				Begin Object Class=K2Node_FunctionEntry Name="K2Node_FunctionEntry_13"
+			   ExtraFlags=201457664
+			   SignatureName="GetHeightRT"
+			   bIsEditable=True
+			   NodePosX=-448
+			   NodePosY=64
+			   NodeGuid=913C30A7452125981360E4A7B27CDDC3
+			   CustomProperties Pin (PinId=C18F6AFF449D3F4919372C9F1B4309E8,PinName="then",Direction="EGPD_Output",PinType.PinCategory="exec",PinType.PinSubCategoryObject=None,PinType.PinSubCategoryMemberReference=(),PinType.bIsArray=False,PinType.bIsReference=False,PinType.bIsConst=False,PinType.bIsWeakPointer=False,LinkedTo=(K2Node_IfThenElse_0 A0B9280D4821886D1CA159A7869DBA1E,),PersistentGuid=00000000000000000000000000000000,bHidden=False,bNotConnectable=False,bDefaultValueIsReadOnly=False,bDefaultValueIsIgnored=False,bAdvancedView=False,)
+			   CustomProperties Pin (PinId=CE29685E44ED33E15E5D0D9748271EF6,PinName="Index",Direction="EGPD_Output",PinType.PinCategory="int",PinType.PinSubCategoryObject=None,PinType.PinSubCategoryMemberReference=(),PinType.bIsArray=False,PinType.bIsReference=False,PinType.bIsConst=False,PinType.bIsWeakPointer=False,LinkedTo=(K2Node_CallFunction_3 6F123690406FCADD70017E8C8F4524A3,K2Node_CallFunction_4 D1B78C8D4C2888F2046CA29CD09318AB,),PersistentGuid=00000000000000000000000000000000,bHidden=False,bNotConnectable=False,bDefaultValueIsReadOnly=False,bDefaultValueIsIgnored=False,bAdvancedView=False,)
+			   CustomProperties UserDefinedPin Name="Index" IsArray=0 IsReference=0 PinDir="EGPD_Output" Category=int 
+			End Object
+			Begin Object Class=K2Node_FunctionResult Name="K2Node_FunctionResult_1"
+			   SignatureName="GetHeightRT"
+			   bIsEditable=True
+			   NodePosX=208
+			   NodePosY=64
+			   NodeGuid=C0E69A544D517D9638CAD391DCA736BC
+			   CustomProperties Pin (PinId=02837D6646C94491A46D7B9DAEF03D2E,PinName="execute",PinType.PinCategory="exec",PinType.PinSubCategoryObject=None,PinType.PinSubCategoryMemberReference=(),PinType.bIsArray=False,PinType.bIsReference=False,PinType.bIsConst=False,PinType.bIsWeakPointer=False,LinkedTo=(K2Node_IfThenElse_0 A0BC7F5245926355C6276F8D8F11B26B,),PersistentGuid=00000000000000000000000000000000,bHidden=False,bNotConnectable=False,bDefaultValueIsReadOnly=False,bDefaultValueIsIgnored=False,bAdvancedView=False,)
+			   CustomProperties Pin (PinId=BC1DCC6B484C793FE4591BBE71E6721F,PinName="HeightRT",PinType.PinCategory="object",PinType.PinSubCategoryObject=Class'/Script/Engine.TextureRenderTarget2D',PinType.PinSubCategoryMemberReference=(),PinType.bIsArray=False,PinType.bIsReference=False,PinType.bIsConst=False,PinType.bIsWeakPointer=False,DefaultObject="/Game/ExampleContent/BlueprintFluidSurface/Height0.Height0",PersistentGuid=00000000000000000000000000000000,bHidden=False,bNotConnectable=False,bDefaultValueIsReadOnly=False,bDefaultValueIsIgnored=False,bAdvancedView=False,)
+			   CustomProperties UserDefinedPin Name="HeightRT" IsArray=0 IsReference=0 PinDir="EGPD_Input" Category=object SubCategoryObject=/Script/Engine.TextureRenderTarget2D 
+			End Object
+			Begin Object Class=K2Node_FunctionResult Name="K2Node_FunctionResult_0"
+			   SignatureName="GetHeightRT"
+			   bIsEditable=True
+			   NodePosX=464
+			   NodePosY=432
+			   NodeGuid=5E2189474D0F658457D194B546E5D93C
+			   CustomProperties Pin (PinId=FBA8E4C741C3F92D9A5AE9ABB2BE5324,PinName="execute",PinType.PinCategory="exec",PinType.PinSubCategoryObject=None,PinType.PinSubCategoryMemberReference=(),PinType.bIsArray=False,PinType.bIsReference=False,PinType.bIsConst=False,PinType.bIsWeakPointer=False,LinkedTo=(K2Node_IfThenElse_1 E12332984D712E65ED04A896EAD95967,),PersistentGuid=00000000000000000000000000000000,bHidden=False,bNotConnectable=False,bDefaultValueIsReadOnly=False,bDefaultValueIsIgnored=False,bAdvancedView=False,)
+			   CustomProperties Pin (PinId=B3FEEEF04A2AD7A86A1CF6B4D11BFCDE,PinName="HeightRT",PinType.PinCategory="object",PinType.PinSubCategoryObject=Class'/Script/Engine.TextureRenderTarget2D',PinType.PinSubCategoryMemberReference=(),PinType.bIsArray=False,PinType.bIsReference=False,PinType.bIsConst=False,PinType.bIsWeakPointer=False,DefaultObject="/Game/ExampleContent/BlueprintFluidSurface/Height2.Height2",PersistentGuid=00000000000000000000000000000000,bHidden=False,bNotConnectable=False,bDefaultValueIsReadOnly=False,bDefaultValueIsIgnored=False,bAdvancedView=False,)
+			   CustomProperties UserDefinedPin Name="HeightRT" IsArray=0 IsReference=0 PinDir="EGPD_Input" Category=object SubCategoryObject=/Script/Engine.TextureRenderTarget2D 
+			End Object
+			Begin Object Class=K2Node_FunctionResult Name="K2Node_FunctionResult_2"
+			   SignatureName="GetHeightRT"
+			   bIsEditable=True
+			   NodePosX=480
+			   NodePosY=304
+			   NodeGuid=3E5987D6480A36A0786171915ED3C87B
+			   CustomProperties Pin (PinId=90A7B313444419A6248B88B449044FAC,PinName="execute",PinType.PinCategory="exec",PinType.PinSubCategoryObject=None,PinType.PinSubCategoryMemberReference=(),PinType.bIsArray=False,PinType.bIsReference=False,PinType.bIsConst=False,PinType.bIsWeakPointer=False,LinkedTo=(K2Node_IfThenElse_1 62F7474B48DDB3A20A77C0B6D4D9C991,),PersistentGuid=00000000000000000000000000000000,bHidden=False,bNotConnectable=False,bDefaultValueIsReadOnly=False,bDefaultValueIsIgnored=False,bAdvancedView=False,)
+			   CustomProperties Pin (PinId=DD55A4E84B4089D33DE6E7827CF2E0D8,PinName="HeightRT",PinType.PinCategory="object",PinType.PinSubCategoryObject=Class'/Script/Engine.TextureRenderTarget2D',PinType.PinSubCategoryMemberReference=(),PinType.bIsArray=False,PinType.bIsReference=False,PinType.bIsConst=False,PinType.bIsWeakPointer=False,DefaultObject="/Game/ExampleContent/BlueprintFluidSurface/Height1.Height1",PersistentGuid=00000000000000000000000000000000,bHidden=False,bNotConnectable=False,bDefaultValueIsReadOnly=False,bDefaultValueIsIgnored=False,bAdvancedView=False,)
+			   CustomProperties UserDefinedPin Name="HeightRT" IsArray=0 IsReference=0 PinDir="EGPD_Input" Category=object SubCategoryObject=/Script/Engine.TextureRenderTarget2D 
+			End Object
+			Begin Object Class=K2Node_IfThenElse Name="K2Node_IfThenElse_0"
+			   NodePosX=-96
+			   NodePosY=64
+			   NodeGuid=3326D19648AE6267EB23EBB1D3B4883D
+			   CustomProperties Pin (PinId=A0B9280D4821886D1CA159A7869DBA1E,PinName="execute",PinType.PinCategory="exec",PinType.PinSubCategoryObject=None,PinType.PinSubCategoryMemberReference=(),PinType.bIsArray=False,PinType.bIsReference=False,PinType.bIsConst=False,PinType.bIsWeakPointer=False,LinkedTo=(K2Node_FunctionEntry_13 C18F6AFF449D3F4919372C9F1B4309E8,),PersistentGuid=00000000000000000000000000000000,bHidden=False,bNotConnectable=False,bDefaultValueIsReadOnly=False,bDefaultValueIsIgnored=False,bAdvancedView=False,)
+			   CustomProperties Pin (PinId=660B87044ACC27740E7A6CACF154AAC8,PinName="Condition",PinType.PinCategory="bool",PinType.PinSubCategoryObject=None,PinType.PinSubCategoryMemberReference=(),PinType.bIsArray=False,PinType.bIsReference=False,PinType.bIsConst=False,PinType.bIsWeakPointer=False,DefaultValue="false",AutogeneratedDefaultValue="true",LinkedTo=(K2Node_CallFunction_3 1F48718240A3D7BE359145B06C8B2935,),PersistentGuid=00000000000000000000000000000000,bHidden=False,bNotConnectable=False,bDefaultValueIsReadOnly=False,bDefaultValueIsIgnored=False,bAdvancedView=False,)
+			   CustomProperties Pin (PinId=A0BC7F5245926355C6276F8D8F11B26B,PinName="then",PinFriendlyName=NSLOCTEXT("K2Node", "true", "true"),Direction="EGPD_Output",PinType.PinCategory="exec",PinType.PinSubCategoryObject=None,PinType.PinSubCategoryMemberReference=(),PinType.bIsArray=False,PinType.bIsReference=False,PinType.bIsConst=False,PinType.bIsWeakPointer=False,LinkedTo=(K2Node_FunctionResult_1 02837D6646C94491A46D7B9DAEF03D2E,),PersistentGuid=00000000000000000000000000000000,bHidden=False,bNotConnectable=False,bDefaultValueIsReadOnly=False,bDefaultValueIsIgnored=False,bAdvancedView=False,)
+			   CustomProperties Pin (PinId=BD10A7BA4FF05FA4C02F4F81D841AE73,PinName="else",PinFriendlyName=NSLOCTEXT("K2Node", "false", "false"),Direction="EGPD_Output",PinType.PinCategory="exec",PinType.PinSubCategoryObject=None,PinType.PinSubCategoryMemberReference=(),PinType.bIsArray=False,PinType.bIsReference=False,PinType.bIsConst=False,PinType.bIsWeakPointer=False,LinkedTo=(K2Node_IfThenElse_1 053DABB2424D9448F37DE6A3BBE1FBA2,),PersistentGuid=00000000000000000000000000000000,bHidden=False,bNotConnectable=False,bDefaultValueIsReadOnly=False,bDefaultValueIsIgnored=False,bAdvancedView=False,)
+			End Object
+			Begin Object Class=K2Node_CallFunction Name="K2Node_CallFunction_4"
+			   bIsPureFunc=True
+			   FunctionReference=(MemberParent=Class'/Script/Engine.KismetMathLibrary',MemberName="EqualEqual_IntInt")
+			   NodePosY=336
+			   NodeGuid=184B252843FE8499E3B0F0A7B9C1350B
+			   CustomProperties Pin (PinId=689B98594FC640543A410AA84E1B2020,PinName="self",PinFriendlyName=NSLOCTEXT("K2Node", "Target", "Target"),PinToolTip="Target\nKismet Math Library Reference",PinType.PinCategory="object",PinType.PinSubCategoryObject=Class'/Script/Engine.KismetMathLibrary',PinType.PinSubCategoryMemberReference=(),PinType.bIsArray=False,PinType.bIsReference=False,PinType.bIsConst=False,PinType.bIsWeakPointer=False,DefaultObject="/Script/Engine.Default__KismetMathLibrary",PersistentGuid=00000000000000000000000000000000,bHidden=True,bNotConnectable=False,bDefaultValueIsReadOnly=False,bDefaultValueIsIgnored=False,bAdvancedView=False,)
+			   CustomProperties Pin (PinId=D1B78C8D4C2888F2046CA29CD09318AB,PinName="A",PinToolTip="A\nInteger",PinType.PinCategory="int",PinType.PinSubCategoryObject=None,PinType.PinSubCategoryMemberReference=(),PinType.bIsArray=False,PinType.bIsReference=False,PinType.bIsConst=False,PinType.bIsWeakPointer=False,DefaultValue="0",AutogeneratedDefaultValue="0",LinkedTo=(K2Node_FunctionEntry_13 CE29685E44ED33E15E5D0D9748271EF6,),PersistentGuid=00000000000000000000000000000000,bHidden=False,bNotConnectable=False,bDefaultValueIsReadOnly=False,bDefaultValueIsIgnored=False,bAdvancedView=False,)
+			   CustomProperties Pin (PinId=D64400C246BBF0A85F138E95C84DC6A5,PinName="B",PinToolTip="B\nInteger",PinType.PinCategory="int",PinType.PinSubCategoryObject=None,PinType.PinSubCategoryMemberReference=(),PinType.bIsArray=False,PinType.bIsReference=False,PinType.bIsConst=False,PinType.bIsWeakPointer=False,DefaultValue="1",AutogeneratedDefaultValue="0",PersistentGuid=00000000000000000000000000000000,bHidden=False,bNotConnectable=False,bDefaultValueIsReadOnly=False,bDefaultValueIsIgnored=False,bAdvancedView=False,)
+			   CustomProperties Pin (PinId=0CC60BBF4BF3CED20102BE86006E679B,PinName="ReturnValue",PinToolTip="Return Value\nBoolean",Direction="EGPD_Output",PinType.PinCategory="bool",PinType.PinSubCategoryObject=None,PinType.PinSubCategoryMemberReference=(),PinType.bIsArray=False,PinType.bIsReference=False,PinType.bIsConst=False,PinType.bIsWeakPointer=False,AutogeneratedDefaultValue="false",LinkedTo=(K2Node_IfThenElse_1 93FA7A8A43AFA3A11050C1812F5606B1,),PersistentGuid=00000000000000000000000000000000,bHidden=False,bNotConnectable=False,bDefaultValueIsReadOnly=False,bDefaultValueIsIgnored=False,bAdvancedView=False,)
+			End Object
+			Begin Object Class=K2Node_CallFunction Name="K2Node_CallFunction_3"
+			   bIsPureFunc=True
+			   FunctionReference=(MemberParent=Class'/Script/Engine.KismetMathLibrary',MemberName="EqualEqual_IntInt")
+			   NodePosX=-240
+			   NodePosY=240
+			   NodeGuid=A28EAB1848A89FA09D115DB4AEA5A700
+			   CustomProperties Pin (PinId=6B6E2BF84A8164759824A2AD6734398A,PinName="self",PinFriendlyName=NSLOCTEXT("K2Node", "Target", "Target"),PinToolTip="Target\nKismet Math Library Reference",PinType.PinCategory="object",PinType.PinSubCategoryObject=Class'/Script/Engine.KismetMathLibrary',PinType.PinSubCategoryMemberReference=(),PinType.bIsArray=False,PinType.bIsReference=False,PinType.bIsConst=False,PinType.bIsWeakPointer=False,DefaultObject="/Script/Engine.Default__KismetMathLibrary",PersistentGuid=00000000000000000000000000000000,bHidden=True,bNotConnectable=False,bDefaultValueIsReadOnly=False,bDefaultValueIsIgnored=False,bAdvancedView=False,)
+			   CustomProperties Pin (PinId=6F123690406FCADD70017E8C8F4524A3,PinName="A",PinToolTip="A\nInteger",PinType.PinCategory="int",PinType.PinSubCategoryObject=None,PinType.PinSubCategoryMemberReference=(),PinType.bIsArray=False,PinType.bIsReference=False,PinType.bIsConst=False,PinType.bIsWeakPointer=False,DefaultValue="0",AutogeneratedDefaultValue="0",LinkedTo=(K2Node_FunctionEntry_13 CE29685E44ED33E15E5D0D9748271EF6,),PersistentGuid=00000000000000000000000000000000,bHidden=False,bNotConnectable=False,bDefaultValueIsReadOnly=False,bDefaultValueIsIgnored=False,bAdvancedView=False,)
+			   CustomProperties Pin (PinId=4A55582E44C7AA3D3B97DF92E43ED551,PinName="B",PinToolTip="B\nInteger",PinType.PinCategory="int",PinType.PinSubCategoryObject=None,PinType.PinSubCategoryMemberReference=(),PinType.bIsArray=False,PinType.bIsReference=False,PinType.bIsConst=False,PinType.bIsWeakPointer=False,DefaultValue="0",AutogeneratedDefaultValue="0",PersistentGuid=00000000000000000000000000000000,bHidden=False,bNotConnectable=False,bDefaultValueIsReadOnly=False,bDefaultValueIsIgnored=False,bAdvancedView=False,)
+			   CustomProperties Pin (PinId=1F48718240A3D7BE359145B06C8B2935,PinName="ReturnValue",PinToolTip="Return Value\nBoolean",Direction="EGPD_Output",PinType.PinCategory="bool",PinType.PinSubCategoryObject=None,PinType.PinSubCategoryMemberReference=(),PinType.bIsArray=False,PinType.bIsReference=False,PinType.bIsConst=False,PinType.bIsWeakPointer=False,AutogeneratedDefaultValue="false",LinkedTo=(K2Node_IfThenElse_0 660B87044ACC27740E7A6CACF154AAC8,),PersistentGuid=00000000000000000000000000000000,bHidden=False,bNotConnectable=False,bDefaultValueIsReadOnly=False,bDefaultValueIsIgnored=False,bAdvancedView=False,)
+			End Object
+			Begin Object Class=K2Node_IfThenElse Name="K2Node_IfThenElse_1"
+			   NodePosX=224
+			   NodePosY=304
+			   NodeGuid=CBB88D764D19A8E5AF8DF69AFDC03E9F
+			   CustomProperties Pin (PinId=053DABB2424D9448F37DE6A3BBE1FBA2,PinName="execute",PinType.PinCategory="exec",PinType.PinSubCategoryObject=None,PinType.PinSubCategoryMemberReference=(),PinType.bIsArray=False,PinType.bIsReference=False,PinType.bIsConst=False,PinType.bIsWeakPointer=False,LinkedTo=(K2Node_IfThenElse_0 BD10A7BA4FF05FA4C02F4F81D841AE73,),PersistentGuid=00000000000000000000000000000000,bHidden=False,bNotConnectable=False,bDefaultValueIsReadOnly=False,bDefaultValueIsIgnored=False,bAdvancedView=False,)
+			   CustomProperties Pin (PinId=93FA7A8A43AFA3A11050C1812F5606B1,PinName="Condition",PinType.PinCategory="bool",PinType.PinSubCategoryObject=None,PinType.PinSubCategoryMemberReference=(),PinType.bIsArray=False,PinType.bIsReference=False,PinType.bIsConst=False,PinType.bIsWeakPointer=False,DefaultValue="false",AutogeneratedDefaultValue="true",LinkedTo=(K2Node_CallFunction_4 0CC60BBF4BF3CED20102BE86006E679B,),PersistentGuid=00000000000000000000000000000000,bHidden=False,bNotConnectable=False,bDefaultValueIsReadOnly=False,bDefaultValueIsIgnored=False,bAdvancedView=False,)
+			   CustomProperties Pin (PinId=62F7474B48DDB3A20A77C0B6D4D9C991,PinName="then",PinFriendlyName=NSLOCTEXT("K2Node", "true", "true"),Direction="EGPD_Output",PinType.PinCategory="exec",PinType.PinSubCategoryObject=None,PinType.PinSubCategoryMemberReference=(),PinType.bIsArray=False,PinType.bIsReference=False,PinType.bIsConst=False,PinType.bIsWeakPointer=False,LinkedTo=(K2Node_FunctionResult_2 90A7B313444419A6248B88B449044FAC,),PersistentGuid=00000000000000000000000000000000,bHidden=False,bNotConnectable=False,bDefaultValueIsReadOnly=False,bDefaultValueIsIgnored=False,bAdvancedView=False,)
+			   CustomProperties Pin (PinId=E12332984D712E65ED04A896EAD95967,PinName="else",PinFriendlyName=NSLOCTEXT("K2Node", "false", "false"),Direction="EGPD_Output",PinType.PinCategory="exec",PinType.PinSubCategoryObject=None,PinType.PinSubCategoryMemberReference=(),PinType.bIsArray=False,PinType.bIsReference=False,PinType.bIsConst=False,PinType.bIsWeakPointer=False,LinkedTo=(K2Node_FunctionResult_0 FBA8E4C741C3F92D9A5AE9ABB2BE5324,),PersistentGuid=00000000000000000000000000000000,bHidden=False,bNotConnectable=False,bDefaultValueIsReadOnly=False,bDefaultValueIsIgnored=False,bAdvancedView=False,)
+			End Object
+			Begin Object Class=EdGraphNode_Comment Name="EdGraphNode_Comment_3"
+			   NodePosX=-496
+			   NodePosY=16
+			   NodeWidth=1216
+			   NodeHeight=560
+			   NodeComment="Get and store height of render target"
+			   NodeGuid=7A08CC0B4155AD6B040A0BA32B80BC13
+			End Object
+			END OBJECT
+	[/PARAM]
+[/OBJECT]
+[/EXCERPT:HeightFunction]
+
+[EXCERPT:GetLastHeight]	
+[OBJECT:Snippet]
+	[PARAMLITERAL:id]
+		GetLastHeight
+	[/PARAMLITERAL]
+	[PARAMLITERAL:type]
+		blueprint
+	[/PARAMLITERAL]
+	[PARAM:image]
+		![](Engine\Rendering\RenderTargets\BlueprintRenderTargets\HowTo\FluidSurface\FS_Store_Last_Height.png)
+	[/PARAM]
+	[PARAM:snippet]
+			BEGIN OBJECT
+				Begin Object Class=K2Node_FunctionEntry Name="K2Node_FunctionEntry_14"
+			   ExtraFlags=201457664
+			   SignatureName="GetLastHeight"
+			   bIsEditable=True
+			   NodePosX=64
+			   NodePosY=80
+			   NodeGuid=03E87A9D404B5DCD0D91B3A9E95D557B
+			   CustomProperties Pin (PinId=9F06A42D485493A5ED5E2A86F323CE80,PinName="then",Direction="EGPD_Output",PinType.PinCategory="exec",PinType.PinSubCategoryObject=None,PinType.PinSubCategoryMemberReference=(),PinType.bIsArray=False,PinType.bIsReference=False,PinType.bIsConst=False,PinType.bIsWeakPointer=False,LinkedTo=(K2Node_CallFunction_7 7A6100CD4810B259DCAD65AFEDBB3BEF,),PersistentGuid=00000000000000000000000000000000,bHidden=False,bNotConnectable=False,bDefaultValueIsReadOnly=False,bDefaultValueIsIgnored=False,bAdvancedView=False,)
+			   CustomProperties Pin (PinId=28696A0A41033788565DA1B2A84ADA66,PinName="CurrentHeightIndex",Direction="EGPD_Output",PinType.PinCategory="int",PinType.PinSubCategoryObject=None,PinType.PinSubCategoryMemberReference=(),PinType.bIsArray=False,PinType.bIsReference=False,PinType.bIsConst=False,PinType.bIsWeakPointer=False,LinkedTo=(K2Node_CallFunction_5 4F23D8484B7D0B7BD0F1638F9BCBF2C5,),PersistentGuid=00000000000000000000000000000000,bHidden=False,bNotConnectable=False,bDefaultValueIsReadOnly=False,bDefaultValueIsIgnored=False,bAdvancedView=False,)
+			   CustomProperties Pin (PinId=E1D5AEDD41D12B4B06B693B054174631,PinName="NumFramesOld",Direction="EGPD_Output",PinType.PinCategory="int",PinType.PinSubCategoryObject=None,PinType.PinSubCategoryMemberReference=(),PinType.bIsArray=False,PinType.bIsReference=False,PinType.bIsConst=False,PinType.bIsWeakPointer=False,LinkedTo=(K2Node_CallFunction_5 565E09CC43447352603079BCAE9F22FD,),PersistentGuid=00000000000000000000000000000000,bHidden=False,bNotConnectable=False,bDefaultValueIsReadOnly=False,bDefaultValueIsIgnored=False,bAdvancedView=False,)
+			   CustomProperties UserDefinedPin Name="CurrentHeightIndex" IsArray=0 IsReference=0 PinDir="EGPD_Output" Category=int 
+			   CustomProperties UserDefinedPin Name="NumFramesOld" IsArray=0 IsReference=0 PinDir="EGPD_Output" Category=int 
+			End Object
+			Begin Object Class=K2Node_FunctionResult Name="K2Node_FunctionResult_3"
+			   SignatureName="GetLastHeight"
+			   bIsEditable=True
+			   NodePosX=1104
+			   NodePosY=77
+			   NodeGuid=3F52E4BF41515475AC280CAE7E1338E3
+			   CustomProperties Pin (PinId=5B8AF4BE4ED1AC82B88640A3AEF4C2CE,PinName="execute",PinType.PinCategory="exec",PinType.PinSubCategoryObject=None,PinType.PinSubCategoryMemberReference=(),PinType.bIsArray=False,PinType.bIsReference=False,PinType.bIsConst=False,PinType.bIsWeakPointer=False,LinkedTo=(K2Node_CallFunction_7 1FC4F07A4D2FF9ED32B9B681042D36D2,),PersistentGuid=00000000000000000000000000000000,bHidden=False,bNotConnectable=False,bDefaultValueIsReadOnly=False,bDefaultValueIsIgnored=False,bAdvancedView=False,)
+			   CustomProperties Pin (PinId=70D7356648EC97BBF0CAFFADF94282BD,PinName="HeightRT",PinType.PinCategory="object",PinType.PinSubCategoryObject=Class'/Script/Engine.TextureRenderTarget2D',PinType.PinSubCategoryMemberReference=(),PinType.bIsArray=False,PinType.bIsReference=False,PinType.bIsConst=False,PinType.bIsWeakPointer=False,LinkedTo=(K2Node_CallFunction_7 B60E8FC84720B18E018030B1CD810586,),PersistentGuid=00000000000000000000000000000000,bHidden=False,bNotConnectable=False,bDefaultValueIsReadOnly=False,bDefaultValueIsIgnored=False,bAdvancedView=False,)
+			   CustomProperties UserDefinedPin Name="HeightRT" IsArray=0 IsReference=0 PinDir="EGPD_MAX" Category=object SubCategoryObject=/Script/Engine.TextureRenderTarget2D 
+			End Object
+			Begin Object Class=K2Node_CommutativeAssociativeBinaryOperator Name="K2Node_CommutativeAssociativeBinaryOperator_65"
+			   bIsPureFunc=True
+			   FunctionReference=(MemberParent=Class'/Script/Engine.KismetMathLibrary',MemberName="Add_IntInt")
+			   NodePosX=400
+			   NodePosY=144
+			   NodeGuid=CBB6E2C54BCC6D069F7CBAAB7D02E5E0
+			   CustomProperties Pin (PinId=17002D2742D87E63B74AB2A8EF48AC16,PinName="self",PinFriendlyName=NSLOCTEXT("K2Node", "Target", "Target"),PinToolTip="Target\nKismet Math Library Reference",PinType.PinCategory="object",PinType.PinSubCategoryObject=Class'/Script/Engine.KismetMathLibrary',PinType.PinSubCategoryMemberReference=(),PinType.bIsArray=False,PinType.bIsReference=False,PinType.bIsConst=False,PinType.bIsWeakPointer=False,DefaultObject="/Script/Engine.Default__KismetMathLibrary",PersistentGuid=00000000000000000000000000000000,bHidden=True,bNotConnectable=False,bDefaultValueIsReadOnly=False,bDefaultValueIsIgnored=False,bAdvancedView=False,)
+			   CustomProperties Pin (PinId=130B8A0D415540B5F89C2FB51160DC4D,PinName="A",PinToolTip="A\nInteger",PinType.PinCategory="int",PinType.PinSubCategoryObject=None,PinType.PinSubCategoryMemberReference=(),PinType.bIsArray=False,PinType.bIsReference=False,PinType.bIsConst=False,PinType.bIsWeakPointer=False,DefaultValue="0",AutogeneratedDefaultValue="0",LinkedTo=(K2Node_CallFunction_5 3A457B3241FF409309921EA29AC0273E,),PersistentGuid=00000000000000000000000000000000,bHidden=False,bNotConnectable=False,bDefaultValueIsReadOnly=False,bDefaultValueIsIgnored=False,bAdvancedView=False,)
+			   CustomProperties Pin (PinId=DCF2AC4947901AF2F777D09417073CC5,PinName="B",PinToolTip="B\nInteger",PinType.PinCategory="int",PinType.PinSubCategoryObject=None,PinType.PinSubCategoryMemberReference=(),PinType.bIsArray=False,PinType.bIsReference=False,PinType.bIsConst=False,PinType.bIsWeakPointer=False,DefaultValue="3",AutogeneratedDefaultValue="1",PersistentGuid=00000000000000000000000000000000,bHidden=False,bNotConnectable=False,bDefaultValueIsReadOnly=False,bDefaultValueIsIgnored=False,bAdvancedView=False,)
+			   CustomProperties Pin (PinId=798F377F40C8B97E588F418609D87614,PinName="ReturnValue",PinToolTip="Return Value\nInteger",Direction="EGPD_Output",PinType.PinCategory="int",PinType.PinSubCategoryObject=None,PinType.PinSubCategoryMemberReference=(),PinType.bIsArray=False,PinType.bIsReference=False,PinType.bIsConst=False,PinType.bIsWeakPointer=False,AutogeneratedDefaultValue="0",LinkedTo=(K2Node_CallFunction_6 4CFFCE0848766DA2C9B9AC959A16851D,),PersistentGuid=00000000000000000000000000000000,bHidden=False,bNotConnectable=False,bDefaultValueIsReadOnly=False,bDefaultValueIsIgnored=False,bAdvancedView=False,)
+			End Object
+			Begin Object Class=K2Node_CallFunction Name="K2Node_CallFunction_5"
+			   bIsPureFunc=True
+			   FunctionReference=(MemberParent=Class'/Script/Engine.KismetMathLibrary',MemberName="Subtract_IntInt")
+			   NodePosX=272
+			   NodePosY=128
+			   NodeGuid=1FB924EC44EE01DD9186E597F41287B2
+			   CustomProperties Pin (PinId=966D1A5C4DCB76F3AAF5E28CA3A04F60,PinName="self",PinFriendlyName=NSLOCTEXT("K2Node", "Target", "Target"),PinToolTip="Target\nKismet Math Library Reference",PinType.PinCategory="object",PinType.PinSubCategoryObject=Class'/Script/Engine.KismetMathLibrary',PinType.PinSubCategoryMemberReference=(),PinType.bIsArray=False,PinType.bIsReference=False,PinType.bIsConst=False,PinType.bIsWeakPointer=False,DefaultObject="/Script/Engine.Default__KismetMathLibrary",PersistentGuid=00000000000000000000000000000000,bHidden=True,bNotConnectable=False,bDefaultValueIsReadOnly=False,bDefaultValueIsIgnored=False,bAdvancedView=False,)
+			   CustomProperties Pin (PinId=4F23D8484B7D0B7BD0F1638F9BCBF2C5,PinName="A",PinToolTip="A\nInteger",PinType.PinCategory="int",PinType.PinSubCategoryObject=None,PinType.PinSubCategoryMemberReference=(),PinType.bIsArray=False,PinType.bIsReference=False,PinType.bIsConst=False,PinType.bIsWeakPointer=False,DefaultValue="0",AutogeneratedDefaultValue="0",LinkedTo=(K2Node_FunctionEntry_14 28696A0A41033788565DA1B2A84ADA66,),PersistentGuid=00000000000000000000000000000000,bHidden=False,bNotConnectable=False,bDefaultValueIsReadOnly=False,bDefaultValueIsIgnored=False,bAdvancedView=False,)
+			   CustomProperties Pin (PinId=565E09CC43447352603079BCAE9F22FD,PinName="B",PinToolTip="B\nInteger",PinType.PinCategory="int",PinType.PinSubCategoryObject=None,PinType.PinSubCategoryMemberReference=(),PinType.bIsArray=False,PinType.bIsReference=False,PinType.bIsConst=False,PinType.bIsWeakPointer=False,DefaultValue="0",AutogeneratedDefaultValue="1",LinkedTo=(K2Node_FunctionEntry_14 E1D5AEDD41D12B4B06B693B054174631,),PersistentGuid=00000000000000000000000000000000,bHidden=False,bNotConnectable=False,bDefaultValueIsReadOnly=False,bDefaultValueIsIgnored=False,bAdvancedView=False,)
+			   CustomProperties Pin (PinId=3A457B3241FF409309921EA29AC0273E,PinName="ReturnValue",PinToolTip="Return Value\nInteger",Direction="EGPD_Output",PinType.PinCategory="int",PinType.PinSubCategoryObject=None,PinType.PinSubCategoryMemberReference=(),PinType.bIsArray=False,PinType.bIsReference=False,PinType.bIsConst=False,PinType.bIsWeakPointer=False,AutogeneratedDefaultValue="0",LinkedTo=(K2Node_CommutativeAssociativeBinaryOperator_65 130B8A0D415540B5F89C2FB51160DC4D,),PersistentGuid=00000000000000000000000000000000,bHidden=False,bNotConnectable=False,bDefaultValueIsReadOnly=False,bDefaultValueIsIgnored=False,bAdvancedView=False,)
+			End Object
+			Begin Object Class=K2Node_CallFunction Name="K2Node_CallFunction_6"
+			   bIsPureFunc=True
+			   FunctionReference=(MemberParent=Class'/Script/Engine.KismetMathLibrary',MemberName="Percent_IntInt")
+			   NodePosX=592
+			   NodePosY=144
+			   NodeGuid=73D9E99341FF372E0058FBB7860084A4
+			   CustomProperties Pin (PinId=1DE2B8E345D252D7B752368777776523,PinName="self",PinFriendlyName=NSLOCTEXT("K2Node", "Target", "Target"),PinToolTip="Target\nKismet Math Library Reference",PinType.PinCategory="object",PinType.PinSubCategoryObject=Class'/Script/Engine.KismetMathLibrary',PinType.PinSubCategoryMemberReference=(),PinType.bIsArray=False,PinType.bIsReference=False,PinType.bIsConst=False,PinType.bIsWeakPointer=False,DefaultObject="/Script/Engine.Default__KismetMathLibrary",PersistentGuid=00000000000000000000000000000000,bHidden=True,bNotConnectable=False,bDefaultValueIsReadOnly=False,bDefaultValueIsIgnored=False,bAdvancedView=False,)
+			   CustomProperties Pin (PinId=4CFFCE0848766DA2C9B9AC959A16851D,PinName="A",PinToolTip="A\nInteger",PinType.PinCategory="int",PinType.PinSubCategoryObject=None,PinType.PinSubCategoryMemberReference=(),PinType.bIsArray=False,PinType.bIsReference=False,PinType.bIsConst=False,PinType.bIsWeakPointer=False,DefaultValue="0",AutogeneratedDefaultValue="0",LinkedTo=(K2Node_CommutativeAssociativeBinaryOperator_65 798F377F40C8B97E588F418609D87614,),PersistentGuid=00000000000000000000000000000000,bHidden=False,bNotConnectable=False,bDefaultValueIsReadOnly=False,bDefaultValueIsIgnored=False,bAdvancedView=False,)
+			   CustomProperties Pin (PinId=441DF6414D31A700FBC91291E8554DB5,PinName="B",PinToolTip="B\nInteger",PinType.PinCategory="int",PinType.PinSubCategoryObject=None,PinType.PinSubCategoryMemberReference=(),PinType.bIsArray=False,PinType.bIsReference=False,PinType.bIsConst=False,PinType.bIsWeakPointer=False,DefaultValue="3",AutogeneratedDefaultValue="1",PersistentGuid=00000000000000000000000000000000,bHidden=False,bNotConnectable=False,bDefaultValueIsReadOnly=False,bDefaultValueIsIgnored=False,bAdvancedView=False,)
+			   CustomProperties Pin (PinId=98724EB943D25E225DE4828843D4C9B1,PinName="ReturnValue",PinToolTip="Return Value\nInteger",Direction="EGPD_Output",PinType.PinCategory="int",PinType.PinSubCategoryObject=None,PinType.PinSubCategoryMemberReference=(),PinType.bIsArray=False,PinType.bIsReference=False,PinType.bIsConst=False,PinType.bIsWeakPointer=False,AutogeneratedDefaultValue="0",LinkedTo=(K2Node_CallFunction_7 E2014ABE479C4287BE419980E31543A6,),PersistentGuid=00000000000000000000000000000000,bHidden=False,bNotConnectable=False,bDefaultValueIsReadOnly=False,bDefaultValueIsIgnored=False,bAdvancedView=False,)
+			End Object
+			Begin Object Class=K2Node_CallFunction Name="K2Node_CallFunction_7"
+			   FunctionReference=(MemberName="GetHeightRT",MemberGuid=2EEA338F4CA878E7228A10BB2DD63E00,bSelfContext=True)
+			   NodePosX=800
+			   NodePosY=61
+			   NodeGuid=AD82B4374BDFD2811EB651A7B331E7BB
+			   CustomProperties Pin (PinId=7A6100CD4810B259DCAD65AFEDBB3BEF,PinName="execute",PinToolTip="\nExec",PinType.PinCategory="exec",PinType.PinSubCategoryObject=None,PinType.PinSubCategoryMemberReference=(),PinType.bIsArray=False,PinType.bIsReference=False,PinType.bIsConst=False,PinType.bIsWeakPointer=False,LinkedTo=(K2Node_FunctionEntry_14 9F06A42D485493A5ED5E2A86F323CE80,),PersistentGuid=00000000000000000000000000000000,bHidden=False,bNotConnectable=False,bDefaultValueIsReadOnly=False,bDefaultValueIsIgnored=False,bAdvancedView=False,)
+			   CustomProperties Pin (PinId=1FC4F07A4D2FF9ED32B9B681042D36D2,PinName="then",PinToolTip="\nExec",Direction="EGPD_Output",PinType.PinCategory="exec",PinType.PinSubCategoryObject=None,PinType.PinSubCategoryMemberReference=(),PinType.bIsArray=False,PinType.bIsReference=False,PinType.bIsConst=False,PinType.bIsWeakPointer=False,LinkedTo=(K2Node_FunctionResult_3 5B8AF4BE4ED1AC82B88640A3AEF4C2CE,),PersistentGuid=00000000000000000000000000000000,bHidden=False,bNotConnectable=False,bDefaultValueIsReadOnly=False,bDefaultValueIsIgnored=False,bAdvancedView=False,)
+			   CustomProperties Pin (PinId=A614D3A94AD3D500F717F4BC121A02DB,PinName="self",PinFriendlyName=NSLOCTEXT("K2Node", "Target", "Target"),PinToolTip="Target\nSelf Reference",PinType.PinCategory="object",PinType.PinSubCategory="self",PinType.PinSubCategoryObject=None,PinType.PinSubCategoryMemberReference=(),PinType.bIsArray=False,PinType.bIsReference=False,PinType.bIsConst=False,PinType.bIsWeakPointer=False,PersistentGuid=00000000000000000000000000000000,bHidden=False,bNotConnectable=False,bDefaultValueIsReadOnly=False,bDefaultValueIsIgnored=False,bAdvancedView=False,)
+			   CustomProperties Pin (PinId=E2014ABE479C4287BE419980E31543A6,PinName="Index",PinToolTip="Index\nInteger",PinType.PinCategory="int",PinType.PinSubCategoryObject=None,PinType.PinSubCategoryMemberReference=(),PinType.bIsArray=False,PinType.bIsReference=False,PinType.bIsConst=False,PinType.bIsWeakPointer=False,DefaultValue="0",AutogeneratedDefaultValue="0",LinkedTo=(K2Node_CallFunction_6 98724EB943D25E225DE4828843D4C9B1,),PersistentGuid=00000000000000000000000000000000,bHidden=False,bNotConnectable=False,bDefaultValueIsReadOnly=False,bDefaultValueIsIgnored=False,bAdvancedView=False,)
+			   CustomProperties Pin (PinId=B60E8FC84720B18E018030B1CD810586,PinName="HeightRT",PinToolTip="Height RT\nTexture Render Target 2D Reference",Direction="EGPD_Output",PinType.PinCategory="object",PinType.PinSubCategoryObject=Class'/Script/Engine.TextureRenderTarget2D',PinType.PinSubCategoryMemberReference=(),PinType.bIsArray=False,PinType.bIsReference=False,PinType.bIsConst=False,PinType.bIsWeakPointer=False,LinkedTo=(K2Node_FunctionResult_3 70D7356648EC97BBF0CAFFADF94282BD,),PersistentGuid=00000000000000000000000000000000,bHidden=False,bNotConnectable=False,bDefaultValueIsReadOnly=False,bDefaultValueIsIgnored=False,bAdvancedView=False,)
+			End Object
+			Begin Object Class=EdGraphNode_Comment Name="EdGraphNode_Comment_4"
+			   NodePosX=16
+			   NodePosY=16
+			   NodeWidth=1272
+			   NodeHeight=243
+			   NodeComment="Store the last height for later use"
+			   NodeGuid=ED781A5F4F580B254F38B4B32A79BEE8
+			End Object
+
+			END OBJECT
+	[/PARAM]
+[/OBJECT]
+[/EXCERPT:GetLastHeight]
+
+[EXCERPT:CameraTrace]	
+[OBJECT:Snippet]
+	[PARAMLITERAL:id]
+		TraceFromCamera
+	[/PARAMLITERAL]
+	[PARAMLITERAL:type]
+		blueprint
+	[/PARAMLITERAL]
+	[PARAM:image]
+		![](Engine\Rendering\RenderTargets\BlueprintRenderTargets\HowTo\HeightFieldPainter\HFP_BP_Setup_05.png)
+	[/PARAM]
+	[PARAM:snippet]
+				Begin Object Class=K2Node_FunctionEntry Name="K2Node_FunctionEntry_9"
+			   ExtraFlags=201457664
+			   SignatureName="TraceFromCamera"
+			   bIsEditable=True
+			   NodePosX=-96
+			   NodePosY=-384
+			   NodeGuid=A6585C2E45244BF6E4B02AAADF163432
+			   CustomProperties Pin (PinId=598C3D6A46B70A76463B448AA52DE23A,PinName="then",Direction="EGPD_Output",PinType.PinCategory="exec",PinType.PinSubCategoryObject=None,PinType.PinSubCategoryMemberReference=(),PinType.bIsArray=False,PinType.bIsReference=False,PinType.bIsConst=False,PinType.bIsWeakPointer=False,LinkedTo=(K2Node_CallFunction_65598 B3AD098C4E9F315318BE2C9A781CFD16,),PersistentGuid=00000000000000000000000000000000,bHidden=False,bNotConnectable=False,bDefaultValueIsReadOnly=False,bDefaultValueIsIgnored=False,bAdvancedView=False,)
+			End Object
+			Begin Object Class=K2Node_MacroInstance Name="K2Node_MacroInstance_65"
+			   MacroGraphReference=(MacroGraph=EdGraph'/Engine/EditorBlueprintResources/StandardMacros.StandardMacros:IsValid',GraphBlueprint=Blueprint'/Engine/EditorBlueprintResources/StandardMacros.StandardMacros',GraphGuid=64422BCD430703FF5CAEA8B79A32AA65)
+			   NodePosX=1984
+			   NodePosY=-432
+			   ErrorType=4
+			   ErrorMsg="Info The execution path doesn\'t end with a return node.  Is Not Valid "
+			   NodeGuid=269E6BC649B1DF4406573EB2C13D874A
+			   CustomProperties Pin (PinId=54DCDF6A4C18DFFBF71F54AE24DA40C0,PinName="exec",PinType.PinCategory="exec",PinType.PinSubCategoryObject=None,PinType.PinSubCategoryMemberReference=(),PinType.bIsArray=False,PinType.bIsReference=False,PinType.bIsConst=False,PinType.bIsWeakPointer=False,LinkedTo=(K2Node_IfThenElse_161 8DD9A5BC40D019A1AFD6F284F8C15FFD,),PersistentGuid=00000000000000000000000000000000,bHidden=False,bNotConnectable=False,bDefaultValueIsReadOnly=False,bDefaultValueIsIgnored=False,bAdvancedView=False,)
+			   CustomProperties Pin (PinId=C3FB28AC47AAD2B5998C2AA9D4862C2E,PinName="InputObject",PinType.PinCategory="object",PinType.PinSubCategoryObject=Class'/Script/CoreUObject.Object',PinType.PinSubCategoryMemberReference=(),PinType.bIsArray=False,PinType.bIsReference=False,PinType.bIsConst=False,PinType.bIsWeakPointer=False,LinkedTo=(K2Node_CallFunction_65759 F1D6985542953219C7EF09A3FBFA8AF0,),PersistentGuid=00000000000000000000000000000000,bHidden=False,bNotConnectable=False,bDefaultValueIsReadOnly=False,bDefaultValueIsIgnored=False,bAdvancedView=False,)
+			   CustomProperties Pin (PinId=E75DA26F4EBE18C63E3CEDA6B1D88DC6,PinName="Is Valid",Direction="EGPD_Output",PinType.PinCategory="exec",PinType.PinSubCategoryObject=None,PinType.PinSubCategoryMemberReference=(),PinType.bIsArray=False,PinType.bIsReference=False,PinType.bIsConst=False,PinType.bIsWeakPointer=False,LinkedTo=(K2Node_CallFunction_1861 6F410C0B46D7B07315AD748B03EB86C8,),PersistentGuid=00000000000000000000000000000000,bHidden=False,bNotConnectable=False,bDefaultValueIsReadOnly=False,bDefaultValueIsIgnored=False,bAdvancedView=False,)
+			   CustomProperties Pin (PinId=7426FEC54F3C31FE985AE088ACC9F56E,PinName="Is Not Valid",Direction="EGPD_Output",PinType.PinCategory="exec",PinType.PinSubCategoryObject=None,PinType.PinSubCategoryMemberReference=(),PinType.bIsArray=False,PinType.bIsReference=False,PinType.bIsConst=False,PinType.bIsWeakPointer=False,PersistentGuid=00000000000000000000000000000000,bHidden=False,bNotConnectable=False,bDefaultValueIsReadOnly=False,bDefaultValueIsIgnored=False,bAdvancedView=False,)
+			End Object
+			Begin Object Class=K2Node_CallFunction Name="K2Node_CallFunction_2902"
+			   bIsPureFunc=True
+			   FunctionReference=(MemberParent=Class'/Script/Engine.GameplayStatics',MemberName="GetPlayerController")
+			   NodePosX=2192
+			   NodePosY=-176
+			   NodeGuid=5089AE9B44A43150C97C5F92ECE42A73
+			   CustomProperties Pin (PinId=45DD77B44C0471A5CDD2C7BB7D1F8F28,PinName="self",PinFriendlyName=NSLOCTEXT("K2Node", "Target", "Target"),PinType.PinCategory="object",PinType.PinSubCategoryObject=Class'/Script/Engine.GameplayStatics',PinType.PinSubCategoryMemberReference=(),PinType.bIsArray=False,PinType.bIsReference=False,PinType.bIsConst=False,PinType.bIsWeakPointer=False,DefaultObject="/Script/Engine.Default__GameplayStatics",PersistentGuid=00000000000000000000000000000000,bHidden=True,bNotConnectable=False,bDefaultValueIsReadOnly=False,bDefaultValueIsIgnored=False,bAdvancedView=False,)
+			   CustomProperties Pin (PinId=56BF7CB143F1FAE6723E31A1D3FC3447,PinName="WorldContextObject",PinType.PinCategory="object",PinType.PinSubCategoryObject=Class'/Script/CoreUObject.Object',PinType.PinSubCategoryMemberReference=(),PinType.bIsArray=False,PinType.bIsReference=False,PinType.bIsConst=False,PinType.bIsWeakPointer=False,PersistentGuid=00000000000000000000000000000000,bHidden=True,bNotConnectable=False,bDefaultValueIsReadOnly=False,bDefaultValueIsIgnored=False,bAdvancedView=False,)
+			   CustomProperties Pin (PinId=6AA23C3242491D37AACF8CB51ECB365D,PinName="PlayerIndex",PinType.PinCategory="int",PinType.PinSubCategoryObject=None,PinType.PinSubCategoryMemberReference=(),PinType.bIsArray=False,PinType.bIsReference=False,PinType.bIsConst=False,PinType.bIsWeakPointer=False,DefaultValue="0",AutogeneratedDefaultValue="0",PersistentGuid=00000000000000000000000000000000,bHidden=False,bNotConnectable=False,bDefaultValueIsReadOnly=False,bDefaultValueIsIgnored=False,bAdvancedView=False,)
+			   CustomProperties Pin (PinId=3119A2234CC8CC2E58EDDA94ED78DB30,PinName="ReturnValue",Direction="EGPD_Output",PinType.PinCategory="object",PinType.PinSubCategoryObject=Class'/Script/Engine.PlayerController',PinType.PinSubCategoryMemberReference=(),PinType.bIsArray=False,PinType.bIsReference=False,PinType.bIsConst=False,PinType.bIsWeakPointer=False,LinkedTo=(K2Node_CallFunction_1861 8AAA25E3411AD93065E153BDF12317E1,),PersistentGuid=00000000000000000000000000000000,bHidden=False,bNotConnectable=False,bDefaultValueIsReadOnly=False,bDefaultValueIsIgnored=False,bAdvancedView=False,)
+			End Object
+			Begin Object Class=K2Node_CallFunction Name="K2Node_CallFunction_64853"
+			   bIsPureFunc=True
+			   FunctionReference=(MemberParent=Class'/Script/Engine.GameplayStatics',MemberName="GetPlayerCameraManager")
+			   NodePosX=-112
+			   NodePosY=-176
+			   NodeGuid=8798F1F24424791656D634AE744FBB71
+			   CustomProperties Pin (PinId=73A2796F477586504925ACB63B1C3373,PinName="self",PinFriendlyName=NSLOCTEXT("K2Node", "Target", "Target"),PinType.PinCategory="object",PinType.PinSubCategoryObject=Class'/Script/Engine.GameplayStatics',PinType.PinSubCategoryMemberReference=(),PinType.bIsArray=False,PinType.bIsReference=False,PinType.bIsConst=False,PinType.bIsWeakPointer=False,DefaultObject="/Script/Engine.Default__GameplayStatics",PersistentGuid=00000000000000000000000000000000,bHidden=True,bNotConnectable=False,bDefaultValueIsReadOnly=False,bDefaultValueIsIgnored=False,bAdvancedView=False,)
+			   CustomProperties Pin (PinId=A29AA4884E79B1B2CBFCAE93A185BFA1,PinName="WorldContextObject",PinType.PinCategory="object",PinType.PinSubCategoryObject=Class'/Script/CoreUObject.Object',PinType.PinSubCategoryMemberReference=(),PinType.bIsArray=False,PinType.bIsReference=False,PinType.bIsConst=False,PinType.bIsWeakPointer=False,DefaultValue="0",PersistentGuid=00000000000000000000000000000000,bHidden=True,bNotConnectable=False,bDefaultValueIsReadOnly=False,bDefaultValueIsIgnored=False,bAdvancedView=False,)
+			   CustomProperties Pin (PinId=F04D1A5446F837775E768082313D9559,PinName="PlayerIndex",PinType.PinCategory="int",PinType.PinSubCategoryObject=None,PinType.PinSubCategoryMemberReference=(),PinType.bIsArray=False,PinType.bIsReference=False,PinType.bIsConst=False,PinType.bIsWeakPointer=False,DefaultValue="0",AutogeneratedDefaultValue="0",PersistentGuid=00000000000000000000000000000000,bHidden=False,bNotConnectable=False,bDefaultValueIsReadOnly=False,bDefaultValueIsIgnored=False,bAdvancedView=False,)
+			   CustomProperties Pin (PinId=88B4C49D414D557C8728F39A3967490F,PinName="ReturnValue",Direction="EGPD_Output",PinType.PinCategory="object",PinType.PinSubCategoryObject=Class'/Script/Engine.PlayerCameraManager',PinType.PinSubCategoryMemberReference=(),PinType.bIsArray=False,PinType.bIsReference=False,PinType.bIsConst=False,PinType.bIsWeakPointer=False,LinkedTo=(K2Node_CallFunction_64919 136B71094B407B9186881DB137E50CB0,K2Node_CallFunction_65094 451585FC4960B7F55857FDB8A5A37392,),PersistentGuid=00000000000000000000000000000000,bHidden=False,bNotConnectable=False,bDefaultValueIsReadOnly=False,bDefaultValueIsIgnored=False,bAdvancedView=False,)
+			End Object
+			Begin Object Class=K2Node_CallFunction Name="K2Node_CallFunction_64919"
+			   bIsPureFunc=True
+			   bIsConstFunc=True
+			   FunctionReference=(MemberName="K2_GetActorLocation",bSelfContext=True)
+			   NodePosX=560
+			   NodePosY=-191
+			   NodeGuid=B1D294AF40FA8CB679B6D4BED0F11D23
+			   CustomProperties Pin (PinId=136B71094B407B9186881DB137E50CB0,PinName="self",PinFriendlyName=NSLOCTEXT("K2Node", "Target", "Target"),PinType.PinCategory="object",PinType.PinSubCategoryObject=Class'/Script/Engine.Actor',PinType.PinSubCategoryMemberReference=(),PinType.bIsArray=False,PinType.bIsReference=False,PinType.bIsConst=False,PinType.bIsWeakPointer=False,LinkedTo=(K2Node_CallFunction_64853 88B4C49D414D557C8728F39A3967490F,),PersistentGuid=00000000000000000000000000000000,bHidden=False,bNotConnectable=False,bDefaultValueIsReadOnly=False,bDefaultValueIsIgnored=False,bAdvancedView=False,)
+			   CustomProperties Pin (PinId=7C3AB90A4E8F92F9FAE0709612FDD604,PinName="ReturnValue",Direction="EGPD_Output",PinType.PinCategory="struct",PinType.PinSubCategoryObject=ScriptStruct'/Script/CoreUObject.Vector',PinType.PinSubCategoryMemberReference=(),PinType.bIsArray=False,PinType.bIsReference=False,PinType.bIsConst=False,PinType.bIsWeakPointer=False,DefaultValue="0, 0, 0",AutogeneratedDefaultValue="0, 0, 0",LinkedTo=(K2Node_CallFunction_65598 85DD87AD4631A8DC79F17588BD6ADD07,K2Node_CommutativeAssociativeBinaryOperator_1357 C2F85EED4A8F818C3CE355A3C56691F5,),PersistentGuid=00000000000000000000000000000000,bHidden=False,bNotConnectable=False,bDefaultValueIsReadOnly=False,bDefaultValueIsIgnored=False,bAdvancedView=False,)
+			End Object
+			Begin Object Class=K2Node_CallFunction Name="K2Node_CallFunction_65094"
+			   bIsPureFunc=True
+			   bIsConstFunc=True
+			   FunctionReference=(MemberName="K2_GetActorRotation",bSelfContext=True)
+			   NodePosX=176
+			   NodePosY=-191
+			   NodeGuid=B7237EBA447AD7195212088A6AF51A87
+			   CustomProperties Pin (PinId=451585FC4960B7F55857FDB8A5A37392,PinName="self",PinFriendlyName=NSLOCTEXT("K2Node", "Target", "Target"),PinType.PinCategory="object",PinType.PinSubCategoryObject=Class'/Script/Engine.Actor',PinType.PinSubCategoryMemberReference=(),PinType.bIsArray=False,PinType.bIsReference=False,PinType.bIsConst=False,PinType.bIsWeakPointer=False,LinkedTo=(K2Node_CallFunction_64853 88B4C49D414D557C8728F39A3967490F,),PersistentGuid=00000000000000000000000000000000,bHidden=False,bNotConnectable=False,bDefaultValueIsReadOnly=False,bDefaultValueIsIgnored=False,bAdvancedView=False,)
+			   CustomProperties Pin (PinId=3048B8A64F99A2395A2A78B01F534663,PinName="ReturnValue",Direction="EGPD_Output",PinType.PinCategory="struct",PinType.PinSubCategoryObject=ScriptStruct'/Script/CoreUObject.Rotator',PinType.PinSubCategoryMemberReference=(),PinType.bIsArray=False,PinType.bIsReference=False,PinType.bIsConst=False,PinType.bIsWeakPointer=False,DefaultValue="0, 0, 0",AutogeneratedDefaultValue="0, 0, 0",LinkedTo=(K2Node_CallFunction_65130 44D474914E7F590449BEE0850F8024CE,),PersistentGuid=00000000000000000000000000000000,bHidden=False,bNotConnectable=False,bDefaultValueIsReadOnly=False,bDefaultValueIsIgnored=False,bAdvancedView=False,)
+			End Object
+			Begin Object Class=K2Node_CallFunction Name="K2Node_CallFunction_65130"
+			   bIsPureFunc=True
+			   FunctionReference=(MemberParent=Class'/Script/Engine.KismetMathLibrary',MemberName="Conv_RotatorToVector")
+			   NodePosX=384
+			   NodePosY=-64
+			   NodeGuid=43732C9F46DD928BC6B78F92F47C8A12
+			   CustomProperties Pin (PinId=06AE3D444E2D6A5652B7C291690B1A65,PinName="self",PinFriendlyName=NSLOCTEXT("K2Node", "Target", "Target"),PinType.PinCategory="object",PinType.PinSubCategoryObject=Class'/Script/Engine.KismetMathLibrary',PinType.PinSubCategoryMemberReference=(),PinType.bIsArray=False,PinType.bIsReference=False,PinType.bIsConst=False,PinType.bIsWeakPointer=False,DefaultObject="/Script/Engine.Default__KismetMathLibrary",PersistentGuid=00000000000000000000000000000000,bHidden=True,bNotConnectable=False,bDefaultValueIsReadOnly=False,bDefaultValueIsIgnored=False,bAdvancedView=False,)
+			   CustomProperties Pin (PinId=44D474914E7F590449BEE0850F8024CE,PinName="InRot",PinType.PinCategory="struct",PinType.PinSubCategoryObject=ScriptStruct'/Script/CoreUObject.Rotator',PinType.PinSubCategoryMemberReference=(),PinType.bIsArray=False,PinType.bIsReference=False,PinType.bIsConst=False,PinType.bIsWeakPointer=False,DefaultValue="0, 0, 0",AutogeneratedDefaultValue="0, 0, 0",LinkedTo=(K2Node_CallFunction_65094 3048B8A64F99A2395A2A78B01F534663,),PersistentGuid=00000000000000000000000000000000,bHidden=False,bNotConnectable=False,bDefaultValueIsReadOnly=False,bDefaultValueIsIgnored=False,bAdvancedView=False,)
+			   CustomProperties Pin (PinId=6A4437574CCDB781A7E425AF26B29860,PinName="ReturnValue",Direction="EGPD_Output",PinType.PinCategory="struct",PinType.PinSubCategoryObject=ScriptStruct'/Script/CoreUObject.Vector',PinType.PinSubCategoryMemberReference=(),PinType.bIsArray=False,PinType.bIsReference=False,PinType.bIsConst=False,PinType.bIsWeakPointer=False,DefaultValue="0, 0, 0",AutogeneratedDefaultValue="0, 0, 0",LinkedTo=(K2Node_CallFunction_65459 CE2B40EF47E2476D00764EA442CB2F0D,),PersistentGuid=00000000000000000000000000000000,bHidden=False,bNotConnectable=False,bDefaultValueIsReadOnly=False,bDefaultValueIsIgnored=False,bAdvancedView=False,)
+			End Object
+			Begin Object Class=K2Node_CallFunction Name="K2Node_CallFunction_65349"
+			   bIsPureFunc=True
+			   FunctionReference=(MemberParent=Class'/Script/Engine.KismetMathLibrary',MemberName="Multiply_VectorFloat")
+			   NodePosX=784
+			   NodePosY=-16
+			   NodeGuid=9ADED0674817535FD7003DB28601EACF
+			   CustomProperties Pin (PinId=ABB067854B9FDB401D7AF2A71639E8A9,PinName="self",PinFriendlyName=NSLOCTEXT("K2Node", "Target", "Target"),PinType.PinCategory="object",PinType.PinSubCategoryObject=Class'/Script/Engine.KismetMathLibrary',PinType.PinSubCategoryMemberReference=(),PinType.bIsArray=False,PinType.bIsReference=False,PinType.bIsConst=False,PinType.bIsWeakPointer=False,DefaultObject="/Script/Engine.Default__KismetMathLibrary",PersistentGuid=00000000000000000000000000000000,bHidden=True,bNotConnectable=False,bDefaultValueIsReadOnly=False,bDefaultValueIsIgnored=False,bAdvancedView=False,)
+			   CustomProperties Pin (PinId=AE6D33C2476153AC065E9095CB766952,PinName="A",PinType.PinCategory="struct",PinType.PinSubCategoryObject=ScriptStruct'/Script/CoreUObject.Vector',PinType.PinSubCategoryMemberReference=(),PinType.bIsArray=False,PinType.bIsReference=False,PinType.bIsConst=False,PinType.bIsWeakPointer=False,DefaultValue="0, 0, 0",AutogeneratedDefaultValue="0, 0, 0",LinkedTo=(K2Node_CallFunction_65459 182E27F347849D96EFDA44BC8BAF228B,),PersistentGuid=00000000000000000000000000000000,bHidden=False,bNotConnectable=False,bDefaultValueIsReadOnly=False,bDefaultValueIsIgnored=False,bAdvancedView=False,)
+			   CustomProperties Pin (PinId=2C5718584516424BB52CF6A79E702026,PinName="B",PinType.PinCategory="float",PinType.PinSubCategoryObject=None,PinType.PinSubCategoryMemberReference=(),PinType.bIsArray=False,PinType.bIsReference=False,PinType.bIsConst=False,PinType.bIsWeakPointer=False,DefaultValue="0.0",AutogeneratedDefaultValue="0.0",LinkedTo=(K2Node_CommutativeAssociativeBinaryOperator_4227 0FAF44A34397D0F87EB02B88E8A653E3,),PersistentGuid=00000000000000000000000000000000,bHidden=False,bNotConnectable=False,bDefaultValueIsReadOnly=False,bDefaultValueIsIgnored=False,bAdvancedView=False,)
+			   CustomProperties Pin (PinId=40870E1F404EA24CCC3295A98D8A13B6,PinName="ReturnValue",Direction="EGPD_Output",PinType.PinCategory="struct",PinType.PinSubCategoryObject=ScriptStruct'/Script/CoreUObject.Vector',PinType.PinSubCategoryMemberReference=(),PinType.bIsArray=False,PinType.bIsReference=False,PinType.bIsConst=False,PinType.bIsWeakPointer=False,DefaultValue="0, 0, 0",AutogeneratedDefaultValue="0, 0, 0",LinkedTo=(K2Node_CommutativeAssociativeBinaryOperator_1357 4410C89D42A350AA1AF9EE830D6A174E,),PersistentGuid=00000000000000000000000000000000,bHidden=False,bNotConnectable=False,bDefaultValueIsReadOnly=False,bDefaultValueIsIgnored=False,bAdvancedView=False,)
+			End Object
+			Begin Object Class=K2Node_CallFunction Name="K2Node_CallFunction_65459"
+			   bIsPureFunc=True
+			   FunctionReference=(MemberParent=Class'/Script/Engine.KismetMathLibrary',MemberName="Normal")
+			   NodePosX=592
+			   NodePosY=-64
+			   NodeGuid=8F88CD2F4F7D9006417C01B2332B7C39
+			   CustomProperties Pin (PinId=E124F3A24616014EE593778FBFC7A74F,PinName="self",PinFriendlyName=NSLOCTEXT("K2Node", "Target", "Target"),PinToolTip="Target\nKismet Math Library Reference",PinType.PinCategory="object",PinType.PinSubCategoryObject=Class'/Script/Engine.KismetMathLibrary',PinType.PinSubCategoryMemberReference=(),PinType.bIsArray=False,PinType.bIsReference=False,PinType.bIsConst=False,PinType.bIsWeakPointer=False,DefaultObject="/Script/Engine.Default__KismetMathLibrary",PersistentGuid=00000000000000000000000000000000,bHidden=True,bNotConnectable=False,bDefaultValueIsReadOnly=False,bDefaultValueIsIgnored=False,bAdvancedView=False,)
+			   CustomProperties Pin (PinId=CE2B40EF47E2476D00764EA442CB2F0D,PinName="A",PinToolTip="A\nVector ",PinType.PinCategory="struct",PinType.PinSubCategoryObject=ScriptStruct'/Script/CoreUObject.Vector',PinType.PinSubCategoryMemberReference=(),PinType.bIsArray=False,PinType.bIsReference=False,PinType.bIsConst=False,PinType.bIsWeakPointer=False,DefaultValue="0, 0, 0",AutogeneratedDefaultValue="0, 0, 0",LinkedTo=(K2Node_CallFunction_65130 6A4437574CCDB781A7E425AF26B29860,),PersistentGuid=00000000000000000000000000000000,bHidden=False,bNotConnectable=False,bDefaultValueIsReadOnly=False,bDefaultValueIsIgnored=False,bAdvancedView=False,)
+			   CustomProperties Pin (PinId=182E27F347849D96EFDA44BC8BAF228B,PinName="ReturnValue",PinToolTip="Return Value\nVector ",Direction="EGPD_Output",PinType.PinCategory="struct",PinType.PinSubCategoryObject=ScriptStruct'/Script/CoreUObject.Vector',PinType.PinSubCategoryMemberReference=(),PinType.bIsArray=False,PinType.bIsReference=False,PinType.bIsConst=False,PinType.bIsWeakPointer=False,DefaultValue="0, 0, 0",AutogeneratedDefaultValue="0, 0, 0",LinkedTo=(K2Node_CallFunction_65349 AE6D33C2476153AC065E9095CB766952,),PersistentGuid=00000000000000000000000000000000,bHidden=False,bNotConnectable=False,bDefaultValueIsReadOnly=False,bDefaultValueIsIgnored=False,bAdvancedView=False,)
+			End Object
+			Begin Object Class=K2Node_CallFunction Name="K2Node_CallFunction_65598"
+			   FunctionReference=(MemberParent=Class'/Script/Engine.KismetSystemLibrary',MemberName="LineTraceSingle_NEW")
+			   NodePosX=1104
+			   NodePosY=-384
+			   NodeGuid=8BE120E14A1ACA5DD4C51B809102B477
+			   AdvancedPinDisplay=Hidden
+			   CustomProperties Pin (PinId=B3AD098C4E9F315318BE2C9A781CFD16,PinName="execute",PinToolTip="\nExec",PinType.PinCategory="exec",PinType.PinSubCategoryObject=None,PinType.PinSubCategoryMemberReference=(),PinType.bIsArray=False,PinType.bIsReference=False,PinType.bIsConst=False,PinType.bIsWeakPointer=False,LinkedTo=(K2Node_FunctionEntry_9 598C3D6A46B70A76463B448AA52DE23A,),PersistentGuid=00000000000000000000000000000000,bHidden=False,bNotConnectable=False,bDefaultValueIsReadOnly=False,bDefaultValueIsIgnored=False,bAdvancedView=False,)
+			   CustomProperties Pin (PinId=E001912C48BB5BEC656460843FE151FA,PinName="then",PinToolTip="\nExec",Direction="EGPD_Output",PinType.PinCategory="exec",PinType.PinSubCategoryObject=None,PinType.PinSubCategoryMemberReference=(),PinType.bIsArray=False,PinType.bIsReference=False,PinType.bIsConst=False,PinType.bIsWeakPointer=False,LinkedTo=(K2Node_IfThenElse_161 66D340DC4A52DC8D727C43ABA6E818C4,),PersistentGuid=00000000000000000000000000000000,bHidden=False,bNotConnectable=False,bDefaultValueIsReadOnly=False,bDefaultValueIsIgnored=False,bAdvancedView=False,)
+			   CustomProperties Pin (PinId=941FA0A3430F4462DB5E30B8A5A81DF4,PinName="self",PinFriendlyName=NSLOCTEXT("K2Node", "Target", "Target"),PinToolTip="Target\nKismet System Library Reference",PinType.PinCategory="object",PinType.PinSubCategoryObject=Class'/Script/Engine.KismetSystemLibrary',PinType.PinSubCategoryMemberReference=(),PinType.bIsArray=False,PinType.bIsReference=False,PinType.bIsConst=False,PinType.bIsWeakPointer=False,DefaultObject="/Script/Engine.Default__KismetSystemLibrary",PersistentGuid=00000000000000000000000000000000,bHidden=True,bNotConnectable=False,bDefaultValueIsReadOnly=False,bDefaultValueIsIgnored=False,bAdvancedView=False,)
+			   CustomProperties Pin (PinId=6F242C374EDACE5FC62C50B118086553,PinName="WorldContextObject",PinToolTip="World Context Object\nObject Reference",PinType.PinCategory="object",PinType.PinSubCategoryObject=Class'/Script/CoreUObject.Object',PinType.PinSubCategoryMemberReference=(),PinType.bIsArray=False,PinType.bIsReference=False,PinType.bIsConst=False,PinType.bIsWeakPointer=False,DefaultValue="0",PersistentGuid=00000000000000000000000000000000,bHidden=True,bNotConnectable=False,bDefaultValueIsReadOnly=False,bDefaultValueIsIgnored=False,bAdvancedView=False,)
+			   CustomProperties Pin (PinId=85DD87AD4631A8DC79F17588BD6ADD07,PinName="Start",PinToolTip="Start\nVector \n\nStart of line segment.",PinType.PinCategory="struct",PinType.PinSubCategoryObject=ScriptStruct'/Script/CoreUObject.Vector',PinType.PinSubCategoryMemberReference=(),PinType.bIsArray=False,PinType.bIsReference=False,PinType.bIsConst=True,PinType.bIsWeakPointer=False,DefaultValue="0, 0, 0",AutogeneratedDefaultValue="0, 0, 0",LinkedTo=(K2Node_CallFunction_64919 7C3AB90A4E8F92F9FAE0709612FDD604,),PersistentGuid=00000000000000000000000000000000,bHidden=False,bNotConnectable=False,bDefaultValueIsReadOnly=False,bDefaultValueIsIgnored=False,bAdvancedView=False,)
+			   CustomProperties Pin (PinId=54C3C6D04942C8BC62F2D696DFF0A12C,PinName="End",PinToolTip="End\nVector \n\nEnd of line segment.",PinType.PinCategory="struct",PinType.PinSubCategoryObject=ScriptStruct'/Script/CoreUObject.Vector',PinType.PinSubCategoryMemberReference=(),PinType.bIsArray=False,PinType.bIsReference=False,PinType.bIsConst=True,PinType.bIsWeakPointer=False,DefaultValue="0, 0, 0",AutogeneratedDefaultValue="0, 0, 0",LinkedTo=(K2Node_CommutativeAssociativeBinaryOperator_1357 1778AF2B4812E5B635D0DC915DFAA735,),PersistentGuid=00000000000000000000000000000000,bHidden=False,bNotConnectable=False,bDefaultValueIsReadOnly=False,bDefaultValueIsIgnored=False,bAdvancedView=False,)
+			   CustomProperties Pin (PinId=FAD144344A968D69191168B056747AB6,PinName="TraceChannel",PinToolTip="Trace Channel\nETraceTypeQuery Enum",PinType.PinCategory="byte",PinType.PinSubCategoryObject=Enum'/Script/Engine.ETraceTypeQuery',PinType.PinSubCategoryMemberReference=(),PinType.bIsArray=False,PinType.bIsReference=False,PinType.bIsConst=False,PinType.bIsWeakPointer=False,DefaultValue="TraceTypeQuery1",PersistentGuid=00000000000000000000000000000000,bHidden=False,bNotConnectable=False,bDefaultValueIsReadOnly=False,bDefaultValueIsIgnored=False,bAdvancedView=False,)
+			   CustomProperties Pin (PinId=8B576AA34608DE6F0B60328E9D9C1166,PinName="bTraceComplex",PinToolTip="Trace Complex\nBoolean\n\nTrue to test against complex collision, false to test against simplified collision.",PinType.PinCategory="bool",PinType.PinSubCategoryObject=None,PinType.PinSubCategoryMemberReference=(),PinType.bIsArray=False,PinType.bIsReference=False,PinType.bIsConst=False,PinType.bIsWeakPointer=False,DefaultValue="false",AutogeneratedDefaultValue="false",PersistentGuid=00000000000000000000000000000000,bHidden=False,bNotConnectable=False,bDefaultValueIsReadOnly=False,bDefaultValueIsIgnored=False,bAdvancedView=False,)
+			   CustomProperties Pin (PinId=E4B43C3146C4E49A9F72ECBAF170A48B,PinName="ActorsToIgnore",PinToolTip="Actors to Ignore\nArray of Actor References",PinType.PinCategory="object",PinType.PinSubCategoryObject=Class'/Script/Engine.Actor',PinType.PinSubCategoryMemberReference=(),PinType.bIsArray=True,PinType.bIsReference=True,PinType.bIsConst=True,PinType.bIsWeakPointer=False,PersistentGuid=00000000000000000000000000000000,bHidden=False,bNotConnectable=False,bDefaultValueIsReadOnly=False,bDefaultValueIsIgnored=True,bAdvancedView=False,)
+			   CustomProperties Pin (PinId=6DC8450447A07CB69FEA6ABBD6C12BED,PinName="DrawDebugType",PinToolTip="Draw Debug Type\nEDrawDebugTrace Enum",PinType.PinCategory="byte",PinType.PinSubCategoryObject=Enum'/Script/Engine.EDrawDebugTrace',PinType.PinSubCategoryMemberReference=(),PinType.bIsArray=False,PinType.bIsReference=False,PinType.bIsConst=False,PinType.bIsWeakPointer=False,DefaultValue="None",PersistentGuid=00000000000000000000000000000000,bHidden=False,bNotConnectable=False,bDefaultValueIsReadOnly=False,bDefaultValueIsIgnored=False,bAdvancedView=False,)
+			   CustomProperties Pin (PinId=43A3FD1943054F5E9C27A298298CE721,PinName="OutHit",PinToolTip="Out Hit\nHit Result Structure\n\nProperties of the trace hit.",Direction="EGPD_Output",PinType.PinCategory="struct",PinType.PinSubCategoryObject=ScriptStruct'/Script/Engine.HitResult',PinType.PinSubCategoryMemberReference=(),PinType.bIsArray=False,PinType.bIsReference=False,PinType.bIsConst=False,PinType.bIsWeakPointer=False,LinkedTo=(K2Node_CallFunction_65759 9BE54D1545016821EA539F83438E9077,K2Node_CallFunction_1861 7453F08741D187AE2B9629985E33A634,),PersistentGuid=00000000000000000000000000000000,bHidden=False,bNotConnectable=False,bDefaultValueIsReadOnly=False,bDefaultValueIsIgnored=False,bAdvancedView=False,)
+			   CustomProperties Pin (PinId=D4D0981940265BE6F98E8B9F2CB3C064,PinName="bIgnoreSelf",PinToolTip="Ignore Self\nBoolean",PinType.PinCategory="bool",PinType.PinSubCategoryObject=None,PinType.PinSubCategoryMemberReference=(),PinType.bIsArray=False,PinType.bIsReference=False,PinType.bIsConst=False,PinType.bIsWeakPointer=False,DefaultValue="false",AutogeneratedDefaultValue="true",PersistentGuid=00000000000000000000000000000000,bHidden=False,bNotConnectable=False,bDefaultValueIsReadOnly=False,bDefaultValueIsIgnored=False,bAdvancedView=False,)
+			   CustomProperties Pin (PinId=17643C41432C7B66D6FF7F9A6A4BA258,PinName="TraceColor",PinToolTip="Trace Color\nLinear Color Structure",PinType.PinCategory="struct",PinType.PinSubCategoryObject=ScriptStruct'/Script/CoreUObject.LinearColor',PinType.PinSubCategoryMemberReference=(),PinType.bIsArray=False,PinType.bIsReference=False,PinType.bIsConst=False,PinType.bIsWeakPointer=False,DefaultValue="(R=1.000000,G=0.000000,B=0.000000,A=1.000000)",AutogeneratedDefaultValue="(R=1.000000,G=0.000000,B=0.000000,A=1.000000)",PersistentGuid=00000000000000000000000000000000,bHidden=False,bNotConnectable=False,bDefaultValueIsReadOnly=False,bDefaultValueIsIgnored=False,bAdvancedView=True,)
+			   CustomProperties Pin (PinId=F62FE7C840096CE5353C2C98BF138F9C,PinName="TraceHitColor",PinToolTip="Trace Hit Color\nLinear Color Structure",PinType.PinCategory="struct",PinType.PinSubCategoryObject=ScriptStruct'/Script/CoreUObject.LinearColor',PinType.PinSubCategoryMemberReference=(),PinType.bIsArray=False,PinType.bIsReference=False,PinType.bIsConst=False,PinType.bIsWeakPointer=False,DefaultValue="(R=0.000000,G=1.000000,B=0.000000,A=1.000000)",AutogeneratedDefaultValue="(R=0.000000,G=1.000000,B=0.000000,A=1.000000)",PersistentGuid=00000000000000000000000000000000,bHidden=False,bNotConnectable=False,bDefaultValueIsReadOnly=False,bDefaultValueIsIgnored=False,bAdvancedView=True,)
+			   CustomProperties Pin (PinId=9A93C66A42BC103DD20019BCA763A17F,PinName="DrawTime",PinToolTip="Draw Time\nFloat",PinType.PinCategory="float",PinType.PinSubCategoryObject=None,PinType.PinSubCategoryMemberReference=(),PinType.bIsArray=False,PinType.bIsReference=False,PinType.bIsConst=False,PinType.bIsWeakPointer=False,DefaultValue="5.000000",AutogeneratedDefaultValue="5.000000",PersistentGuid=00000000000000000000000000000000,bHidden=False,bNotConnectable=False,bDefaultValueIsReadOnly=False,bDefaultValueIsIgnored=False,bAdvancedView=True,)
+			   CustomProperties Pin (PinId=3A64487743C81996C17E56B543685516,PinName="ReturnValue",PinToolTip="Return Value\nBoolean\n\nTrue if there was a hit, false otherwise.",Direction="EGPD_Output",PinType.PinCategory="bool",PinType.PinSubCategoryObject=None,PinType.PinSubCategoryMemberReference=(),PinType.bIsArray=False,PinType.bIsReference=False,PinType.bIsConst=False,PinType.bIsWeakPointer=False,DefaultValue="false",AutogeneratedDefaultValue="false",LinkedTo=(K2Node_IfThenElse_161 1144399140FB3DC67E1A66B5FEB4A1E1,),PersistentGuid=00000000000000000000000000000000,bHidden=False,bNotConnectable=False,bDefaultValueIsReadOnly=False,bDefaultValueIsIgnored=False,bAdvancedView=False,)
+			End Object
+			Begin Object Class=K2Node_CommutativeAssociativeBinaryOperator Name="K2Node_CommutativeAssociativeBinaryOperator_1357"
+			   bIsPureFunc=True
+			   FunctionReference=(MemberParent=Class'/Script/Engine.KismetMathLibrary',MemberName="Add_VectorVector")
+			   NodePosX=912
+			   NodePosY=-96
+			   NodeGuid=8DC37CA3400BFCC550BCC1BE62581E73
+			   CustomProperties Pin (PinId=B729CFBC424DA01AC5E021B25FE0F7CD,PinName="self",PinFriendlyName=NSLOCTEXT("K2Node", "Target", "Target"),PinType.PinCategory="object",PinType.PinSubCategoryObject=Class'/Script/Engine.KismetMathLibrary',PinType.PinSubCategoryMemberReference=(),PinType.bIsArray=False,PinType.bIsReference=False,PinType.bIsConst=False,PinType.bIsWeakPointer=False,DefaultObject="/Script/Engine.Default__KismetMathLibrary",PersistentGuid=00000000000000000000000000000000,bHidden=True,bNotConnectable=False,bDefaultValueIsReadOnly=False,bDefaultValueIsIgnored=False,bAdvancedView=False,)
+			   CustomProperties Pin (PinId=4410C89D42A350AA1AF9EE830D6A174E,PinName="A",PinType.PinCategory="struct",PinType.PinSubCategoryObject=ScriptStruct'/Script/CoreUObject.Vector',PinType.PinSubCategoryMemberReference=(),PinType.bIsArray=False,PinType.bIsReference=False,PinType.bIsConst=False,PinType.bIsWeakPointer=False,DefaultValue="0, 0, 0",AutogeneratedDefaultValue="0, 0, 0",LinkedTo=(K2Node_CallFunction_65349 40870E1F404EA24CCC3295A98D8A13B6,),PersistentGuid=00000000000000000000000000000000,bHidden=False,bNotConnectable=False,bDefaultValueIsReadOnly=False,bDefaultValueIsIgnored=False,bAdvancedView=False,)
+			   CustomProperties Pin (PinId=C2F85EED4A8F818C3CE355A3C56691F5,PinName="B",PinType.PinCategory="struct",PinType.PinSubCategoryObject=ScriptStruct'/Script/CoreUObject.Vector',PinType.PinSubCategoryMemberReference=(),PinType.bIsArray=False,PinType.bIsReference=False,PinType.bIsConst=False,PinType.bIsWeakPointer=False,DefaultValue="0, 0, 0",AutogeneratedDefaultValue="0, 0, 0",LinkedTo=(K2Node_CallFunction_64919 7C3AB90A4E8F92F9FAE0709612FDD604,),PersistentGuid=00000000000000000000000000000000,bHidden=False,bNotConnectable=False,bDefaultValueIsReadOnly=False,bDefaultValueIsIgnored=False,bAdvancedView=False,)
+			   CustomProperties Pin (PinId=1778AF2B4812E5B635D0DC915DFAA735,PinName="ReturnValue",Direction="EGPD_Output",PinType.PinCategory="struct",PinType.PinSubCategoryObject=ScriptStruct'/Script/CoreUObject.Vector',PinType.PinSubCategoryMemberReference=(),PinType.bIsArray=False,PinType.bIsReference=False,PinType.bIsConst=False,PinType.bIsWeakPointer=False,AutogeneratedDefaultValue="0, 0, 0",LinkedTo=(K2Node_CallFunction_65598 54C3C6D04942C8BC62F2D696DFF0A12C,),PersistentGuid=00000000000000000000000000000000,bHidden=False,bNotConnectable=False,bDefaultValueIsReadOnly=False,bDefaultValueIsIgnored=False,bAdvancedView=False,)
+			End Object
+			Begin Object Class=K2Node_CallFunction Name="K2Node_CallFunction_65759"
+			   bIsPureFunc=True
+			   FunctionReference=(MemberParent=Class'/Script/Engine.GameplayStatics',MemberName="BreakHitResult")
+			   NodePosX=1568
+			   NodePosY=-288
+			   NodeGuid=92BF07174B1E8AA8AC947F9AF1A1E6E7
+			   CustomProperties Pin (PinId=CF78D28A455F553D39502496BC89B70E,PinName="self",PinFriendlyName=NSLOCTEXT("K2Node", "Target", "Target"),PinToolTip="Target\nGameplay Statics Reference",PinType.PinCategory="object",PinType.PinSubCategoryObject=Class'/Script/Engine.GameplayStatics',PinType.PinSubCategoryMemberReference=(),PinType.bIsArray=False,PinType.bIsReference=False,PinType.bIsConst=False,PinType.bIsWeakPointer=False,DefaultObject="/Script/Engine.Default__GameplayStatics",PersistentGuid=00000000000000000000000000000000,bHidden=True,bNotConnectable=False,bDefaultValueIsReadOnly=False,bDefaultValueIsIgnored=False,bAdvancedView=False,)
+			   CustomProperties Pin (PinId=9BE54D1545016821EA539F83438E9077,PinName="Hit",PinToolTip="Hit\nHit Result Structure (by ref)\n\nThe source HitResult.",PinType.PinCategory="struct",PinType.PinSubCategoryObject=ScriptStruct'/Script/Engine.HitResult',PinType.PinSubCategoryMemberReference=(),PinType.bIsArray=False,PinType.bIsReference=True,PinType.bIsConst=True,PinType.bIsWeakPointer=False,LinkedTo=(K2Node_CallFunction_65598 43A3FD1943054F5E9C27A298298CE721,),PersistentGuid=00000000000000000000000000000000,bHidden=False,bNotConnectable=False,bDefaultValueIsReadOnly=False,bDefaultValueIsIgnored=True,bAdvancedView=False,)
+			   CustomProperties Pin (PinId=42FAE549473540FC07A70799D95A29B3,PinName="bBlockingHit",PinToolTip="Blocking Hit\nBoolean\n\nTrue if there was a blocking hit, false otherwise.",Direction="EGPD_Output",PinType.PinCategory="bool",PinType.PinSubCategoryObject=None,PinType.PinSubCategoryMemberReference=(),PinType.bIsArray=False,PinType.bIsReference=False,PinType.bIsConst=False,PinType.bIsWeakPointer=False,DefaultValue="false",AutogeneratedDefaultValue="false",PersistentGuid=00000000000000000000000000000000,bHidden=False,bNotConnectable=False,bDefaultValueIsReadOnly=False,bDefaultValueIsIgnored=False,bAdvancedView=False,)
+			   CustomProperties Pin (PinId=31399B19487BF09D8EF1088E8ED56118,PinName="bInitialOverlap",PinToolTip="Initial Overlap\nBoolean\n\nTrue if the hit started in an initial overlap. In this case some other values should be interpreted differently. Time will be 0, ImpactPoint will equal Location, and normals will be equal and indicate a depenetration vector.",Direction="EGPD_Output",PinType.PinCategory="bool",PinType.PinSubCategoryObject=None,PinType.PinSubCategoryMemberReference=(),PinType.bIsArray=False,PinType.bIsReference=False,PinType.bIsConst=False,PinType.bIsWeakPointer=False,DefaultValue="false",AutogeneratedDefaultValue="false",PersistentGuid=00000000000000000000000000000000,bHidden=False,bNotConnectable=False,bDefaultValueIsReadOnly=False,bDefaultValueIsIgnored=False,bAdvancedView=False,)
+			   CustomProperties Pin (PinId=4312416249D19F19294CADA26CDE64E5,PinName="Time",PinToolTip="Time\nFloat\n\n\'Time\' of impact along trace direction ranging from [0.0 to 1.0) if there is a hit, indicating time between start and end. Equals 1.0 if there is no hit.",Direction="EGPD_Output",PinType.PinCategory="float",PinType.PinSubCategoryObject=None,PinType.PinSubCategoryMemberReference=(),PinType.bIsArray=False,PinType.bIsReference=False,PinType.bIsConst=False,PinType.bIsWeakPointer=False,DefaultValue="0.0",AutogeneratedDefaultValue="0.0",PersistentGuid=00000000000000000000000000000000,bHidden=False,bNotConnectable=False,bDefaultValueIsReadOnly=False,bDefaultValueIsIgnored=False,bAdvancedView=False,)
+			   CustomProperties Pin (PinId=C3E1DE9D4B514E4FCDA1159FEA7E7832,PinName="Location",PinToolTip="Location\nVector \n\nLocation of the hit in world space. If this was a swept shape test, this is the location where we can place the shape in the world where it will not penetrate.",Direction="EGPD_Output",PinType.PinCategory="struct",PinType.PinSubCategoryObject=ScriptStruct'/Script/CoreUObject.Vector',PinType.PinSubCategoryMemberReference=(),PinType.bIsArray=False,PinType.bIsReference=False,PinType.bIsConst=False,PinType.bIsWeakPointer=False,AutogeneratedDefaultValue="0, 0, 0",PersistentGuid=00000000000000000000000000000000,bHidden=False,bNotConnectable=False,bDefaultValueIsReadOnly=False,bDefaultValueIsIgnored=False,bAdvancedView=False,)
+			   CustomProperties Pin (PinId=EB9F6AED44EDCC7D8F313FB63A3B4139,PinName="ImpactPoint",PinToolTip="Impact Point\nVector \n\nLocation of the actual contact point of the trace shape with the surface of the hit object. Equal to Location in the case of an initial overlap.",Direction="EGPD_Output",PinType.PinCategory="struct",PinType.PinSubCategoryObject=ScriptStruct'/Script/CoreUObject.Vector',PinType.PinSubCategoryMemberReference=(),PinType.bIsArray=False,PinType.bIsReference=False,PinType.bIsConst=False,PinType.bIsWeakPointer=False,AutogeneratedDefaultValue="0, 0, 0",PersistentGuid=00000000000000000000000000000000,bHidden=False,bNotConnectable=False,bDefaultValueIsReadOnly=False,bDefaultValueIsIgnored=False,bAdvancedView=False,)
+			   CustomProperties Pin (PinId=AC37F9C247A40EF82D1755B03A6EF0AF,PinName="Normal",PinToolTip="Normal\nVector \n\nNormal of the hit in world space, for the object that was swept (e.g. for a sphere trace this points towards the sphere\'s center). Equal to ImpactNormal for line tests.",Direction="EGPD_Output",PinType.PinCategory="struct",PinType.PinSubCategoryObject=ScriptStruct'/Script/CoreUObject.Vector',PinType.PinSubCategoryMemberReference=(),PinType.bIsArray=False,PinType.bIsReference=False,PinType.bIsConst=False,PinType.bIsWeakPointer=False,AutogeneratedDefaultValue="0, 0, 0",PersistentGuid=00000000000000000000000000000000,bHidden=False,bNotConnectable=False,bDefaultValueIsReadOnly=False,bDefaultValueIsIgnored=False,bAdvancedView=False,)
+			   CustomProperties Pin (PinId=43DE840740F7EDB5AEEE2BA59BA63807,PinName="ImpactNormal",PinToolTip="Impact Normal\nVector \n\nNormal of the hit in world space, for the object that was hit by the sweep.",Direction="EGPD_Output",PinType.PinCategory="struct",PinType.PinSubCategoryObject=ScriptStruct'/Script/CoreUObject.Vector',PinType.PinSubCategoryMemberReference=(),PinType.bIsArray=False,PinType.bIsReference=False,PinType.bIsConst=False,PinType.bIsWeakPointer=False,AutogeneratedDefaultValue="0, 0, 0",PersistentGuid=00000000000000000000000000000000,bHidden=False,bNotConnectable=False,bDefaultValueIsReadOnly=False,bDefaultValueIsIgnored=False,bAdvancedView=False,)
+			   CustomProperties Pin (PinId=532D29F44B5DCE5F021B698ED58D27C4,PinName="PhysMat",PinToolTip="Phys Mat\nPhysical Material Reference\n\nPhysical material that was hit. Must set bReturnPhysicalMaterial to true in the query params for this to be returned.",Direction="EGPD_Output",PinType.PinCategory="object",PinType.PinSubCategoryObject=Class'/Script/Engine.PhysicalMaterial',PinType.PinSubCategoryMemberReference=(),PinType.bIsArray=False,PinType.bIsReference=False,PinType.bIsConst=False,PinType.bIsWeakPointer=False,PersistentGuid=00000000000000000000000000000000,bHidden=False,bNotConnectable=False,bDefaultValueIsReadOnly=False,bDefaultValueIsIgnored=False,bAdvancedView=False,)
+			   CustomProperties Pin (PinId=F1D6985542953219C7EF09A3FBFA8AF0,PinName="HitActor",PinToolTip="Hit Actor\nActor Reference\n\nActor hit by the trace.",Direction="EGPD_Output",PinType.PinCategory="object",PinType.PinSubCategoryObject=Class'/Script/Engine.Actor',PinType.PinSubCategoryMemberReference=(),PinType.bIsArray=False,PinType.bIsReference=False,PinType.bIsConst=False,PinType.bIsWeakPointer=False,LinkedTo=(K2Node_MacroInstance_65 C3FB28AC47AAD2B5998C2AA9D4862C2E,K2Node_CallFunction_1861 BB747167446E2DDBEFFFDD982A9CD245,),PersistentGuid=00000000000000000000000000000000,bHidden=False,bNotConnectable=False,bDefaultValueIsReadOnly=False,bDefaultValueIsIgnored=False,bAdvancedView=False,)
+			   CustomProperties Pin (PinId=FDFC05184366EBC9BDC42AA3A95EF523,PinName="HitComponent",PinToolTip="Hit Component\nPrimitive Component Reference\n\nPrimitiveComponent hit by the trace.",Direction="EGPD_Output",PinType.PinCategory="object",PinType.PinSubCategoryObject=Class'/Script/Engine.PrimitiveComponent',PinType.PinSubCategoryMemberReference=(),PinType.bIsArray=False,PinType.bIsReference=False,PinType.bIsConst=False,PinType.bIsWeakPointer=False,PersistentGuid=00000000000000000000000000000000,bHidden=False,bNotConnectable=False,bDefaultValueIsReadOnly=False,bDefaultValueIsIgnored=False,bAdvancedView=False,)
+			   CustomProperties Pin (PinId=CE3689BE4446FD1C971B7A9942E7AC24,PinName="HitBoneName",PinToolTip="Hit Bone Name\nName\n\nName of the bone hit (valid only if we hit a skeletal mesh).",Direction="EGPD_Output",PinType.PinCategory="name",PinType.PinSubCategoryObject=None,PinType.PinSubCategoryMemberReference=(),PinType.bIsArray=False,PinType.bIsReference=False,PinType.bIsConst=False,PinType.bIsWeakPointer=False,DefaultValue="None",AutogeneratedDefaultValue="None",PersistentGuid=00000000000000000000000000000000,bHidden=False,bNotConnectable=False,bDefaultValueIsReadOnly=False,bDefaultValueIsIgnored=False,bAdvancedView=False,)
+			   CustomProperties Pin (PinId=1E9AAE5D433101E315ABFCB6EE0C3F4E,PinName="HitItem",PinToolTip="Hit Item\nInteger\n\nPrimitive-specific data recording which item in the primitive was hit",Direction="EGPD_Output",PinType.PinCategory="int",PinType.PinSubCategoryObject=None,PinType.PinSubCategoryMemberReference=(),PinType.bIsArray=False,PinType.bIsReference=False,PinType.bIsConst=False,PinType.bIsWeakPointer=False,DefaultValue="0",AutogeneratedDefaultValue="0",PersistentGuid=00000000000000000000000000000000,bHidden=False,bNotConnectable=False,bDefaultValueIsReadOnly=False,bDefaultValueIsIgnored=False,bAdvancedView=False,)
+			   CustomProperties Pin (PinId=FA0E13FF4A7001F230C36C84199334C5,PinName="FaceIndex",PinToolTip="Face Index\nInteger\n\nIf colliding with trimesh or landscape, index of face that was hit.",Direction="EGPD_Output",PinType.PinCategory="int",PinType.PinSubCategoryObject=None,PinType.PinSubCategoryMemberReference=(),PinType.bIsArray=False,PinType.bIsReference=False,PinType.bIsConst=False,PinType.bIsWeakPointer=False,DefaultValue="0",AutogeneratedDefaultValue="0",PersistentGuid=00000000000000000000000000000000,bHidden=False,bNotConnectable=False,bDefaultValueIsReadOnly=False,bDefaultValueIsIgnored=False,bAdvancedView=False,)
+			   CustomProperties Pin (PinId=6AFDED3441A726D6362809827812A060,PinName="TraceStart",PinToolTip="Trace Start\nVector ",Direction="EGPD_Output",PinType.PinCategory="struct",PinType.PinSubCategoryObject=ScriptStruct'/Script/CoreUObject.Vector',PinType.PinSubCategoryMemberReference=(),PinType.bIsArray=False,PinType.bIsReference=False,PinType.bIsConst=False,PinType.bIsWeakPointer=False,DefaultValue="0, 0, 0",AutogeneratedDefaultValue="0, 0, 0",PersistentGuid=00000000000000000000000000000000,bHidden=False,bNotConnectable=False,bDefaultValueIsReadOnly=False,bDefaultValueIsIgnored=False,bAdvancedView=False,)
+			   CustomProperties Pin (PinId=96D907364CB400491CEFF4B6B3921F92,PinName="TraceEnd",PinToolTip="Trace End\nVector ",Direction="EGPD_Output",PinType.PinCategory="struct",PinType.PinSubCategoryObject=ScriptStruct'/Script/CoreUObject.Vector',PinType.PinSubCategoryMemberReference=(),PinType.bIsArray=False,PinType.bIsReference=False,PinType.bIsConst=False,PinType.bIsWeakPointer=False,DefaultValue="0, 0, 0",AutogeneratedDefaultValue="0, 0, 0",PersistentGuid=00000000000000000000000000000000,bHidden=False,bNotConnectable=False,bDefaultValueIsReadOnly=False,bDefaultValueIsIgnored=False,bAdvancedView=False,)
+			End Object
+			Begin Object Class=K2Node_CommutativeAssociativeBinaryOperator Name="K2Node_CommutativeAssociativeBinaryOperator_4227"
+			   bIsPureFunc=True
+			   FunctionReference=(MemberParent=Class'/Script/Engine.KismetMathLibrary',MemberName="Multiply_FloatFloat")
+			   NodePosX=592
+			   NodePosY=16
+			   NodeGuid=E816C87A41EA136F98171CAC80832A2C
+			   CustomProperties Pin (PinId=0D4F3C2044C7E8EA342002AF074BCB55,PinName="self",PinFriendlyName=NSLOCTEXT("K2Node", "Target", "Target"),PinType.PinCategory="object",PinType.PinSubCategoryObject=Class'/Script/Engine.KismetMathLibrary',PinType.PinSubCategoryMemberReference=(),PinType.bIsArray=False,PinType.bIsReference=False,PinType.bIsConst=False,PinType.bIsWeakPointer=False,DefaultObject="/Script/Engine.Default__KismetMathLibrary",PersistentGuid=00000000000000000000000000000000,bHidden=True,bNotConnectable=False,bDefaultValueIsReadOnly=False,bDefaultValueIsIgnored=False,bAdvancedView=False,)
+			   CustomProperties Pin (PinId=AE00303B4D5F38B1B48C718ABEC37EF0,PinName="A",PinType.PinCategory="float",PinType.PinSubCategoryObject=None,PinType.PinSubCategoryMemberReference=(),PinType.bIsArray=False,PinType.bIsReference=False,PinType.bIsConst=False,PinType.bIsWeakPointer=False,DefaultValue="0.0",AutogeneratedDefaultValue="0.0",LinkedTo=(K2Node_VariableGet_281 F3DEDC984C80801BC87F67B0097115B8,),PersistentGuid=00000000000000000000000000000000,bHidden=False,bNotConnectable=False,bDefaultValueIsReadOnly=False,bDefaultValueIsIgnored=False,bAdvancedView=False,)
+			   CustomProperties Pin (PinId=290EC66447D9833EAF1FD8914F55B78C,PinName="B",PinType.PinCategory="float",PinType.PinSubCategoryObject=None,PinType.PinSubCategoryMemberReference=(),PinType.bIsArray=False,PinType.bIsReference=False,PinType.bIsConst=False,PinType.bIsWeakPointer=False,DefaultValue="10",AutogeneratedDefaultValue="0.0",PersistentGuid=00000000000000000000000000000000,bHidden=False,bNotConnectable=False,bDefaultValueIsReadOnly=False,bDefaultValueIsIgnored=False,bAdvancedView=False,)
+			   CustomProperties Pin (PinId=0FAF44A34397D0F87EB02B88E8A653E3,PinName="ReturnValue",Direction="EGPD_Output",PinType.PinCategory="float",PinType.PinSubCategoryObject=None,PinType.PinSubCategoryMemberReference=(),PinType.bIsArray=False,PinType.bIsReference=False,PinType.bIsConst=False,PinType.bIsWeakPointer=False,DefaultValue="0.0",AutogeneratedDefaultValue="0.0",LinkedTo=(K2Node_CallFunction_65349 2C5718584516424BB52CF6A79E702026,),PersistentGuid=00000000000000000000000000000000,bHidden=False,bNotConnectable=False,bDefaultValueIsReadOnly=False,bDefaultValueIsIgnored=False,bAdvancedView=False,)
+			End Object
+			Begin Object Class=K2Node_IfThenElse Name="K2Node_IfThenElse_161"
+			   NodePosX=1600
+			   NodePosY=-432
+			   ErrorType=4
+			   ErrorMsg="Info The execution path doesn\'t end with a return node.  False "
+			   NodeGuid=C447E33540D01EDA1CD72EAB0CFEE477
+			   CustomProperties Pin (PinId=66D340DC4A52DC8D727C43ABA6E818C4,PinName="execute",PinType.PinCategory="exec",PinType.PinSubCategoryObject=None,PinType.PinSubCategoryMemberReference=(),PinType.bIsArray=False,PinType.bIsReference=False,PinType.bIsConst=False,PinType.bIsWeakPointer=False,LinkedTo=(K2Node_CallFunction_65598 E001912C48BB5BEC656460843FE151FA,),PersistentGuid=00000000000000000000000000000000,bHidden=False,bNotConnectable=False,bDefaultValueIsReadOnly=False,bDefaultValueIsIgnored=False,bAdvancedView=False,)
+			   CustomProperties Pin (PinId=1144399140FB3DC67E1A66B5FEB4A1E1,PinName="Condition",PinType.PinCategory="bool",PinType.PinSubCategoryObject=None,PinType.PinSubCategoryMemberReference=(),PinType.bIsArray=False,PinType.bIsReference=False,PinType.bIsConst=False,PinType.bIsWeakPointer=False,DefaultValue="true",AutogeneratedDefaultValue="true",LinkedTo=(K2Node_CallFunction_65598 3A64487743C81996C17E56B543685516,),PersistentGuid=00000000000000000000000000000000,bHidden=False,bNotConnectable=False,bDefaultValueIsReadOnly=False,bDefaultValueIsIgnored=False,bAdvancedView=False,)
+			   CustomProperties Pin (PinId=8DD9A5BC40D019A1AFD6F284F8C15FFD,PinName="then",PinFriendlyName=NSLOCTEXT("K2Node", "true", "true"),Direction="EGPD_Output",PinType.PinCategory="exec",PinType.PinSubCategoryObject=None,PinType.PinSubCategoryMemberReference=(),PinType.bIsArray=False,PinType.bIsReference=False,PinType.bIsConst=False,PinType.bIsWeakPointer=False,LinkedTo=(K2Node_MacroInstance_65 54DCDF6A4C18DFFBF71F54AE24DA40C0,),PersistentGuid=00000000000000000000000000000000,bHidden=False,bNotConnectable=False,bDefaultValueIsReadOnly=False,bDefaultValueIsIgnored=False,bAdvancedView=False,)
+			   CustomProperties Pin (PinId=4CD062D04319B76E8FD9D6B0CDDC7EA9,PinName="else",PinFriendlyName=NSLOCTEXT("K2Node", "false", "false"),Direction="EGPD_Output",PinType.PinCategory="exec",PinType.PinSubCategoryObject=None,PinType.PinSubCategoryMemberReference=(),PinType.bIsArray=False,PinType.bIsReference=False,PinType.bIsConst=False,PinType.bIsWeakPointer=False,PersistentGuid=00000000000000000000000000000000,bHidden=False,bNotConnectable=False,bDefaultValueIsReadOnly=False,bDefaultValueIsIgnored=False,bAdvancedView=False,)
+			End Object
+			Begin Object Class=EdGraphNode_Comment Name="EdGraphNode_Comment_590"
+			   NodePosX=-160
+			   NodePosY=-528
+			   NodeWidth=1616
+			   NodeHeight=704
+			   NodeComment="Trace from the center of the screen in direction of the camera"
+			   NodeGuid=6B615868413EE32AD4B9FC8B9F5DB6B5
+			End Object
+			Begin Object Class=EdGraphNode_Comment Name="EdGraphNode_Comment_596"
+			   NodePosX=1488
+			   NodePosY=-528
+			   NodeWidth=1328
+			   NodeHeight=704
+			   NodeComment="If the Trace hits something apply point damage"
+			   NodeGuid=765187474B4D4B14235EAF883D241CEC
+			End Object
+			Begin Object Class=K2Node_CallFunction Name="K2Node_CallFunction_1861"
+			   FunctionReference=(MemberParent=Class'/Script/Engine.GameplayStatics',MemberName="ApplyPointDamage")
+			   NodePosX=2560
+			   NodePosY=-432
+			   NodeGuid=EA250A4244DD16B5695B73803C315BA1
+			   CustomProperties Pin (PinId=6F410C0B46D7B07315AD748B03EB86C8,PinName="execute",PinToolTip="\nExec",PinType.PinCategory="exec",PinType.PinSubCategoryObject=None,PinType.PinSubCategoryMemberReference=(),PinType.bIsArray=False,PinType.bIsReference=False,PinType.bIsConst=False,PinType.bIsWeakPointer=False,LinkedTo=(K2Node_MacroInstance_65 E75DA26F4EBE18C63E3CEDA6B1D88DC6,),PersistentGuid=00000000000000000000000000000000,bHidden=False,bNotConnectable=False,bDefaultValueIsReadOnly=False,bDefaultValueIsIgnored=False,bAdvancedView=False,)
+			   CustomProperties Pin (PinId=57573271406B9D0CCE9C28BF2DE94719,PinName="then",PinToolTip="\nExec",Direction="EGPD_Output",PinType.PinCategory="exec",PinType.PinSubCategoryObject=None,PinType.PinSubCategoryMemberReference=(),PinType.bIsArray=False,PinType.bIsReference=False,PinType.bIsConst=False,PinType.bIsWeakPointer=False,PersistentGuid=00000000000000000000000000000000,bHidden=False,bNotConnectable=False,bDefaultValueIsReadOnly=False,bDefaultValueIsIgnored=False,bAdvancedView=False,)
+			   CustomProperties Pin (PinId=52375B5841E301566AD3AA961C2B94D1,PinName="self",PinFriendlyName=NSLOCTEXT("K2Node", "Target", "Target"),PinToolTip="Target\nGameplay Statics Reference",PinType.PinCategory="object",PinType.PinSubCategoryObject=Class'/Script/Engine.GameplayStatics',PinType.PinSubCategoryMemberReference=(),PinType.bIsArray=False,PinType.bIsReference=False,PinType.bIsConst=False,PinType.bIsWeakPointer=False,DefaultObject="/Script/Engine.Default__GameplayStatics",PersistentGuid=00000000000000000000000000000000,bHidden=True,bNotConnectable=False,bDefaultValueIsReadOnly=False,bDefaultValueIsIgnored=False,bAdvancedView=False,)
+			   CustomProperties Pin (PinId=BB747167446E2DDBEFFFDD982A9CD245,PinName="DamagedActor",PinToolTip="Damaged Actor\nActor Reference\n\nActor that will be damaged.",PinType.PinCategory="object",PinType.PinSubCategoryObject=Class'/Script/Engine.Actor',PinType.PinSubCategoryMemberReference=(),PinType.bIsArray=False,PinType.bIsReference=False,PinType.bIsConst=False,PinType.bIsWeakPointer=False,LinkedTo=(K2Node_CallFunction_65759 F1D6985542953219C7EF09A3FBFA8AF0,),PersistentGuid=00000000000000000000000000000000,bHidden=False,bNotConnectable=False,bDefaultValueIsReadOnly=False,bDefaultValueIsIgnored=False,bAdvancedView=False,)
+			   CustomProperties Pin (PinId=E0E59CE348F9C5A31938AA864E1EDF8C,PinName="BaseDamage",PinToolTip="Base Damage\nFloat\n\nThe base damage to apply.",PinType.PinCategory="float",PinType.PinSubCategoryObject=None,PinType.PinSubCategoryMemberReference=(),PinType.bIsArray=False,PinType.bIsReference=False,PinType.bIsConst=False,PinType.bIsWeakPointer=False,DefaultValue="100",AutogeneratedDefaultValue="0.0",PersistentGuid=00000000000000000000000000000000,bHidden=False,bNotConnectable=False,bDefaultValueIsReadOnly=False,bDefaultValueIsIgnored=False,bAdvancedView=False,)
+			   CustomProperties Pin (PinId=744D73DE45C865A6B95D1A90F0DF7412,PinName="HitFromDirection",PinToolTip="Hit from Direction\nVector  (by ref)\n\nDirection the hit came FROM",PinType.PinCategory="struct",PinType.PinSubCategoryObject=ScriptStruct'/Script/CoreUObject.Vector',PinType.PinSubCategoryMemberReference=(),PinType.bIsArray=False,PinType.bIsReference=True,PinType.bIsConst=True,PinType.bIsWeakPointer=False,DefaultValue="0, 0, 0",AutogeneratedDefaultValue="0, 0, 0",LinkedTo=(K2Node_CallFunction_1862 9C37BD7C42AFE3084A157196D1784CD3,),PersistentGuid=00000000000000000000000000000000,bHidden=False,bNotConnectable=False,bDefaultValueIsReadOnly=False,bDefaultValueIsIgnored=True,bAdvancedView=False,)
+			   CustomProperties Pin (PinId=7453F08741D187AE2B9629985E33A634,PinName="HitInfo",PinToolTip="Hit Info\nHit Result Structure (by ref)\n\nCollision or trace result that describes the hit",PinType.PinCategory="struct",PinType.PinSubCategoryObject=ScriptStruct'/Script/Engine.HitResult',PinType.PinSubCategoryMemberReference=(),PinType.bIsArray=False,PinType.bIsReference=True,PinType.bIsConst=True,PinType.bIsWeakPointer=False,LinkedTo=(K2Node_CallFunction_65598 43A3FD1943054F5E9C27A298298CE721,),PersistentGuid=00000000000000000000000000000000,bHidden=False,bNotConnectable=False,bDefaultValueIsReadOnly=False,bDefaultValueIsIgnored=True,bAdvancedView=False,)
+			   CustomProperties Pin (PinId=8AAA25E3411AD93065E153BDF12317E1,PinName="EventInstigator",PinToolTip="Event Instigator\nController Reference\n\nController that was responsible for causing this damage (e.g. player who shot the weapon)",PinType.PinCategory="object",PinType.PinSubCategoryObject=Class'/Script/Engine.Controller',PinType.PinSubCategoryMemberReference=(),PinType.bIsArray=False,PinType.bIsReference=False,PinType.bIsConst=False,PinType.bIsWeakPointer=False,LinkedTo=(K2Node_CallFunction_2902 3119A2234CC8CC2E58EDDA94ED78DB30,),PersistentGuid=00000000000000000000000000000000,bHidden=False,bNotConnectable=False,bDefaultValueIsReadOnly=False,bDefaultValueIsIgnored=False,bAdvancedView=False,)
+			   CustomProperties Pin (PinId=471A4E874191001E241068AF8F5CA28F,PinName="DamageCauser",PinToolTip="Damage Causer\nActor Reference\n\nActor that actually caused the damage (e.g. the grenade that exploded)",PinType.PinCategory="object",PinType.PinSubCategoryObject=Class'/Script/Engine.Actor',PinType.PinSubCategoryMemberReference=(),PinType.bIsArray=False,PinType.bIsReference=False,PinType.bIsConst=False,PinType.bIsWeakPointer=False,LinkedTo=(K2Node_Self_15 CFA0BC05479F37C57BC5E890E311DF98,),PersistentGuid=00000000000000000000000000000000,bHidden=False,bNotConnectable=False,bDefaultValueIsReadOnly=False,bDefaultValueIsIgnored=False,bAdvancedView=False,)
+			   CustomProperties Pin (PinId=423FB9E642E33BBA76EDACAF6D63DCA0,PinName="DamageTypeClass",PinToolTip="Damage Type Class\nDamage Type Class\n\nClass that describes the damage that was done.",PinType.PinCategory="class",PinType.PinSubCategoryObject=Class'/Script/Engine.DamageType',PinType.PinSubCategoryMemberReference=(),PinType.bIsArray=False,PinType.bIsReference=False,PinType.bIsConst=False,PinType.bIsWeakPointer=False,DefaultObject="/Script/Engine.DamageType",PersistentGuid=00000000000000000000000000000000,bHidden=False,bNotConnectable=False,bDefaultValueIsReadOnly=False,bDefaultValueIsIgnored=False,bAdvancedView=False,)
+			End Object
+			Begin Object Class=K2Node_CallFunction Name="K2Node_CallFunction_1862"
+			   bIsPureFunc=True
+			   FunctionReference=(MemberParent=Class'/Script/Engine.KismetMathLibrary',MemberName="MakeVector")
+			   NodePosX=2240
+			   NodePosY=-352
+			   NodeGuid=90D1D25A44EFBE2A50A7C4AB73385384
+			   CustomProperties Pin (PinId=89111CA54D160B2C7528509D086B822B,PinName="self",PinFriendlyName=NSLOCTEXT("K2Node", "Target", "Target"),PinToolTip="Target\nKismet Math Library Reference",PinType.PinCategory="object",PinType.PinSubCategoryObject=Class'/Script/Engine.KismetMathLibrary',PinType.PinSubCategoryMemberReference=(),PinType.bIsArray=False,PinType.bIsReference=False,PinType.bIsConst=False,PinType.bIsWeakPointer=False,DefaultObject="/Script/Engine.Default__KismetMathLibrary",PersistentGuid=00000000000000000000000000000000,bHidden=True,bNotConnectable=False,bDefaultValueIsReadOnly=False,bDefaultValueIsIgnored=False,bAdvancedView=False,)
+			   CustomProperties Pin (PinId=396AF363497B394C5C3BA4BCB5A80C67,PinName="X",PinToolTip="X\nFloat",PinType.PinCategory="float",PinType.PinSubCategoryObject=None,PinType.PinSubCategoryMemberReference=(),PinType.bIsArray=False,PinType.bIsReference=False,PinType.bIsConst=False,PinType.bIsWeakPointer=False,DefaultValue="0.0",AutogeneratedDefaultValue="0.0",PersistentGuid=00000000000000000000000000000000,bHidden=False,bNotConnectable=False,bDefaultValueIsReadOnly=False,bDefaultValueIsIgnored=False,bAdvancedView=False,)
+			   CustomProperties Pin (PinId=24D8F5EE410BCDB0B82938825D4FCECC,PinName="Y",PinToolTip="Y\nFloat",PinType.PinCategory="float",PinType.PinSubCategoryObject=None,PinType.PinSubCategoryMemberReference=(),PinType.bIsArray=False,PinType.bIsReference=False,PinType.bIsConst=False,PinType.bIsWeakPointer=False,DefaultValue="0.0",AutogeneratedDefaultValue="0.0",PersistentGuid=00000000000000000000000000000000,bHidden=False,bNotConnectable=False,bDefaultValueIsReadOnly=False,bDefaultValueIsIgnored=False,bAdvancedView=False,)
+			   CustomProperties Pin (PinId=8C7BE8EB4BFE81E73D8821ABC5B5C51D,PinName="Z",PinToolTip="Z\nFloat",PinType.PinCategory="float",PinType.PinSubCategoryObject=None,PinType.PinSubCategoryMemberReference=(),PinType.bIsArray=False,PinType.bIsReference=False,PinType.bIsConst=False,PinType.bIsWeakPointer=False,DefaultValue="1",AutogeneratedDefaultValue="0.0",PersistentGuid=00000000000000000000000000000000,bHidden=False,bNotConnectable=False,bDefaultValueIsReadOnly=False,bDefaultValueIsIgnored=False,bAdvancedView=False,)
+			   CustomProperties Pin (PinId=9C37BD7C42AFE3084A157196D1784CD3,PinName="ReturnValue",PinToolTip="Return Value\nVector ",Direction="EGPD_Output",PinType.PinCategory="struct",PinType.PinSubCategoryObject=ScriptStruct'/Script/CoreUObject.Vector',PinType.PinSubCategoryMemberReference=(),PinType.bIsArray=False,PinType.bIsReference=False,PinType.bIsConst=False,PinType.bIsWeakPointer=False,AutogeneratedDefaultValue="0, 0, 0",LinkedTo=(K2Node_CallFunction_1861 744D73DE45C865A6B95D1A90F0DF7412,),PersistentGuid=00000000000000000000000000000000,bHidden=False,bNotConnectable=False,bDefaultValueIsReadOnly=False,bDefaultValueIsIgnored=False,bAdvancedView=False,)
+			End Object
+			Begin Object Class=K2Node_VariableGet Name="K2Node_VariableGet_281"
+			   VariableReference=(MemberName="InteractionDistance",MemberGuid=024096844ACA19D4D6CB4DA334906B98,bSelfContext=True)
+			   NodePosX=384
+			   NodePosY=64
+			   ErrorType=1
+			   ErrorMsg="Error The property associated with  InteractionDistance  could not be found"
+			   NodeGuid=EC502D5E41D525430986138A615256A9
+			   CustomProperties Pin (PinId=F3DEDC984C80801BC87F67B0097115B8,PinName="InteractionDistance",Direction="EGPD_Output",PinType.PinCategory="float",PinType.PinSubCategoryObject=None,PinType.PinSubCategoryMemberReference=(),PinType.bIsArray=False,PinType.bIsReference=False,PinType.bIsConst=False,PinType.bIsWeakPointer=False,AutogeneratedDefaultValue="0.0",LinkedTo=(K2Node_CommutativeAssociativeBinaryOperator_4227 AE00303B4D5F38B1B48C718ABEC37EF0,),PersistentGuid=00000000000000000000000000000000,bHidden=False,bNotConnectable=False,bDefaultValueIsReadOnly=False,bDefaultValueIsIgnored=False,bAdvancedView=False,)
+			   CustomProperties Pin (PinId=7C7F96A046F0B4BB8BDA5C86407AD84A,PinName="self",PinFriendlyName=NSLOCTEXT("K2Node", "Target", "Target"),PinType.PinCategory="object",PinType.PinSubCategoryObject=BlueprintGeneratedClass'/Game/ExampleContent/BlueprintFluidSurface/HeightfieldPainter.HeightfieldPainter_C',PinType.PinSubCategoryMemberReference=(),PinType.bIsArray=False,PinType.bIsReference=False,PinType.bIsConst=False,PinType.bIsWeakPointer=False,PersistentGuid=00000000000000000000000000000000,bHidden=True,bNotConnectable=False,bDefaultValueIsReadOnly=False,bDefaultValueIsIgnored=False,bAdvancedView=False,)
+			End Object
+			Begin Object Class=K2Node_Self Name="K2Node_Self_15"
+			   NodePosX=2336
+			   NodePosY=-80
+			   NodeGuid=519A1C834BAB76F02E68CE802A2ED70E
+			   CustomProperties Pin (PinId=CFA0BC05479F37C57BC5E890E311DF98,PinName="self",Direction="EGPD_Output",PinType.PinCategory="object",PinType.PinSubCategory="self",PinType.PinSubCategoryObject=None,PinType.PinSubCategoryMemberReference=(),PinType.bIsArray=False,PinType.bIsReference=False,PinType.bIsConst=False,PinType.bIsWeakPointer=False,LinkedTo=(K2Node_CallFunction_1861 471A4E874191001E241068AF8F5CA28F,),PersistentGuid=00000000000000000000000000000000,bHidden=False,bNotConnectable=False,bDefaultValueIsReadOnly=False,bDefaultValueIsIgnored=False,bAdvancedView=False,)
+			End Object
+	[/PARAM]
+[/OBJECT]
+[/EXCERPT:CameraTrace]
+
+[EXCERPT:ConScript]	
+[OBJECT:Snippet]
+	[PARAMLITERAL:id]
+		ConScript
+	[/PARAMLITERAL]
+	[PARAMLITERAL:type]
+		blueprint
+	[/PARAMLITERAL]
+	[PARAM:image]
+		![](Engine\Rendering\RenderTargets\BlueprintRenderTargets\HowTo\FluidSurface\FS_Con_Script.png)
+	[/PARAM]
+	[PARAM:snippet]
+			BEGIN OBJECT
+							Begin Object Class=K2Node_FunctionEntry Name="K2Node_FunctionEntry_0"
+			   SignatureClass=Class'/Script/Engine.Actor'
+			   SignatureName="UserConstructionScript"
+			   NodePosX=64
+			   NodePosY=80
+			   NodeGuid=FC1BD3C84EF746019371C3B6D4B78A3A
+			   CustomProperties Pin (PinId=D4F4417A43C11A1D5D957D822F70747A,PinName="then",Direction="EGPD_Output",PinType.PinCategory="exec",PinType.PinSubCategoryObject=None,PinType.PinSubCategoryMemberReference=(),PinType.bIsArray=False,PinType.bIsReference=False,PinType.bIsConst=False,PinType.bIsWeakPointer=False,LinkedTo=(K2Node_CallFunction_0 DF1FF164453221B85502C2901471C25D,),PersistentGuid=00000000000000000000000000000000,bHidden=False,bNotConnectable=False,bDefaultValueIsReadOnly=False,bDefaultValueIsIgnored=False,bAdvancedView=False,)
+			End Object
+			Begin Object Class=K2Node_CallFunction Name="K2Node_CallFunction_0"
+			   FunctionReference=(MemberParent=Class'/Script/Engine.StaticMeshComponent',MemberName="SetStaticMesh")
+			   NodePosX=400
+			   NodePosY=64
+			   NodeGuid=DCC3240F4054FFD4C86DB8867FB17F87
+			   CustomProperties Pin (PinId=DF1FF164453221B85502C2901471C25D,PinName="execute",PinToolTip="\nExec",PinType.PinCategory="exec",PinType.PinSubCategoryObject=None,PinType.PinSubCategoryMemberReference=(),PinType.bIsArray=False,PinType.bIsReference=False,PinType.bIsConst=False,PinType.bIsWeakPointer=False,LinkedTo=(K2Node_FunctionEntry_0 D4F4417A43C11A1D5D957D822F70747A,),PersistentGuid=00000000000000000000000000000000,bHidden=False,bNotConnectable=False,bDefaultValueIsReadOnly=False,bDefaultValueIsIgnored=False,bAdvancedView=False,)
+			   CustomProperties Pin (PinId=B9881C0142A99A1A7ACD86946C02DD98,PinName="then",PinToolTip="\nExec",Direction="EGPD_Output",PinType.PinCategory="exec",PinType.PinSubCategoryObject=None,PinType.PinSubCategoryMemberReference=(),PinType.bIsArray=False,PinType.bIsReference=False,PinType.bIsConst=False,PinType.bIsWeakPointer=False,LinkedTo=(K2Node_CallFunction_98 DDFBB4294203AC5FB928DDBC3A82DC1E,),PersistentGuid=00000000000000000000000000000000,bHidden=False,bNotConnectable=False,bDefaultValueIsReadOnly=False,bDefaultValueIsIgnored=False,bAdvancedView=False,)
+			   CustomProperties Pin (PinId=391A69D247367ED3E62D4B869CDD282E,PinName="self",PinFriendlyName=NSLOCTEXT("K2Node", "Target", "Target"),PinToolTip="Target\nStatic Mesh Component Reference",PinType.PinCategory="object",PinType.PinSubCategoryObject=Class'/Script/Engine.StaticMeshComponent',PinType.PinSubCategoryMemberReference=(),PinType.bIsArray=False,PinType.bIsReference=False,PinType.bIsConst=False,PinType.bIsWeakPointer=False,LinkedTo=(K2Node_VariableGet_32 8C2A8CB743E63ABCF573FC84E3867E6E,),PersistentGuid=00000000000000000000000000000000,bHidden=False,bNotConnectable=False,bDefaultValueIsReadOnly=False,bDefaultValueIsIgnored=False,bAdvancedView=False,)
+			   CustomProperties Pin (PinId=A858BC66470AD5D0D3E7C0B3013DDE38,PinName="NewMesh",PinToolTip="New Mesh\nStatic Mesh Reference",PinType.PinCategory="object",PinType.PinSubCategoryObject=Class'/Script/Engine.StaticMesh',PinType.PinSubCategoryMemberReference=(),PinType.bIsArray=False,PinType.bIsReference=False,PinType.bIsConst=False,PinType.bIsWeakPointer=False,LinkedTo=(K2Node_VariableGet_90 46E846144D26AF7BDE861DB50DAC1F0F,),PersistentGuid=00000000000000000000000000000000,bHidden=False,bNotConnectable=False,bDefaultValueIsReadOnly=False,bDefaultValueIsIgnored=False,bAdvancedView=False,)
+			   CustomProperties Pin (PinId=DBF0B5984A4CDA2580ACF3981A0F4978,PinName="ReturnValue",PinToolTip="Return Value\nBoolean",Direction="EGPD_Output",PinType.PinCategory="bool",PinType.PinSubCategoryObject=None,PinType.PinSubCategoryMemberReference=(),PinType.bIsArray=False,PinType.bIsReference=False,PinType.bIsConst=False,PinType.bIsWeakPointer=False,DefaultValue="false",AutogeneratedDefaultValue="false",PersistentGuid=00000000000000000000000000000000,bHidden=False,bNotConnectable=False,bDefaultValueIsReadOnly=False,bDefaultValueIsIgnored=False,bAdvancedView=False,)
+			End Object
+			Begin Object Class=K2Node_VariableGet Name="K2Node_VariableGet_90"
+			   VariableReference=(MemberName="WaterMeshAsset",MemberGuid=1E7B134B4BB9507C500AAF8E074F94AE,bSelfContext=True)
+			   NodePosX=208
+			   NodePosY=256
+			   NodeGuid=6414FF7B4B790DE458ED3C8DF2E78A5E
+			   CustomProperties Pin (PinId=46E846144D26AF7BDE861DB50DAC1F0F,PinName="WaterMeshAsset",Direction="EGPD_Output",PinType.PinCategory="object",PinType.PinSubCategoryObject=Class'/Script/Engine.StaticMesh',PinType.PinSubCategoryMemberReference=(),PinType.bIsArray=False,PinType.bIsReference=False,PinType.bIsConst=False,PinType.bIsWeakPointer=False,LinkedTo=(K2Node_CallFunction_0 A858BC66470AD5D0D3E7C0B3013DDE38,),PersistentGuid=00000000000000000000000000000000,bHidden=False,bNotConnectable=False,bDefaultValueIsReadOnly=False,bDefaultValueIsIgnored=False,bAdvancedView=False,)
+			   CustomProperties Pin (PinId=876704FB44426B730BC72EBC94E7A86F,PinName="self",PinFriendlyName=NSLOCTEXT("K2Node", "Target", "Target"),PinType.PinCategory="object",PinType.PinSubCategoryObject=BlueprintGeneratedClass'/Game/ExampleContent/BlueprintFluidSurface/WaterSurface_BP.WaterSurface_BP_C',PinType.PinSubCategoryMemberReference=(),PinType.bIsArray=False,PinType.bIsReference=False,PinType.bIsConst=False,PinType.bIsWeakPointer=False,PersistentGuid=00000000000000000000000000000000,bHidden=True,bNotConnectable=False,bDefaultValueIsReadOnly=False,bDefaultValueIsIgnored=False,bAdvancedView=False,)
+			End Object
+			Begin Object Class=K2Node_CallFunction Name="K2Node_CallFunction_1"
+			   FunctionReference=(MemberParent=Class'/Script/Engine.PrimitiveComponent',MemberName="SetMaterial")
+			   NodePosX=1200
+			   NodePosY=63
+			   NodeGuid=A0EED7AA45541C99668205AB130EE1F8
+			   CustomProperties Pin (PinId=632CD37F4755247C5416E48F42DC2780,PinName="execute",PinType.PinCategory="exec",PinType.PinSubCategoryObject=None,PinType.PinSubCategoryMemberReference=(),PinType.bIsArray=False,PinType.bIsReference=False,PinType.bIsConst=False,PinType.bIsWeakPointer=False,LinkedTo=(K2Node_VariableSet_54 F84011A84D02E23F2AE1079260B38F98,),PersistentGuid=00000000000000000000000000000000,bHidden=False,bNotConnectable=False,bDefaultValueIsReadOnly=False,bDefaultValueIsIgnored=False,bAdvancedView=False,)
+			   CustomProperties Pin (PinId=CF8EC9BE4C3F30D7287255904423FAA7,PinName="then",Direction="EGPD_Output",PinType.PinCategory="exec",PinType.PinSubCategoryObject=None,PinType.PinSubCategoryMemberReference=(),PinType.bIsArray=False,PinType.bIsReference=False,PinType.bIsConst=False,PinType.bIsWeakPointer=False,PersistentGuid=00000000000000000000000000000000,bHidden=False,bNotConnectable=False,bDefaultValueIsReadOnly=False,bDefaultValueIsIgnored=False,bAdvancedView=False,)
+			   CustomProperties Pin (PinId=C48D7D564D13F270DA1266BD1E9280B9,PinName="self",PinFriendlyName=NSLOCTEXT("K2Node", "Target", "Target"),PinType.PinCategory="object",PinType.PinSubCategoryObject=Class'/Script/Engine.PrimitiveComponent',PinType.PinSubCategoryMemberReference=(),PinType.bIsArray=False,PinType.bIsReference=False,PinType.bIsConst=False,PinType.bIsWeakPointer=False,LinkedTo=(K2Node_VariableGet_31 A57F7863412610395909CBB241520793,),PersistentGuid=00000000000000000000000000000000,bHidden=False,bNotConnectable=False,bDefaultValueIsReadOnly=False,bDefaultValueIsIgnored=False,bAdvancedView=False,)
+			   CustomProperties Pin (PinId=14D3B8E649B8477F3FAE6B9028A91B06,PinName="ElementIndex",PinType.PinCategory="int",PinType.PinSubCategoryObject=None,PinType.PinSubCategoryMemberReference=(),PinType.bIsArray=False,PinType.bIsReference=False,PinType.bIsConst=False,PinType.bIsWeakPointer=False,DefaultValue="0",AutogeneratedDefaultValue="0",PersistentGuid=00000000000000000000000000000000,bHidden=False,bNotConnectable=False,bDefaultValueIsReadOnly=False,bDefaultValueIsIgnored=False,bAdvancedView=False,)
+			   CustomProperties Pin (PinId=5C7FC206400C1639B8D25CA515017C30,PinName="Material",PinType.PinCategory="object",PinType.PinSubCategoryObject=Class'/Script/Engine.MaterialInterface',PinType.PinSubCategoryMemberReference=(),PinType.bIsArray=False,PinType.bIsReference=False,PinType.bIsConst=False,PinType.bIsWeakPointer=False,LinkedTo=(K2Node_CallFunction_98 ABCDC3F1428BBCD3C67176AB6DB2F642,),PersistentGuid=00000000000000000000000000000000,bHidden=False,bNotConnectable=False,bDefaultValueIsReadOnly=False,bDefaultValueIsIgnored=False,bAdvancedView=False,)
+			End Object
+			Begin Object Class=K2Node_VariableGet Name="K2Node_VariableGet_96"
+			   VariableReference=(MemberName="WaterMaterial",MemberGuid=7A08FB3440678BE53972B28848FEEBE4,bSelfContext=True)
+			   NodePosX=448
+			   NodePosY=304
+			   NodeGuid=4ACF53134A0B7543E12AFC90FA362121
+			   CustomProperties Pin (PinId=CA0EE3824870973F0E0D27B2A34E54EC,PinName="WaterMaterial",Direction="EGPD_Output",PinType.PinCategory="object",PinType.PinSubCategoryObject=Class'/Script/Engine.MaterialInterface',PinType.PinSubCategoryMemberReference=(),PinType.bIsArray=False,PinType.bIsReference=False,PinType.bIsConst=False,PinType.bIsWeakPointer=False,LinkedTo=(K2Node_CallFunction_98 2D6C20F4445A6B3CEAE4E4A5C9D49C82,),PersistentGuid=00000000000000000000000000000000,bHidden=False,bNotConnectable=False,bDefaultValueIsReadOnly=False,bDefaultValueIsIgnored=False,bAdvancedView=False,)
+			   CustomProperties Pin (PinId=83FFBEE04B50FF7F3B27EE85CFA66BE9,PinName="self",PinFriendlyName=NSLOCTEXT("K2Node", "Target", "Target"),PinType.PinCategory="object",PinType.PinSubCategoryObject=BlueprintGeneratedClass'/Game/ExampleContent/BlueprintFluidSurface/WaterSurface_BP.WaterSurface_BP_C',PinType.PinSubCategoryMemberReference=(),PinType.bIsArray=False,PinType.bIsReference=False,PinType.bIsConst=False,PinType.bIsWeakPointer=False,PersistentGuid=00000000000000000000000000000000,bHidden=True,bNotConnectable=False,bDefaultValueIsReadOnly=False,bDefaultValueIsIgnored=False,bAdvancedView=False,)
+			End Object
+			Begin Object Class=K2Node_CallFunction Name="K2Node_CallFunction_98"
+			   FunctionReference=(MemberParent=Class'/Script/Engine.PrimitiveComponent',MemberName="CreateDynamicMaterialInstance")
+			   NodePosX=704
+			   NodePosY=64
+			   ErrorType=1
+			   ErrorMsg="Error This blueprint (self) is not a PrimitiveComponent, therefore \' Target \' must have a connection."
+			   NodeGuid=5881977D4EA5F38D74B659A091C89AED
+			   CustomProperties Pin (PinId=DDFBB4294203AC5FB928DDBC3A82DC1E,PinName="execute",PinToolTip="\nExec",PinType.PinCategory="exec",PinType.PinSubCategoryObject=None,PinType.PinSubCategoryMemberReference=(),PinType.bIsArray=False,PinType.bIsReference=False,PinType.bIsConst=False,PinType.bIsWeakPointer=False,LinkedTo=(K2Node_CallFunction_0 B9881C0142A99A1A7ACD86946C02DD98,),PersistentGuid=00000000000000000000000000000000,bHidden=False,bNotConnectable=False,bDefaultValueIsReadOnly=False,bDefaultValueIsIgnored=False,bAdvancedView=False,)
+			   CustomProperties Pin (PinId=2D2DC36F42CDB2E197632BA2B6AEC5DB,PinName="then",PinToolTip="\nExec",Direction="EGPD_Output",PinType.PinCategory="exec",PinType.PinSubCategoryObject=None,PinType.PinSubCategoryMemberReference=(),PinType.bIsArray=False,PinType.bIsReference=False,PinType.bIsConst=False,PinType.bIsWeakPointer=False,LinkedTo=(K2Node_VariableSet_54 823A98AD45AF085CE5C7BABE671A72C9,),PersistentGuid=00000000000000000000000000000000,bHidden=False,bNotConnectable=False,bDefaultValueIsReadOnly=False,bDefaultValueIsIgnored=False,bAdvancedView=False,)
+			   CustomProperties Pin (PinId=C428EDF34B24E1BDAC09BFADBBE2D412,PinName="self",PinFriendlyName=NSLOCTEXT("K2Node", "Target", "Target"),PinToolTip="Target\nPrimitive Component Reference",PinType.PinCategory="object",PinType.PinSubCategoryObject=Class'/Script/Engine.PrimitiveComponent',PinType.PinSubCategoryMemberReference=(),PinType.bIsArray=False,PinType.bIsReference=False,PinType.bIsConst=False,PinType.bIsWeakPointer=False,LinkedTo=(K2Node_VariableGet_31 A57F7863412610395909CBB241520793,),PersistentGuid=00000000000000000000000000000000,bHidden=False,bNotConnectable=False,bDefaultValueIsReadOnly=False,bDefaultValueIsIgnored=False,bAdvancedView=False,)
+			   CustomProperties Pin (PinId=D31EE7E4499FA64E16AB04976C183DF3,PinName="ElementIndex",PinToolTip="Element Index\nInteger\n\nThe index of the skin to replace the material for.  If invalid, the material is unchanged and NULL is returned.",PinType.PinCategory="int",PinType.PinSubCategoryObject=None,PinType.PinSubCategoryMemberReference=(),PinType.bIsArray=False,PinType.bIsReference=False,PinType.bIsConst=False,PinType.bIsWeakPointer=False,DefaultValue="0",AutogeneratedDefaultValue="0",PersistentGuid=00000000000000000000000000000000,bHidden=False,bNotConnectable=False,bDefaultValueIsReadOnly=False,bDefaultValueIsIgnored=False,bAdvancedView=False,)
+			   CustomProperties Pin (PinId=2D6C20F4445A6B3CEAE4E4A5C9D49C82,PinName="SourceMaterial",PinToolTip="Source Material\nMaterial Interface Reference",PinType.PinCategory="object",PinType.PinSubCategoryObject=Class'/Script/Engine.MaterialInterface',PinType.PinSubCategoryMemberReference=(),PinType.bIsArray=False,PinType.bIsReference=False,PinType.bIsConst=False,PinType.bIsWeakPointer=False,LinkedTo=(K2Node_VariableGet_96 CA0EE3824870973F0E0D27B2A34E54EC,),PersistentGuid=00000000000000000000000000000000,bHidden=False,bNotConnectable=False,bDefaultValueIsReadOnly=False,bDefaultValueIsIgnored=False,bAdvancedView=False,)
+			   CustomProperties Pin (PinId=ABCDC3F1428BBCD3C67176AB6DB2F642,PinName="ReturnValue",PinToolTip="Return Value\nMaterial Instance Dynamic Reference",Direction="EGPD_Output",PinType.PinCategory="object",PinType.PinSubCategoryObject=Class'/Script/Engine.MaterialInstanceDynamic',PinType.PinSubCategoryMemberReference=(),PinType.bIsArray=False,PinType.bIsReference=False,PinType.bIsConst=False,PinType.bIsWeakPointer=False,LinkedTo=(K2Node_CallFunction_1 5C7FC206400C1639B8D25CA515017C30,K2Node_VariableSet_54 8DEE861149A44436A09E2385403877DF,),PersistentGuid=00000000000000000000000000000000,bHidden=False,bNotConnectable=False,bDefaultValueIsReadOnly=False,bDefaultValueIsIgnored=False,bAdvancedView=False,)
+			End Object
+			Begin Object Class=K2Node_VariableSet Name="K2Node_VariableSet_54"
+			   VariableReference=(MemberName="WaterMID",MemberGuid=900EBFD242C0177F9BFCF7A541E03E45,bSelfContext=True)
+			   NodePosX=1008
+			   NodePosY=95
+			   NodeGuid=AA9999CF4D77647CD7A417BD7413069B
+			   CustomProperties Pin (PinId=823A98AD45AF085CE5C7BABE671A72C9,PinName="execute",PinType.PinCategory="exec",PinType.PinSubCategoryObject=None,PinType.PinSubCategoryMemberReference=(),PinType.bIsArray=False,PinType.bIsReference=False,PinType.bIsConst=False,PinType.bIsWeakPointer=False,LinkedTo=(K2Node_CallFunction_98 2D2DC36F42CDB2E197632BA2B6AEC5DB,),PersistentGuid=00000000000000000000000000000000,bHidden=False,bNotConnectable=False,bDefaultValueIsReadOnly=False,bDefaultValueIsIgnored=False,bAdvancedView=False,)
+			   CustomProperties Pin (PinId=F84011A84D02E23F2AE1079260B38F98,PinName="then",Direction="EGPD_Output",PinType.PinCategory="exec",PinType.PinSubCategoryObject=None,PinType.PinSubCategoryMemberReference=(),PinType.bIsArray=False,PinType.bIsReference=False,PinType.bIsConst=False,PinType.bIsWeakPointer=False,LinkedTo=(K2Node_CallFunction_1 632CD37F4755247C5416E48F42DC2780,),PersistentGuid=00000000000000000000000000000000,bHidden=False,bNotConnectable=False,bDefaultValueIsReadOnly=False,bDefaultValueIsIgnored=False,bAdvancedView=False,)
+			   CustomProperties Pin (PinId=8DEE861149A44436A09E2385403877DF,PinName="WaterMID",PinType.PinCategory="object",PinType.PinSubCategoryObject=Class'/Script/Engine.MaterialInstanceDynamic',PinType.PinSubCategoryMemberReference=(),PinType.bIsArray=False,PinType.bIsReference=False,PinType.bIsConst=False,PinType.bIsWeakPointer=False,LinkedTo=(K2Node_CallFunction_98 ABCDC3F1428BBCD3C67176AB6DB2F642,),PersistentGuid=00000000000000000000000000000000,bHidden=False,bNotConnectable=False,bDefaultValueIsReadOnly=False,bDefaultValueIsIgnored=False,bAdvancedView=False,)
+			   CustomProperties Pin (PinId=EEE14CCB4B26F135A609EDA9FDD7E204,PinName="Output_Get",PinToolTip="Retrieves the value of the variable, can use instead of a separate Get node",Direction="EGPD_Output",PinType.PinCategory="object",PinType.PinSubCategoryObject=Class'/Script/Engine.MaterialInstanceDynamic',PinType.PinSubCategoryMemberReference=(),PinType.bIsArray=False,PinType.bIsReference=False,PinType.bIsConst=False,PinType.bIsWeakPointer=False,PersistentGuid=00000000000000000000000000000000,bHidden=False,bNotConnectable=False,bDefaultValueIsReadOnly=False,bDefaultValueIsIgnored=False,bAdvancedView=False,)
+			   CustomProperties Pin (PinId=15F41C99419768BF86865DA78417A1C3,PinName="self",PinFriendlyName=NSLOCTEXT("K2Node", "Target", "Target"),PinType.PinCategory="object",PinType.PinSubCategoryObject=BlueprintGeneratedClass'/Game/ExampleContent/BlueprintFluidSurface/WaterSurface_BP.WaterSurface_BP_C',PinType.PinSubCategoryMemberReference=(),PinType.bIsArray=False,PinType.bIsReference=False,PinType.bIsConst=False,PinType.bIsWeakPointer=False,PersistentGuid=00000000000000000000000000000000,bHidden=True,bNotConnectable=False,bDefaultValueIsReadOnly=False,bDefaultValueIsIgnored=False,bAdvancedView=False,)
+			End Object
+			Begin Object Class=K2Node_VariableGet Name="K2Node_VariableGet_31"
+			   VariableReference=(MemberName="StaticMeshComponent",bSelfContext=True)
+			   NodePosX=432
+			   NodePosY=352
+			   NodeGuid=A518C0D843C9741527E3E1A049D2C5AB
+			   CustomProperties Pin (PinId=A57F7863412610395909CBB241520793,PinName="StaticMeshComponent",Direction="EGPD_Output",PinType.PinCategory="object",PinType.PinSubCategoryObject=Class'/Script/Engine.StaticMeshComponent',PinType.PinSubCategoryMemberReference=(),PinType.bIsArray=False,PinType.bIsReference=False,PinType.bIsConst=False,PinType.bIsWeakPointer=False,LinkedTo=(K2Node_CallFunction_98 C428EDF34B24E1BDAC09BFADBBE2D412,K2Node_CallFunction_1 C48D7D564D13F270DA1266BD1E9280B9,),PersistentGuid=00000000000000000000000000000000,bHidden=False,bNotConnectable=False,bDefaultValueIsReadOnly=False,bDefaultValueIsIgnored=False,bAdvancedView=False,)
+			   CustomProperties Pin (PinId=0EA9B030444B6103BD5640B0D5EBDD90,PinName="self",PinFriendlyName=NSLOCTEXT("K2Node", "Target", "Target"),PinType.PinCategory="object",PinType.PinSubCategoryObject=Class'/Script/Engine.StaticMeshActor',PinType.PinSubCategoryMemberReference=(),PinType.bIsArray=False,PinType.bIsReference=False,PinType.bIsConst=False,PinType.bIsWeakPointer=False,PersistentGuid=00000000000000000000000000000000,bHidden=True,bNotConnectable=False,bDefaultValueIsReadOnly=False,bDefaultValueIsIgnored=False,bAdvancedView=False,)
+			End Object
+			Begin Object Class=K2Node_VariableGet Name="K2Node_VariableGet_32"
+			   VariableReference=(MemberName="StaticMeshComponent",bSelfContext=True)
+			   NodePosX=176
+			   NodePosY=208
+			   NodeGuid=09DCEE7344FEAE148639658BA267B664
+			   CustomProperties Pin (PinId=8C2A8CB743E63ABCF573FC84E3867E6E,PinName="StaticMeshComponent",Direction="EGPD_Output",PinType.PinCategory="object",PinType.PinSubCategoryObject=Class'/Script/Engine.StaticMeshComponent',PinType.PinSubCategoryMemberReference=(),PinType.bIsArray=False,PinType.bIsReference=False,PinType.bIsConst=False,PinType.bIsWeakPointer=False,LinkedTo=(K2Node_CallFunction_0 391A69D247367ED3E62D4B869CDD282E,),PersistentGuid=00000000000000000000000000000000,bHidden=False,bNotConnectable=False,bDefaultValueIsReadOnly=False,bDefaultValueIsIgnored=False,bAdvancedView=False,)
+			   CustomProperties Pin (PinId=A90C5BE34B3B27D95B9B25B7D114DAF9,PinName="self",PinFriendlyName=NSLOCTEXT("K2Node", "Target", "Target"),PinType.PinCategory="object",PinType.PinSubCategoryObject=Class'/Script/Engine.StaticMeshActor',PinType.PinSubCategoryMemberReference=(),PinType.bIsArray=False,PinType.bIsReference=False,PinType.bIsConst=False,PinType.bIsWeakPointer=False,PersistentGuid=00000000000000000000000000000000,bHidden=True,bNotConnectable=False,bDefaultValueIsReadOnly=False,bDefaultValueIsIgnored=False,bAdvancedView=False,)
+			End Object
+			Begin Object Class=EdGraphNode_Comment Name="EdGraphNode_Comment_5"
+			   NodePosX=16
+			   NodePosY=16
+			   NodeWidth=1462
+			   NodeHeight=422
+			   NodeComment="Create the DMI for the water surface and apply it to the Static Mesh"
+			   NodeGuid=695F7E2E43A956EAACF86E82BDDF4462
+			End Object
+			END OBJECT
+	[/PARAM]
+[/OBJECT]
+[/EXCERPT:ConScript]
+
+[EXCERPT:ClearRT]	
+[OBJECT:Snippet]
+	[PARAMLITERAL:id]
+		ClearRT
+	[/PARAMLITERAL]
+	[PARAMLITERAL:type]
+		blueprint
+	[/PARAMLITERAL]
+	[PARAM:image]
+		![](Engine\Rendering\RenderTargets\BlueprintRenderTargets\HowTo\FluidSurface\FS_Clear_All_RT.png)
+	[/PARAM]
+	[PARAM:snippet]
+			BEGIN OBJECT
+				Begin Object Class=K2Node_Event Name="K2Node_Event_0"
+			   EventReference=(MemberParent=Class'/Script/Engine.Actor',MemberName="ReceiveBeginPlay")
+			   bOverrideFunction=True
+			   NodePosX=-1424
+			   NodePosY=-2384
+			   bCommentBubblePinned=True
+			   NodeGuid=B0FF665C447C29B1B16D798B6C8E79B9
+			   CustomProperties Pin (PinId=24AD5535476408C1DB299B8CC9AF8DCD,PinName="OutputDelegate",Direction="EGPD_Output",PinType.PinCategory="delegate",PinType.PinSubCategoryObject=None,PinType.PinSubCategoryMemberReference=(MemberParent=Class'/Script/Engine.Actor',MemberName="ReceiveBeginPlay"),PinType.bIsArray=False,PinType.bIsReference=False,PinType.bIsConst=False,PinType.bIsWeakPointer=False,PersistentGuid=00000000000000000000000000000000,bHidden=False,bNotConnectable=False,bDefaultValueIsReadOnly=False,bDefaultValueIsIgnored=False,bAdvancedView=False,)
+			   CustomProperties Pin (PinId=9605A65648EC47C19996AE9E878BCF7A,PinName="then",Direction="EGPD_Output",PinType.PinCategory="exec",PinType.PinSubCategoryObject=None,PinType.PinSubCategoryMemberReference=(),PinType.bIsArray=False,PinType.bIsReference=False,PinType.bIsConst=False,PinType.bIsWeakPointer=False,LinkedTo=(K2Node_CallFunction_107 5C2310FA44CE30E274CC61A28E553B5B,),PersistentGuid=00000000000000000000000000000000,bHidden=False,bNotConnectable=False,bDefaultValueIsReadOnly=False,bDefaultValueIsIgnored=False,bAdvancedView=False,)
+			End Object
+			Begin Object Class=EdGraphNode_Comment Name="EdGraphNode_Comment_38"
+			   NodePosX=-1440
+			   NodePosY=-2480
+			   NodeWidth=1216
+			   NodeHeight=352
+			   NodeComment="Clear all render targets of data"
+			   NodeGuid=18E7985549BCC508CB8D0897AB3A2A83
+			End Object
+			Begin Object Class=K2Node_CallFunction Name="K2Node_CallFunction_107"
+			   FunctionReference=(MemberParent=Class'/Script/Engine.KismetRenderingLibrary',MemberName="ClearRenderTarget2D")
+			   NodePosX=-1264
+			   NodePosY=-2384
+			   NodeGuid=88B7C1BA4822B25CC215D388DC0F62A8
+			   CustomProperties Pin (PinId=5C2310FA44CE30E274CC61A28E553B5B,PinName="execute",PinToolTip="\nExec",PinType.PinCategory="exec",PinType.PinSubCategoryObject=None,PinType.PinSubCategoryMemberReference=(),PinType.bIsArray=False,PinType.bIsReference=False,PinType.bIsConst=False,PinType.bIsWeakPointer=False,LinkedTo=(K2Node_Event_0 9605A65648EC47C19996AE9E878BCF7A,),PersistentGuid=00000000000000000000000000000000,bHidden=False,bNotConnectable=False,bDefaultValueIsReadOnly=False,bDefaultValueIsIgnored=False,bAdvancedView=False,)
+			   CustomProperties Pin (PinId=8BDA19DC4A485F2A61548FBC7468CFA9,PinName="then",PinToolTip="\nExec",Direction="EGPD_Output",PinType.PinCategory="exec",PinType.PinSubCategoryObject=None,PinType.PinSubCategoryMemberReference=(),PinType.bIsArray=False,PinType.bIsReference=False,PinType.bIsConst=False,PinType.bIsWeakPointer=False,LinkedTo=(K2Node_CallFunction_108 5C2310FA44CE30E274CC61A28E553B5B,),PersistentGuid=00000000000000000000000000000000,bHidden=False,bNotConnectable=False,bDefaultValueIsReadOnly=False,bDefaultValueIsIgnored=False,bAdvancedView=False,)
+			   CustomProperties Pin (PinId=EED1C67549D944089BCDA5908E3A6A24,PinName="self",PinFriendlyName=NSLOCTEXT("K2Node", "Target", "Target"),PinToolTip="Target\nKismet Rendering Library Reference",PinType.PinCategory="object",PinType.PinSubCategoryObject=Class'/Script/Engine.KismetRenderingLibrary',PinType.PinSubCategoryMemberReference=(),PinType.bIsArray=False,PinType.bIsReference=False,PinType.bIsConst=False,PinType.bIsWeakPointer=False,DefaultObject="/Script/Engine.Default__KismetRenderingLibrary",PersistentGuid=00000000000000000000000000000000,bHidden=True,bNotConnectable=False,bDefaultValueIsReadOnly=False,bDefaultValueIsIgnored=False,bAdvancedView=False,)
+			   CustomProperties Pin (PinId=7C31D3454D1A2CDA4039C9952708A9C6,PinName="WorldContextObject",PinToolTip="World Context Object\nObject Reference",PinType.PinCategory="object",PinType.PinSubCategoryObject=Class'/Script/CoreUObject.Object',PinType.PinSubCategoryMemberReference=(),PinType.bIsArray=False,PinType.bIsReference=False,PinType.bIsConst=False,PinType.bIsWeakPointer=False,PersistentGuid=00000000000000000000000000000000,bHidden=True,bNotConnectable=False,bDefaultValueIsReadOnly=False,bDefaultValueIsIgnored=False,bAdvancedView=False,)
+			   CustomProperties Pin (PinId=EC058DC94A5130C4CA239AB2FFA860F5,PinName="TextureRenderTarget",PinToolTip="Texture Render Target\nTexture Render Target 2D Reference",PinType.PinCategory="object",PinType.PinSubCategoryObject=Class'/Script/Engine.TextureRenderTarget2D',PinType.PinSubCategoryMemberReference=(),PinType.bIsArray=False,PinType.bIsReference=False,PinType.bIsConst=False,PinType.bIsWeakPointer=False,DefaultObject="/Game/ExampleContent/BlueprintFluidSurface/Height0.Height0",PersistentGuid=00000000000000000000000000000000,bHidden=False,bNotConnectable=False,bDefaultValueIsReadOnly=False,bDefaultValueIsIgnored=False,bAdvancedView=False,)
+			   CustomProperties Pin (PinId=164F9ACD44819BA4602A5384B875DCF4,PinName="ClearColor",PinToolTip="Clear Color\nLinear Color Structure",PinType.PinCategory="struct",PinType.PinSubCategoryObject=ScriptStruct'/Script/CoreUObject.LinearColor',PinType.PinSubCategoryMemberReference=(),PinType.bIsArray=False,PinType.bIsReference=False,PinType.bIsConst=False,PinType.bIsWeakPointer=False,PersistentGuid=00000000000000000000000000000000,bHidden=False,bNotConnectable=False,bDefaultValueIsReadOnly=False,bDefaultValueIsIgnored=False,bAdvancedView=False,)
+			End Object
+			Begin Object Class=K2Node_CallFunction Name="K2Node_CallFunction_108"
+			   FunctionReference=(MemberParent=Class'/Script/Engine.KismetRenderingLibrary',MemberName="ClearRenderTarget2D")
+			   NodePosX=-1024
+			   NodePosY=-2384
+			   NodeGuid=417DF2D8414BB5A244EDA2972B80701E
+			   CustomProperties Pin (PinId=5C2310FA44CE30E274CC61A28E553B5B,PinName="execute",PinToolTip="\nExec",PinType.PinCategory="exec",PinType.PinSubCategoryObject=None,PinType.PinSubCategoryMemberReference=(),PinType.bIsArray=False,PinType.bIsReference=False,PinType.bIsConst=False,PinType.bIsWeakPointer=False,LinkedTo=(K2Node_CallFunction_107 8BDA19DC4A485F2A61548FBC7468CFA9,),PersistentGuid=00000000000000000000000000000000,bHidden=False,bNotConnectable=False,bDefaultValueIsReadOnly=False,bDefaultValueIsIgnored=False,bAdvancedView=False,)
+			   CustomProperties Pin (PinId=8BDA19DC4A485F2A61548FBC7468CFA9,PinName="then",PinToolTip="\nExec",Direction="EGPD_Output",PinType.PinCategory="exec",PinType.PinSubCategoryObject=None,PinType.PinSubCategoryMemberReference=(),PinType.bIsArray=False,PinType.bIsReference=False,PinType.bIsConst=False,PinType.bIsWeakPointer=False,LinkedTo=(K2Node_CallFunction_106 5C2310FA44CE30E274CC61A28E553B5B,),PersistentGuid=00000000000000000000000000000000,bHidden=False,bNotConnectable=False,bDefaultValueIsReadOnly=False,bDefaultValueIsIgnored=False,bAdvancedView=False,)
+			   CustomProperties Pin (PinId=EED1C67549D944089BCDA5908E3A6A24,PinName="self",PinFriendlyName=NSLOCTEXT("K2Node", "Target", "Target"),PinToolTip="Target\nKismet Rendering Library Reference",PinType.PinCategory="object",PinType.PinSubCategoryObject=Class'/Script/Engine.KismetRenderingLibrary',PinType.PinSubCategoryMemberReference=(),PinType.bIsArray=False,PinType.bIsReference=False,PinType.bIsConst=False,PinType.bIsWeakPointer=False,DefaultObject="/Script/Engine.Default__KismetRenderingLibrary",PersistentGuid=00000000000000000000000000000000,bHidden=True,bNotConnectable=False,bDefaultValueIsReadOnly=False,bDefaultValueIsIgnored=False,bAdvancedView=False,)
+			   CustomProperties Pin (PinId=7C31D3454D1A2CDA4039C9952708A9C6,PinName="WorldContextObject",PinToolTip="World Context Object\nObject Reference",PinType.PinCategory="object",PinType.PinSubCategoryObject=Class'/Script/CoreUObject.Object',PinType.PinSubCategoryMemberReference=(),PinType.bIsArray=False,PinType.bIsReference=False,PinType.bIsConst=False,PinType.bIsWeakPointer=False,PersistentGuid=00000000000000000000000000000000,bHidden=True,bNotConnectable=False,bDefaultValueIsReadOnly=False,bDefaultValueIsIgnored=False,bAdvancedView=False,)
+			   CustomProperties Pin (PinId=EC058DC94A5130C4CA239AB2FFA860F5,PinName="TextureRenderTarget",PinToolTip="Texture Render Target\nTexture Render Target 2D Reference",PinType.PinCategory="object",PinType.PinSubCategoryObject=Class'/Script/Engine.TextureRenderTarget2D',PinType.PinSubCategoryMemberReference=(),PinType.bIsArray=False,PinType.bIsReference=False,PinType.bIsConst=False,PinType.bIsWeakPointer=False,DefaultObject="/Game/ExampleContent/BlueprintFluidSurface/Height1.Height1",PersistentGuid=00000000000000000000000000000000,bHidden=False,bNotConnectable=False,bDefaultValueIsReadOnly=False,bDefaultValueIsIgnored=False,bAdvancedView=False,)
+			   CustomProperties Pin (PinId=164F9ACD44819BA4602A5384B875DCF4,PinName="ClearColor",PinToolTip="Clear Color\nLinear Color Structure",PinType.PinCategory="struct",PinType.PinSubCategoryObject=ScriptStruct'/Script/CoreUObject.LinearColor',PinType.PinSubCategoryMemberReference=(),PinType.bIsArray=False,PinType.bIsReference=False,PinType.bIsConst=False,PinType.bIsWeakPointer=False,PersistentGuid=00000000000000000000000000000000,bHidden=False,bNotConnectable=False,bDefaultValueIsReadOnly=False,bDefaultValueIsIgnored=False,bAdvancedView=False,)
+			End Object
+			Begin Object Class=K2Node_CallFunction Name="K2Node_CallFunction_106"
+			   FunctionReference=(MemberParent=Class'/Script/Engine.KismetRenderingLibrary',MemberName="ClearRenderTarget2D")
+			   NodePosX=-784
+			   NodePosY=-2384
+			   NodeGuid=4ADE7FAF42E73E9FAC461E828A2606FB
+			   CustomProperties Pin (PinId=5C2310FA44CE30E274CC61A28E553B5B,PinName="execute",PinToolTip="\nExec",PinType.PinCategory="exec",PinType.PinSubCategoryObject=None,PinType.PinSubCategoryMemberReference=(),PinType.bIsArray=False,PinType.bIsReference=False,PinType.bIsConst=False,PinType.bIsWeakPointer=False,LinkedTo=(K2Node_CallFunction_108 8BDA19DC4A485F2A61548FBC7468CFA9,),PersistentGuid=00000000000000000000000000000000,bHidden=False,bNotConnectable=False,bDefaultValueIsReadOnly=False,bDefaultValueIsIgnored=False,bAdvancedView=False,)
+			   CustomProperties Pin (PinId=8BDA19DC4A485F2A61548FBC7468CFA9,PinName="then",PinToolTip="\nExec",Direction="EGPD_Output",PinType.PinCategory="exec",PinType.PinSubCategoryObject=None,PinType.PinSubCategoryMemberReference=(),PinType.bIsArray=False,PinType.bIsReference=False,PinType.bIsConst=False,PinType.bIsWeakPointer=False,LinkedTo=(K2Node_CallFunction_162 D16A6750404074A9A56440B940DA4390,),PersistentGuid=00000000000000000000000000000000,bHidden=False,bNotConnectable=False,bDefaultValueIsReadOnly=False,bDefaultValueIsIgnored=False,bAdvancedView=False,)
+			   CustomProperties Pin (PinId=EED1C67549D944089BCDA5908E3A6A24,PinName="self",PinFriendlyName=NSLOCTEXT("K2Node", "Target", "Target"),PinToolTip="Target\nKismet Rendering Library Reference",PinType.PinCategory="object",PinType.PinSubCategoryObject=Class'/Script/Engine.KismetRenderingLibrary',PinType.PinSubCategoryMemberReference=(),PinType.bIsArray=False,PinType.bIsReference=False,PinType.bIsConst=False,PinType.bIsWeakPointer=False,DefaultObject="/Script/Engine.Default__KismetRenderingLibrary",PersistentGuid=00000000000000000000000000000000,bHidden=True,bNotConnectable=False,bDefaultValueIsReadOnly=False,bDefaultValueIsIgnored=False,bAdvancedView=False,)
+			   CustomProperties Pin (PinId=7C31D3454D1A2CDA4039C9952708A9C6,PinName="WorldContextObject",PinToolTip="World Context Object\nObject Reference",PinType.PinCategory="object",PinType.PinSubCategoryObject=Class'/Script/CoreUObject.Object',PinType.PinSubCategoryMemberReference=(),PinType.bIsArray=False,PinType.bIsReference=False,PinType.bIsConst=False,PinType.bIsWeakPointer=False,PersistentGuid=00000000000000000000000000000000,bHidden=True,bNotConnectable=False,bDefaultValueIsReadOnly=False,bDefaultValueIsIgnored=False,bAdvancedView=False,)
+			   CustomProperties Pin (PinId=EC058DC94A5130C4CA239AB2FFA860F5,PinName="TextureRenderTarget",PinToolTip="Texture Render Target\nTexture Render Target 2D Reference",PinType.PinCategory="object",PinType.PinSubCategoryObject=Class'/Script/Engine.TextureRenderTarget2D',PinType.PinSubCategoryMemberReference=(),PinType.bIsArray=False,PinType.bIsReference=False,PinType.bIsConst=False,PinType.bIsWeakPointer=False,DefaultObject="/Game/ExampleContent/BlueprintFluidSurface/Height2.Height2",PersistentGuid=00000000000000000000000000000000,bHidden=False,bNotConnectable=False,bDefaultValueIsReadOnly=False,bDefaultValueIsIgnored=False,bAdvancedView=False,)
+			   CustomProperties Pin (PinId=164F9ACD44819BA4602A5384B875DCF4,PinName="ClearColor",PinToolTip="Clear Color\nLinear Color Structure",PinType.PinCategory="struct",PinType.PinSubCategoryObject=ScriptStruct'/Script/CoreUObject.LinearColor',PinType.PinSubCategoryMemberReference=(),PinType.bIsArray=False,PinType.bIsReference=False,PinType.bIsConst=False,PinType.bIsWeakPointer=False,PersistentGuid=00000000000000000000000000000000,bHidden=False,bNotConnectable=False,bDefaultValueIsReadOnly=False,bDefaultValueIsIgnored=False,bAdvancedView=False,)
+			End Object
+			Begin Object Class=K2Node_CallFunction Name="K2Node_CallFunction_162"
+			   FunctionReference=(MemberParent=Class'/Script/Engine.MaterialInstanceDynamic',MemberName="SetTextureParameterValue")
+			   NodePosX=-512
+			   NodePosY=-2400
+			   ErrorType=1
+			   ErrorMsg="Error This blueprint (self) is not a MaterialInstanceDynamic, therefore \' self \' must have a connection."
+			   NodeGuid=BCDAE8EB4797A378F90C1784301C25E0
+			   CustomProperties Pin (PinId=D16A6750404074A9A56440B940DA4390,PinName="execute",PinToolTip="\nExec",PinType.PinCategory="exec",PinType.PinSubCategoryObject=None,PinType.PinSubCategoryMemberReference=(),PinType.bIsArray=False,PinType.bIsReference=False,PinType.bIsConst=False,PinType.bIsWeakPointer=False,LinkedTo=(K2Node_CallFunction_106 8BDA19DC4A485F2A61548FBC7468CFA9,),PersistentGuid=00000000000000000000000000000000,bHidden=False,bNotConnectable=False,bDefaultValueIsReadOnly=False,bDefaultValueIsIgnored=False,bAdvancedView=False,)
+			   CustomProperties Pin (PinId=077792A64C10A1D106073E8F502680C6,PinName="then",PinToolTip="\nExec",Direction="EGPD_Output",PinType.PinCategory="exec",PinType.PinSubCategoryObject=None,PinType.PinSubCategoryMemberReference=(),PinType.bIsArray=False,PinType.bIsReference=False,PinType.bIsConst=False,PinType.bIsWeakPointer=False,PersistentGuid=00000000000000000000000000000000,bHidden=False,bNotConnectable=False,bDefaultValueIsReadOnly=False,bDefaultValueIsIgnored=False,bAdvancedView=False,)
+			   CustomProperties Pin (PinId=88F7A1134B6D0BC4419B5FBC4CFDE428,PinName="self",PinFriendlyName=NSLOCTEXT("K2Node", "Target", "Target"),PinToolTip="Target\nMaterial Instance Dynamic Reference",PinType.PinCategory="object",PinType.PinSubCategoryObject=Class'/Script/Engine.MaterialInstanceDynamic',PinType.PinSubCategoryMemberReference=(),PinType.bIsArray=False,PinType.bIsReference=False,PinType.bIsConst=False,PinType.bIsWeakPointer=False,LinkedTo=(K2Node_VariableGet_144 C6F718884A2C145C15A8D2AF9A5433B5,),PersistentGuid=00000000000000000000000000000000,bHidden=False,bNotConnectable=False,bDefaultValueIsReadOnly=False,bDefaultValueIsIgnored=False,bAdvancedView=False,)
+			   CustomProperties Pin (PinId=4AB1EED4499534AF3C301CB3F6D9D8B7,PinName="ParameterName",PinToolTip="Parameter Name\nName",PinType.PinCategory="name",PinType.PinSubCategoryObject=None,PinType.PinSubCategoryMemberReference=(),PinType.bIsArray=False,PinType.bIsReference=False,PinType.bIsConst=False,PinType.bIsWeakPointer=False,DefaultValue="HeightfieldNormal",AutogeneratedDefaultValue="None",PersistentGuid=00000000000000000000000000000000,bHidden=False,bNotConnectable=False,bDefaultValueIsReadOnly=False,bDefaultValueIsIgnored=False,bAdvancedView=False,)
+			   CustomProperties Pin (PinId=B49C168443CE292437D9F3B9A08A8F60,PinName="Value",PinToolTip="Value\nTexture Reference",PinType.PinCategory="object",PinType.PinSubCategoryObject=Class'/Script/Engine.Texture',PinType.PinSubCategoryMemberReference=(),PinType.bIsArray=False,PinType.bIsReference=False,PinType.bIsConst=False,PinType.bIsWeakPointer=False,DefaultObject="/Game/ExampleContent/BlueprintFluidSurface/HeightfieldNormal.HeightfieldNormal",PersistentGuid=00000000000000000000000000000000,bHidden=False,bNotConnectable=False,bDefaultValueIsReadOnly=False,bDefaultValueIsIgnored=False,bAdvancedView=False,)
+			End Object
+			Begin Object Class=K2Node_VariableGet Name="K2Node_VariableGet_144"
+			   VariableReference=(MemberName="WaterMID",MemberGuid=ABDD084C4292C4C32F8604AE0ADE9342,bSelfContext=True)
+			   NodePosX=-736
+			   NodePosY=-2224
+			   NodeGuid=C668A7014031F89F9A4AD1B182B0D894
+			   CustomProperties Pin (PinId=C6F718884A2C145C15A8D2AF9A5433B5,PinName="WaterMID",Direction="EGPD_Output",PinType.PinCategory="object",PinType.PinSubCategoryObject=Class'/Script/Engine.MaterialInstanceDynamic',PinType.PinSubCategoryMemberReference=(),PinType.bIsArray=False,PinType.bIsReference=False,PinType.bIsConst=False,PinType.bIsWeakPointer=False,LinkedTo=(K2Node_CallFunction_162 88F7A1134B6D0BC4419B5FBC4CFDE428,),PersistentGuid=00000000000000000000000000000000,bHidden=False,bNotConnectable=False,bDefaultValueIsReadOnly=False,bDefaultValueIsIgnored=False,bAdvancedView=False,)
+			   CustomProperties Pin (PinId=7D6E1FE544CE93F1C2C0EAAAA5F7C1F3,PinName="self",PinFriendlyName=NSLOCTEXT("K2Node", "Target", "Target"),PinType.PinCategory="object",PinType.PinSubCategoryObject=BlueprintGeneratedClass'/Game/ExampleContent/BlueprintFluidSurface/WaterSurface_BP.WaterSurface_BP_C',PinType.PinSubCategoryMemberReference=(),PinType.bIsArray=False,PinType.bIsReference=False,PinType.bIsConst=False,PinType.bIsWeakPointer=False,PersistentGuid=00000000000000000000000000000000,bHidden=True,bNotConnectable=False,bDefaultValueIsReadOnly=False,bDefaultValueIsIgnored=False,bAdvancedView=False,)
+			End Object
+			END OBJECT
+	[/PARAM]
+[/OBJECT]
+[/EXCERPT:ClearRT]
+
+[EXCERPT:MouseDown]	
+[OBJECT:Snippet]
+	[PARAMLITERAL:id]
+		MouseDown
+	[/PARAMLITERAL]
+	[PARAMLITERAL:type]
+		blueprint
+	[/PARAMLITERAL]
+	[PARAM:image]
+		![](Engine\Rendering\RenderTargets\BlueprintRenderTargets\HowTo\FluidSurface\FS_Check_If_Left_Mouse.png)
+	[/PARAM]
+	[PARAM:snippet]
+			BEGIN OBJECT
+							Begin Object Class=K2Node_InputKey Name="K2Node_InputKey_0"
+			   InputKey=LeftMouseButton
+			   NodePosX=-1488
+			   NodePosY=-1728
+			   NodeGuid=D62C2B5543959E8D66DC14B9E9A98605
+			   CustomProperties Pin (PinId=4D5FE2B549E275EDF0E659896A9261E4,PinName="Pressed",Direction="EGPD_Output",PinType.PinCategory="exec",PinType.PinSubCategoryObject=None,PinType.PinSubCategoryMemberReference=(),PinType.bIsArray=False,PinType.bIsReference=False,PinType.bIsConst=False,PinType.bIsWeakPointer=False,LinkedTo=(K2Node_VariableSet_144 DF34F0B548E46FCCB2AB54856F7CFFC5,),PersistentGuid=00000000000000000000000000000000,bHidden=False,bNotConnectable=False,bDefaultValueIsReadOnly=False,bDefaultValueIsIgnored=False,bAdvancedView=False,)
+			   CustomProperties Pin (PinId=B150F43D4C370A3041E5B9A2DD8FC834,PinName="Released",Direction="EGPD_Output",PinType.PinCategory="exec",PinType.PinSubCategoryObject=None,PinType.PinSubCategoryMemberReference=(),PinType.bIsArray=False,PinType.bIsReference=False,PinType.bIsConst=False,PinType.bIsWeakPointer=False,LinkedTo=(K2Node_VariableSet_155 9312E8C94B08C9FAF76638810A0F61A1,),PersistentGuid=00000000000000000000000000000000,bHidden=False,bNotConnectable=False,bDefaultValueIsReadOnly=False,bDefaultValueIsIgnored=False,bAdvancedView=False,)
+			   CustomProperties Pin (PinId=43D7BA3045BD03142E77E597FD7EAD86,PinName="Key",Direction="EGPD_Output",PinType.PinCategory="struct",PinType.PinSubCategoryObject=ScriptStruct'/Script/InputCore.Key',PinType.PinSubCategoryMemberReference=(),PinType.bIsArray=False,PinType.bIsReference=False,PinType.bIsConst=False,PinType.bIsWeakPointer=False,DefaultValue="AnyKey",PersistentGuid=00000000000000000000000000000000,bHidden=False,bNotConnectable=False,bDefaultValueIsReadOnly=False,bDefaultValueIsIgnored=False,bAdvancedView=False,)
+			End Object
+			Begin Object Class=K2Node_VariableSet Name="K2Node_VariableSet_144"
+			   VariableReference=(MemberName="MouseDown",MemberGuid=209E1D534C15EFFA81760DA2C7F9D56C,bSelfContext=True)
+			   NodePosX=-1200
+			   NodePosY=-1712
+			   NodeGuid=FF7233234B166BB44415619BBA2274E0
+			   CustomProperties Pin (PinId=DF34F0B548E46FCCB2AB54856F7CFFC5,PinName="execute",PinType.PinCategory="exec",PinType.PinSubCategoryObject=None,PinType.PinSubCategoryMemberReference=(),PinType.bIsArray=False,PinType.bIsReference=False,PinType.bIsConst=False,PinType.bIsWeakPointer=False,LinkedTo=(K2Node_InputKey_0 4D5FE2B549E275EDF0E659896A9261E4,),PersistentGuid=00000000000000000000000000000000,bHidden=False,bNotConnectable=False,bDefaultValueIsReadOnly=False,bDefaultValueIsIgnored=False,bAdvancedView=False,)
+			   CustomProperties Pin (PinId=A68EDFFF4BFBADBD6B1C60AB8EEE7A3D,PinName="then",Direction="EGPD_Output",PinType.PinCategory="exec",PinType.PinSubCategoryObject=None,PinType.PinSubCategoryMemberReference=(),PinType.bIsArray=False,PinType.bIsReference=False,PinType.bIsConst=False,PinType.bIsWeakPointer=False,PersistentGuid=00000000000000000000000000000000,bHidden=False,bNotConnectable=False,bDefaultValueIsReadOnly=False,bDefaultValueIsIgnored=False,bAdvancedView=False,)
+			   CustomProperties Pin (PinId=8F7952DE47C7BACDD7565587C283D7EE,PinName="MouseDown",PinType.PinCategory="bool",PinType.PinSubCategoryObject=None,PinType.PinSubCategoryMemberReference=(),PinType.bIsArray=False,PinType.bIsReference=False,PinType.bIsConst=False,PinType.bIsWeakPointer=False,DefaultValue="true",AutogeneratedDefaultValue="false",PersistentGuid=00000000000000000000000000000000,bHidden=False,bNotConnectable=False,bDefaultValueIsReadOnly=False,bDefaultValueIsIgnored=False,bAdvancedView=False,)
+			   CustomProperties Pin (PinId=6799AEA64465D17287DE96888FB9C1D6,PinName="Output_Get",PinToolTip="Retrieves the value of the variable, can use instead of a separate Get node",Direction="EGPD_Output",PinType.PinCategory="bool",PinType.PinSubCategoryObject=None,PinType.PinSubCategoryMemberReference=(),PinType.bIsArray=False,PinType.bIsReference=False,PinType.bIsConst=False,PinType.bIsWeakPointer=False,DefaultValue="false",AutogeneratedDefaultValue="false",PersistentGuid=00000000000000000000000000000000,bHidden=False,bNotConnectable=False,bDefaultValueIsReadOnly=False,bDefaultValueIsIgnored=False,bAdvancedView=False,)
+			   CustomProperties Pin (PinId=9F765A6742545A9CD3801AB2E5FE499D,PinName="self",PinFriendlyName=NSLOCTEXT("K2Node", "Target", "Target"),PinType.PinCategory="object",PinType.PinSubCategoryObject=BlueprintGeneratedClass'/Game/ExampleContent/BlueprintFluidSurface/WaterSurface_BP.WaterSurface_BP_C',PinType.PinSubCategoryMemberReference=(),PinType.bIsArray=False,PinType.bIsReference=False,PinType.bIsConst=False,PinType.bIsWeakPointer=False,PersistentGuid=00000000000000000000000000000000,bHidden=True,bNotConnectable=False,bDefaultValueIsReadOnly=False,bDefaultValueIsIgnored=False,bAdvancedView=False,)
+			End Object
+			Begin Object Class=K2Node_VariableSet Name="K2Node_VariableSet_155"
+			   VariableReference=(MemberName="MouseDown",MemberGuid=209E1D534C15EFFA81760DA2C7F9D56C,bSelfContext=True)
+			   NodePosX=-1200
+			   NodePosY=-1600
+			   NodeGuid=ED1FF2C74F0CA35956392A9237BFE7F7
+			   CustomProperties Pin (PinId=9312E8C94B08C9FAF76638810A0F61A1,PinName="execute",PinType.PinCategory="exec",PinType.PinSubCategoryObject=None,PinType.PinSubCategoryMemberReference=(),PinType.bIsArray=False,PinType.bIsReference=False,PinType.bIsConst=False,PinType.bIsWeakPointer=False,LinkedTo=(K2Node_InputKey_0 B150F43D4C370A3041E5B9A2DD8FC834,),PersistentGuid=00000000000000000000000000000000,bHidden=False,bNotConnectable=False,bDefaultValueIsReadOnly=False,bDefaultValueIsIgnored=False,bAdvancedView=False,)
+			   CustomProperties Pin (PinId=CB1027A1495B8948A22EB58160AA4C3E,PinName="then",Direction="EGPD_Output",PinType.PinCategory="exec",PinType.PinSubCategoryObject=None,PinType.PinSubCategoryMemberReference=(),PinType.bIsArray=False,PinType.bIsReference=False,PinType.bIsConst=False,PinType.bIsWeakPointer=False,PersistentGuid=00000000000000000000000000000000,bHidden=False,bNotConnectable=False,bDefaultValueIsReadOnly=False,bDefaultValueIsIgnored=False,bAdvancedView=False,)
+			   CustomProperties Pin (PinId=3754DE984E6A106CC45C95831AEDF03A,PinName="MouseDown",PinType.PinCategory="bool",PinType.PinSubCategoryObject=None,PinType.PinSubCategoryMemberReference=(),PinType.bIsArray=False,PinType.bIsReference=False,PinType.bIsConst=False,PinType.bIsWeakPointer=False,DefaultValue="false",AutogeneratedDefaultValue="false",PersistentGuid=00000000000000000000000000000000,bHidden=False,bNotConnectable=False,bDefaultValueIsReadOnly=False,bDefaultValueIsIgnored=False,bAdvancedView=False,)
+			   CustomProperties Pin (PinId=CD9E960E4A4357C090CAFF8D5F8D9AB9,PinName="Output_Get",PinToolTip="Retrieves the value of the variable, can use instead of a separate Get node",Direction="EGPD_Output",PinType.PinCategory="bool",PinType.PinSubCategoryObject=None,PinType.PinSubCategoryMemberReference=(),PinType.bIsArray=False,PinType.bIsReference=False,PinType.bIsConst=False,PinType.bIsWeakPointer=False,DefaultValue="false",AutogeneratedDefaultValue="false",PersistentGuid=00000000000000000000000000000000,bHidden=False,bNotConnectable=False,bDefaultValueIsReadOnly=False,bDefaultValueIsIgnored=False,bAdvancedView=False,)
+			   CustomProperties Pin (PinId=02064B664DAD8D023DD5518F1B238963,PinName="self",PinFriendlyName=NSLOCTEXT("K2Node", "Target", "Target"),PinType.PinCategory="object",PinType.PinSubCategoryObject=BlueprintGeneratedClass'/Game/ExampleContent/BlueprintFluidSurface/WaterSurface_BP.WaterSurface_BP_C',PinType.PinSubCategoryMemberReference=(),PinType.bIsArray=False,PinType.bIsReference=False,PinType.bIsConst=False,PinType.bIsWeakPointer=False,PersistentGuid=00000000000000000000000000000000,bHidden=True,bNotConnectable=False,bDefaultValueIsReadOnly=False,bDefaultValueIsIgnored=False,bAdvancedView=False,)
+			End Object
+			Begin Object Class=EdGraphNode_Comment Name="EdGraphNode_Comment_6"
+			   NodePosX=-1520
+			   NodePosY=-1776
+			   NodeWidth=581
+			   NodeHeight=303
+			   NodeComment="Check to see if Left Mouse Button is held down"
+			   NodeGuid=B1C0857840B90CB6FB5219A7C99CCCF7
+			End Object
+			END OBJECT
+	[/PARAM]
+[/OBJECT]
+[/EXCERPT:MouseDown]
+
+[EXCERPT:TakeDamage]	
+[OBJECT:Snippet]
+	[PARAMLITERAL:id]
+		TakeDamage
+	[/PARAMLITERAL]
+	[PARAMLITERAL:type]
+		blueprint
+	[/PARAMLITERAL]
+	[PARAM:image]
+		![](Engine\Rendering\RenderTargets\BlueprintRenderTargets\HowTo\FluidSurface\FS_Take_Damage_Force.png)
+	[/PARAM]
+	[PARAM:snippet]
+			BEGIN OBJECT
+							Begin Object Class=K2Node_Event Name="K2Node_Event_51"
+			   EventReference=(MemberParent=Class'/Script/Engine.Actor',MemberName="ReceivePointDamage")
+			   bOverrideFunction=True
+			   NodePosX=240
+			   NodePosY=-3216
+			   NodeGuid=01D175A9404CA37FD979D5AAA296BE8B
+			   CustomProperties Pin (PinId=CEA11B954E55C6E301F0DDB73B737331,PinName="OutputDelegate",Direction="EGPD_Output",PinType.PinCategory="delegate",PinType.PinSubCategoryObject=None,PinType.PinSubCategoryMemberReference=(MemberParent=Class'/Script/Engine.Actor',MemberName="ReceivePointDamage"),PinType.bIsArray=False,PinType.bIsReference=False,PinType.bIsConst=False,PinType.bIsWeakPointer=False,PersistentGuid=00000000000000000000000000000000,bHidden=False,bNotConnectable=False,bDefaultValueIsReadOnly=False,bDefaultValueIsIgnored=False,bAdvancedView=False,)
+			   CustomProperties Pin (PinId=7BC1398142AE6DCA11411BBF6C1C7F4E,PinName="then",Direction="EGPD_Output",PinType.PinCategory="exec",PinType.PinSubCategoryObject=None,PinType.PinSubCategoryMemberReference=(),PinType.bIsArray=False,PinType.bIsReference=False,PinType.bIsConst=False,PinType.bIsWeakPointer=False,LinkedTo=(K2Node_CallFunction_1529 4C40F4A0473F9F4E714A4F970B7F10DB,),PersistentGuid=00000000000000000000000000000000,bHidden=False,bNotConnectable=False,bDefaultValueIsReadOnly=False,bDefaultValueIsIgnored=False,bAdvancedView=False,)
+			   CustomProperties Pin (PinId=60CC33914AE4CFBF54C5D2807C150F9C,PinName="Damage",PinToolTip="Damage\nFloat",Direction="EGPD_Output",PinType.PinCategory="float",PinType.PinSubCategoryObject=None,PinType.PinSubCategoryMemberReference=(),PinType.bIsArray=False,PinType.bIsReference=False,PinType.bIsConst=False,PinType.bIsWeakPointer=False,DefaultValue="0.0",AutogeneratedDefaultValue="0.0",PersistentGuid=00000000000000000000000000000000,bHidden=False,bNotConnectable=False,bDefaultValueIsReadOnly=False,bDefaultValueIsIgnored=False,bAdvancedView=False,)
+			   CustomProperties Pin (PinId=16E9F4B64E717E33F3F8AB9D1FE60265,PinName="DamageType",PinToolTip="Damage Type\nDamage Type Reference",Direction="EGPD_Output",PinType.PinCategory="object",PinType.PinSubCategoryObject=Class'/Script/Engine.DamageType',PinType.PinSubCategoryMemberReference=(),PinType.bIsArray=False,PinType.bIsReference=False,PinType.bIsConst=True,PinType.bIsWeakPointer=False,PersistentGuid=00000000000000000000000000000000,bHidden=False,bNotConnectable=False,bDefaultValueIsReadOnly=False,bDefaultValueIsIgnored=False,bAdvancedView=False,)
+			   CustomProperties Pin (PinId=640542634B4797AB7140EAB2A04EF819,PinName="HitLocation",PinToolTip="Hit Location\nVector ",Direction="EGPD_Output",PinType.PinCategory="struct",PinType.PinSubCategoryObject=ScriptStruct'/Script/CoreUObject.Vector',PinType.PinSubCategoryMemberReference=(),PinType.bIsArray=False,PinType.bIsReference=False,PinType.bIsConst=False,PinType.bIsWeakPointer=False,AutogeneratedDefaultValue="0, 0, 0",LinkedTo=(K2Node_CallFunction_1530 850583744D608EF2952864ABE9A12E9E,),PersistentGuid=00000000000000000000000000000000,bHidden=False,bNotConnectable=False,bDefaultValueIsReadOnly=False,bDefaultValueIsIgnored=False,bAdvancedView=False,)
+			   CustomProperties Pin (PinId=977D30BB43F35C9B14B82D96656AEE48,PinName="HitNormal",PinToolTip="Hit Normal\nVector ",Direction="EGPD_Output",PinType.PinCategory="struct",PinType.PinSubCategoryObject=ScriptStruct'/Script/CoreUObject.Vector',PinType.PinSubCategoryMemberReference=(),PinType.bIsArray=False,PinType.bIsReference=False,PinType.bIsConst=False,PinType.bIsWeakPointer=False,DefaultValue="0, 0, 0",AutogeneratedDefaultValue="0, 0, 0",PersistentGuid=00000000000000000000000000000000,bHidden=False,bNotConnectable=False,bDefaultValueIsReadOnly=False,bDefaultValueIsIgnored=False,bAdvancedView=False,)
+			   CustomProperties Pin (PinId=F3CA2D394CE3DE4376083189EE70EE01,PinName="HitComponent",PinToolTip="Hit Component\nPrimitive Component Reference",Direction="EGPD_Output",PinType.PinCategory="object",PinType.PinSubCategoryObject=Class'/Script/Engine.PrimitiveComponent',PinType.PinSubCategoryMemberReference=(),PinType.bIsArray=False,PinType.bIsReference=False,PinType.bIsConst=False,PinType.bIsWeakPointer=False,PersistentGuid=00000000000000000000000000000000,bHidden=False,bNotConnectable=False,bDefaultValueIsReadOnly=False,bDefaultValueIsIgnored=False,bAdvancedView=False,)
+			   CustomProperties Pin (PinId=E09A9FB14D431DADE2BB769CBACDB17A,PinName="BoneName",PinToolTip="Bone Name\nName",Direction="EGPD_Output",PinType.PinCategory="name",PinType.PinSubCategoryObject=None,PinType.PinSubCategoryMemberReference=(),PinType.bIsArray=False,PinType.bIsReference=False,PinType.bIsConst=False,PinType.bIsWeakPointer=False,DefaultValue="None",AutogeneratedDefaultValue="None",PersistentGuid=00000000000000000000000000000000,bHidden=False,bNotConnectable=False,bDefaultValueIsReadOnly=False,bDefaultValueIsIgnored=False,bAdvancedView=False,)
+			   CustomProperties Pin (PinId=80A5C0F84FD774A9E91907821B24714C,PinName="ShotFromDirection",PinToolTip="Shot from Direction\nVector ",Direction="EGPD_Output",PinType.PinCategory="struct",PinType.PinSubCategoryObject=ScriptStruct'/Script/CoreUObject.Vector',PinType.PinSubCategoryMemberReference=(),PinType.bIsArray=False,PinType.bIsReference=False,PinType.bIsConst=False,PinType.bIsWeakPointer=False,DefaultValue="0, 0, 0",AutogeneratedDefaultValue="0, 0, 0",PersistentGuid=00000000000000000000000000000000,bHidden=False,bNotConnectable=False,bDefaultValueIsReadOnly=False,bDefaultValueIsIgnored=False,bAdvancedView=False,)
+			   CustomProperties Pin (PinId=5D9B6A9F49DC7BCFD58603AB5847E975,PinName="InstigatedBy",PinToolTip="Instigated By\nController Reference",Direction="EGPD_Output",PinType.PinCategory="object",PinType.PinSubCategoryObject=Class'/Script/Engine.Controller',PinType.PinSubCategoryMemberReference=(),PinType.bIsArray=False,PinType.bIsReference=False,PinType.bIsConst=False,PinType.bIsWeakPointer=False,PersistentGuid=00000000000000000000000000000000,bHidden=False,bNotConnectable=False,bDefaultValueIsReadOnly=False,bDefaultValueIsIgnored=False,bAdvancedView=False,)
+			   CustomProperties Pin (PinId=CC306F3F459174FB247CD18B4B3919C5,PinName="DamageCauser",PinToolTip="Damage Causer\nActor Reference",Direction="EGPD_Output",PinType.PinCategory="object",PinType.PinSubCategoryObject=Class'/Script/Engine.Actor',PinType.PinSubCategoryMemberReference=(),PinType.bIsArray=False,PinType.bIsReference=False,PinType.bIsConst=False,PinType.bIsWeakPointer=False,PersistentGuid=00000000000000000000000000000000,bHidden=False,bNotConnectable=False,bDefaultValueIsReadOnly=False,bDefaultValueIsIgnored=False,bAdvancedView=False,)
+			End Object
+			Begin Object Class=K2Node_CallFunction Name="K2Node_CallFunction_1529"
+			   FunctionReference=(MemberName="GetHeightRT",MemberGuid=2EEA338F4CA878E7228A10BB2DD63E00,bSelfContext=True)
+			   NodePosX=624
+			   NodePosY=-3231
+			   NodeGuid=CF7D58C6400DA035CBD0DF95DDF7D103
+			   CustomProperties Pin (PinId=4C40F4A0473F9F4E714A4F970B7F10DB,PinName="execute",PinToolTip="\nExec",PinType.PinCategory="exec",PinType.PinSubCategoryObject=None,PinType.PinSubCategoryMemberReference=(),PinType.bIsArray=False,PinType.bIsReference=False,PinType.bIsConst=False,PinType.bIsWeakPointer=False,LinkedTo=(K2Node_Event_51 7BC1398142AE6DCA11411BBF6C1C7F4E,),PersistentGuid=00000000000000000000000000000000,bHidden=False,bNotConnectable=False,bDefaultValueIsReadOnly=False,bDefaultValueIsIgnored=False,bAdvancedView=False,)
+			   CustomProperties Pin (PinId=BF80119E4AE88655005DD5AD3A44C53F,PinName="then",PinToolTip="\nExec",Direction="EGPD_Output",PinType.PinCategory="exec",PinType.PinSubCategoryObject=None,PinType.PinSubCategoryMemberReference=(),PinType.bIsArray=False,PinType.bIsReference=False,PinType.bIsConst=False,PinType.bIsWeakPointer=False,LinkedTo=(K2Node_CallMaterialParameterCollectionFunction_116 604EF317495C13D515FA428315721F21,),PersistentGuid=00000000000000000000000000000000,bHidden=False,bNotConnectable=False,bDefaultValueIsReadOnly=False,bDefaultValueIsIgnored=False,bAdvancedView=False,)
+			   CustomProperties Pin (PinId=8D86142F46B816E5C261C48AA899B738,PinName="self",PinFriendlyName=NSLOCTEXT("K2Node", "Target", "Target"),PinToolTip="Target\nSelf Reference",PinType.PinCategory="object",PinType.PinSubCategory="self",PinType.PinSubCategoryObject=None,PinType.PinSubCategoryMemberReference=(),PinType.bIsArray=False,PinType.bIsReference=False,PinType.bIsConst=False,PinType.bIsWeakPointer=False,PersistentGuid=00000000000000000000000000000000,bHidden=False,bNotConnectable=False,bDefaultValueIsReadOnly=False,bDefaultValueIsIgnored=False,bAdvancedView=False,)
+			   CustomProperties Pin (PinId=B87A9D034F088A639A6C739B70AE43A7,PinName="Index",PinToolTip="Index\nInteger",PinType.PinCategory="int",PinType.PinSubCategoryObject=None,PinType.PinSubCategoryMemberReference=(),PinType.bIsArray=False,PinType.bIsReference=False,PinType.bIsConst=False,PinType.bIsWeakPointer=False,DefaultValue="0",AutogeneratedDefaultValue="0",LinkedTo=(K2Node_VariableGet_475 E9132FD4449E60190B16CB91811A684A,),PersistentGuid=00000000000000000000000000000000,bHidden=False,bNotConnectable=False,bDefaultValueIsReadOnly=False,bDefaultValueIsIgnored=False,bAdvancedView=False,)
+			   CustomProperties Pin (PinId=E9A0065B436A442D29A940B6587482B6,PinName="HeightRT",PinToolTip="Height RT\nTexture Render Target 2D Reference",Direction="EGPD_Output",PinType.PinCategory="object",PinType.PinSubCategoryObject=Class'/Script/Engine.TextureRenderTarget2D',PinType.PinSubCategoryMemberReference=(),PinType.bIsArray=False,PinType.bIsReference=False,PinType.bIsConst=False,PinType.bIsWeakPointer=False,LinkedTo=(K2Node_CallFunction_1535 B57C770C4D827E19FDE1F885B23E018C,),PersistentGuid=00000000000000000000000000000000,bHidden=False,bNotConnectable=False,bDefaultValueIsReadOnly=False,bDefaultValueIsIgnored=False,bAdvancedView=False,)
+			End Object
+			Begin Object Class=K2Node_VariableGet Name="K2Node_VariableGet_475"
+			   VariableReference=(MemberName="HeightState",MemberGuid=8E10E24F4154A224C2065A89B367BD62,bSelfContext=True)
+			   NodePosX=464
+			   NodePosY=-3136
+			   NodeGuid=B06DEE7F4D435622C9DEFFAFE3DCB2FD
+			   CustomProperties Pin (PinId=E9132FD4449E60190B16CB91811A684A,PinName="HeightState",Direction="EGPD_Output",PinType.PinCategory="int",PinType.PinSubCategoryObject=None,PinType.PinSubCategoryMemberReference=(),PinType.bIsArray=False,PinType.bIsReference=False,PinType.bIsConst=False,PinType.bIsWeakPointer=False,AutogeneratedDefaultValue="0",LinkedTo=(K2Node_CallFunction_1529 B87A9D034F088A639A6C739B70AE43A7,),PersistentGuid=00000000000000000000000000000000,bHidden=False,bNotConnectable=False,bDefaultValueIsReadOnly=False,bDefaultValueIsIgnored=False,bAdvancedView=False,)
+			   CustomProperties Pin (PinId=A131864E4604B336093FCCA7070527A1,PinName="self",PinFriendlyName=NSLOCTEXT("K2Node", "Target", "Target"),PinType.PinCategory="object",PinType.PinSubCategoryObject=BlueprintGeneratedClass'/Game/ExampleContent/BlueprintFluidSurface/WaterSurface_BP.WaterSurface_BP_C',PinType.PinSubCategoryMemberReference=(),PinType.bIsArray=False,PinType.bIsReference=False,PinType.bIsConst=False,PinType.bIsWeakPointer=False,PersistentGuid=00000000000000000000000000000000,bHidden=True,bNotConnectable=False,bDefaultValueIsReadOnly=False,bDefaultValueIsIgnored=False,bAdvancedView=False,)
+			End Object
+			Begin Object Class=K2Node_CallFunction Name="K2Node_CallFunction_1530"
+			   bIsPureFunc=True
+			   FunctionReference=(MemberParent=Class'/Script/Engine.KismetMathLibrary',MemberName="TransformLocation")
+			   NodePosX=1264
+			   NodePosY=-2918
+			   NodeGuid=72E8C78F426697EE13C365A81C189973
+			   CustomProperties Pin (PinId=CEB3CA4647BBC3C737152680770F0BE1,PinName="self",PinFriendlyName=NSLOCTEXT("K2Node", "Target", "Target"),PinToolTip="Target\nKismet Math Library Reference",PinType.PinCategory="object",PinType.PinSubCategoryObject=Class'/Script/Engine.KismetMathLibrary',PinType.PinSubCategoryMemberReference=(),PinType.bIsArray=False,PinType.bIsReference=False,PinType.bIsConst=False,PinType.bIsWeakPointer=False,DefaultObject="/Script/Engine.Default__KismetMathLibrary",PersistentGuid=00000000000000000000000000000000,bHidden=True,bNotConnectable=False,bDefaultValueIsReadOnly=False,bDefaultValueIsIgnored=False,bAdvancedView=False,)
+			   CustomProperties Pin (PinId=8A4E9BDC43985D182203EEBE4D64ABED,PinName="T",PinToolTip="T\nTransform  (by ref)",PinType.PinCategory="struct",PinType.PinSubCategoryObject=ScriptStruct'/Script/CoreUObject.Transform',PinType.PinSubCategoryMemberReference=(),PinType.bIsArray=False,PinType.bIsReference=True,PinType.bIsConst=True,PinType.bIsWeakPointer=False,LinkedTo=(K2Node_CallFunction_1532 BDA890074B460518CA47E0B5319A61D4,),PersistentGuid=00000000000000000000000000000000,bHidden=False,bNotConnectable=False,bDefaultValueIsReadOnly=False,bDefaultValueIsIgnored=True,bAdvancedView=False,)
+			   CustomProperties Pin (PinId=850583744D608EF2952864ABE9A12E9E,PinName="Location",PinToolTip="Location\nVector ",PinType.PinCategory="struct",PinType.PinSubCategoryObject=ScriptStruct'/Script/CoreUObject.Vector',PinType.PinSubCategoryMemberReference=(),PinType.bIsArray=False,PinType.bIsReference=False,PinType.bIsConst=False,PinType.bIsWeakPointer=False,DefaultValue="0, 0, 0",AutogeneratedDefaultValue="0, 0, 0",LinkedTo=(K2Node_Event_51 640542634B4797AB7140EAB2A04EF819,),PersistentGuid=00000000000000000000000000000000,bHidden=False,bNotConnectable=False,bDefaultValueIsReadOnly=False,bDefaultValueIsIgnored=False,bAdvancedView=False,)
+			   CustomProperties Pin (PinId=7335087343EF48C6242BCDBFC09F2B35,PinName="ReturnValue",PinToolTip="Return Value\nVector ",Direction="EGPD_Output",PinType.PinCategory="struct",PinType.PinSubCategoryObject=ScriptStruct'/Script/CoreUObject.Vector',PinType.PinSubCategoryMemberReference=(),PinType.bIsArray=False,PinType.bIsReference=False,PinType.bIsConst=False,PinType.bIsWeakPointer=False,AutogeneratedDefaultValue="0, 0, 0",LinkedTo=(K2Node_CallFunction_1534 694D50D04A4FE8E7DC8CA689C5CDBE67,),PersistentGuid=00000000000000000000000000000000,bHidden=False,bNotConnectable=False,bDefaultValueIsReadOnly=False,bDefaultValueIsIgnored=False,bAdvancedView=False,)
+			End Object
+			Begin Object Class=K2Node_CallFunction Name="K2Node_CallFunction_1531"
+			   bIsPureFunc=True
+			   bIsConstFunc=True
+			   FunctionReference=(MemberParent=Class'/Script/Engine.SceneComponent',MemberName="K2_GetComponentToWorld")
+			   NodePosX=768
+			   NodePosY=-2934
+			   NodeGuid=EE1730EA47E906A22FD8AABC01F2D1AE
+			   CustomProperties Pin (PinId=5A1010634C4DE3EAF3EF409F23EC1215,PinName="self",PinFriendlyName=NSLOCTEXT("K2Node", "Target", "Target"),PinType.PinCategory="object",PinType.PinSubCategoryObject=Class'/Script/Engine.SceneComponent',PinType.PinSubCategoryMemberReference=(),PinType.bIsArray=False,PinType.bIsReference=False,PinType.bIsConst=False,PinType.bIsWeakPointer=False,LinkedTo=(K2Node_VariableGet_476 78F118CD49FC76D027795D843DFE71FF,),PersistentGuid=00000000000000000000000000000000,bHidden=False,bNotConnectable=False,bDefaultValueIsReadOnly=False,bDefaultValueIsIgnored=False,bAdvancedView=False,)
+			   CustomProperties Pin (PinId=5150BE354EF59F455E6CA988E1351299,PinName="ReturnValue",Direction="EGPD_Output",PinType.PinCategory="struct",PinType.PinSubCategoryObject=ScriptStruct'/Script/CoreUObject.Transform',PinType.PinSubCategoryMemberReference=(),PinType.bIsArray=False,PinType.bIsReference=False,PinType.bIsConst=False,PinType.bIsWeakPointer=False,LinkedTo=(K2Node_CallFunction_1532 557C12A64EC70024F7CC519F97D5D932,),PersistentGuid=00000000000000000000000000000000,bHidden=False,bNotConnectable=False,bDefaultValueIsReadOnly=False,bDefaultValueIsIgnored=False,bAdvancedView=False,)
+			End Object
+			Begin Object Class=K2Node_CallFunction Name="K2Node_CallFunction_1532"
+			   bIsPureFunc=True
+			   FunctionReference=(MemberParent=Class'/Script/Engine.KismetMathLibrary',MemberName="InvertTransform")
+			   NodePosX=1008
+			   NodePosY=-2918
+			   NodeGuid=2CE014EB4EA136DD0D8E5E899B99A773
+			   CustomProperties Pin (PinId=C4A955704E1399074BD903B945BE2245,PinName="self",PinFriendlyName=NSLOCTEXT("K2Node", "Target", "Target"),PinToolTip="Target\nKismet Math Library Reference",PinType.PinCategory="object",PinType.PinSubCategoryObject=Class'/Script/Engine.KismetMathLibrary',PinType.PinSubCategoryMemberReference=(),PinType.bIsArray=False,PinType.bIsReference=False,PinType.bIsConst=False,PinType.bIsWeakPointer=False,DefaultObject="/Script/Engine.Default__KismetMathLibrary",PersistentGuid=00000000000000000000000000000000,bHidden=True,bNotConnectable=False,bDefaultValueIsReadOnly=False,bDefaultValueIsIgnored=False,bAdvancedView=False,)
+			   CustomProperties Pin (PinId=557C12A64EC70024F7CC519F97D5D932,PinName="T",PinToolTip="T\nTransform  (by ref)\n\nThe transform you wish to invert",PinType.PinCategory="struct",PinType.PinSubCategoryObject=ScriptStruct'/Script/CoreUObject.Transform',PinType.PinSubCategoryMemberReference=(),PinType.bIsArray=False,PinType.bIsReference=True,PinType.bIsConst=True,PinType.bIsWeakPointer=False,LinkedTo=(K2Node_CallFunction_1531 5150BE354EF59F455E6CA988E1351299,),PersistentGuid=00000000000000000000000000000000,bHidden=False,bNotConnectable=False,bDefaultValueIsReadOnly=False,bDefaultValueIsIgnored=True,bAdvancedView=False,)
+			   CustomProperties Pin (PinId=BDA890074B460518CA47E0B5319A61D4,PinName="ReturnValue",PinToolTip="Return Value\nTransform \n\nThe inverse of T.",Direction="EGPD_Output",PinType.PinCategory="struct",PinType.PinSubCategoryObject=ScriptStruct'/Script/CoreUObject.Transform',PinType.PinSubCategoryMemberReference=(),PinType.bIsArray=False,PinType.bIsReference=False,PinType.bIsConst=False,PinType.bIsWeakPointer=False,LinkedTo=(K2Node_CallFunction_1530 8A4E9BDC43985D182203EEBE4D64ABED,),PersistentGuid=00000000000000000000000000000000,bHidden=False,bNotConnectable=False,bDefaultValueIsReadOnly=False,bDefaultValueIsIgnored=False,bAdvancedView=False,)
+			End Object
+			Begin Object Class=K2Node_CallFunction Name="K2Node_CallFunction_1533"
+			   bIsPureFunc=True
+			   FunctionReference=(MemberParent=Class'/Script/Engine.KismetMathLibrary',MemberName="Conv_VectorToLinearColor")
+			   NodePosX=1648
+			   NodePosY=-2894
+			   NodeGuid=4069F3B1421B0DB5F3B9318C1FC2A26D
+			   CustomProperties Pin (PinId=E21437414C7397632C246B96B157F7B6,PinName="self",PinFriendlyName=NSLOCTEXT("K2Node", "Target", "Target"),PinToolTip="Target\nKismet Math Library Reference",PinType.PinCategory="object",PinType.PinSubCategoryObject=Class'/Script/Engine.KismetMathLibrary',PinType.PinSubCategoryMemberReference=(),PinType.bIsArray=False,PinType.bIsReference=False,PinType.bIsConst=False,PinType.bIsWeakPointer=False,DefaultObject="/Script/Engine.Default__KismetMathLibrary",PersistentGuid=00000000000000000000000000000000,bHidden=True,bNotConnectable=False,bDefaultValueIsReadOnly=False,bDefaultValueIsIgnored=False,bAdvancedView=False,)
+			   CustomProperties Pin (PinId=1E73613645F8697CABA715886E923B95,PinName="InVec",PinToolTip="In Vec\nVector ",PinType.PinCategory="struct",PinType.PinSubCategoryObject=ScriptStruct'/Script/CoreUObject.Vector',PinType.PinSubCategoryMemberReference=(),PinType.bIsArray=False,PinType.bIsReference=False,PinType.bIsConst=False,PinType.bIsWeakPointer=False,DefaultValue="0, 0, 0",AutogeneratedDefaultValue="0, 0, 0",LinkedTo=(K2Node_CallFunction_1534 62F0E4DF49FC6CCA925D4E9A9BC03D67,),PersistentGuid=00000000000000000000000000000000,bHidden=False,bNotConnectable=False,bDefaultValueIsReadOnly=False,bDefaultValueIsIgnored=False,bAdvancedView=False,)
+			   CustomProperties Pin (PinId=86266FF5416699BCD33F46A84671E715,PinName="ReturnValue",PinToolTip="Return Value\nLinear Color Structure",Direction="EGPD_Output",PinType.PinCategory="struct",PinType.PinSubCategoryObject=ScriptStruct'/Script/CoreUObject.LinearColor',PinType.PinSubCategoryMemberReference=(),PinType.bIsArray=False,PinType.bIsReference=False,PinType.bIsConst=False,PinType.bIsWeakPointer=False,LinkedTo=(K2Node_CallFunction_1536 1523A2CA42868EC87B90F39450925710,),PersistentGuid=00000000000000000000000000000000,bHidden=False,bNotConnectable=False,bDefaultValueIsReadOnly=False,bDefaultValueIsIgnored=False,bAdvancedView=False,)
+			End Object
+			Begin Object Class=K2Node_CallFunction Name="K2Node_CallFunction_1534"
+			   bIsPureFunc=True
+			   FunctionReference=(MemberParent=Class'/Script/Engine.KismetMathLibrary',MemberName="Divide_VectorFloat")
+			   NodePosX=1504
+			   NodePosY=-2894
+			   NodeGuid=084501754AC7E7E336079E850628B7F1
+			   CustomProperties Pin (PinId=17EBF57640D6253D4AE04282E40F95B0,PinName="self",PinFriendlyName=NSLOCTEXT("K2Node", "Target", "Target"),PinToolTip="Target\nKismet Math Library Reference",PinType.PinCategory="object",PinType.PinSubCategoryObject=Class'/Script/Engine.KismetMathLibrary',PinType.PinSubCategoryMemberReference=(),PinType.bIsArray=False,PinType.bIsReference=False,PinType.bIsConst=False,PinType.bIsWeakPointer=False,DefaultObject="/Script/Engine.Default__KismetMathLibrary",PersistentGuid=00000000000000000000000000000000,bHidden=True,bNotConnectable=False,bDefaultValueIsReadOnly=False,bDefaultValueIsIgnored=False,bAdvancedView=False,)
+			   CustomProperties Pin (PinId=694D50D04A4FE8E7DC8CA689C5CDBE67,PinName="A",PinToolTip="A\nVector ",PinType.PinCategory="struct",PinType.PinSubCategoryObject=ScriptStruct'/Script/CoreUObject.Vector',PinType.PinSubCategoryMemberReference=(),PinType.bIsArray=False,PinType.bIsReference=False,PinType.bIsConst=False,PinType.bIsWeakPointer=False,DefaultValue="0, 0, 0",AutogeneratedDefaultValue="0, 0, 0",LinkedTo=(K2Node_CallFunction_1530 7335087343EF48C6242BCDBFC09F2B35,),PersistentGuid=00000000000000000000000000000000,bHidden=False,bNotConnectable=False,bDefaultValueIsReadOnly=False,bDefaultValueIsIgnored=False,bAdvancedView=False,)
+			   CustomProperties Pin (PinId=4FB4472943DA9A99ECFFC992F975CAAA,PinName="B",PinToolTip="B\nFloat",PinType.PinCategory="float",PinType.PinSubCategoryObject=None,PinType.PinSubCategoryMemberReference=(),PinType.bIsArray=False,PinType.bIsReference=False,PinType.bIsConst=False,PinType.bIsWeakPointer=False,DefaultValue="1000",AutogeneratedDefaultValue="1.000000",PersistentGuid=00000000000000000000000000000000,bHidden=False,bNotConnectable=False,bDefaultValueIsReadOnly=False,bDefaultValueIsIgnored=False,bAdvancedView=False,)
+			   CustomProperties Pin (PinId=62F0E4DF49FC6CCA925D4E9A9BC03D67,PinName="ReturnValue",PinToolTip="Return Value\nVector ",Direction="EGPD_Output",PinType.PinCategory="struct",PinType.PinSubCategoryObject=ScriptStruct'/Script/CoreUObject.Vector',PinType.PinSubCategoryMemberReference=(),PinType.bIsArray=False,PinType.bIsReference=False,PinType.bIsConst=False,PinType.bIsWeakPointer=False,AutogeneratedDefaultValue="0, 0, 0",LinkedTo=(K2Node_CallFunction_1533 1E73613645F8697CABA715886E923B95,),PersistentGuid=00000000000000000000000000000000,bHidden=False,bNotConnectable=False,bDefaultValueIsReadOnly=False,bDefaultValueIsIgnored=False,bAdvancedView=False,)
+			End Object
+			Begin Object Class=K2Node_CallFunction Name="K2Node_CallFunction_1535"
+			   FunctionReference=(MemberParent=Class'/Script/Engine.KismetRenderingLibrary',MemberName="DrawMaterialToRenderTarget")
+			   NodePosX=1664
+			   NodePosY=-3216
+			   NodeGuid=0A16A7E64EAF109F42DFF084F838ED62
+			   CustomProperties Pin (PinId=D89DC81641CA4093CA72DAAF5CACC388,PinName="execute",PinType.PinCategory="exec",PinType.PinSubCategoryObject=None,PinType.PinSubCategoryMemberReference=(),PinType.bIsArray=False,PinType.bIsReference=False,PinType.bIsConst=False,PinType.bIsWeakPointer=False,LinkedTo=(K2Node_CallFunction_1536 F5065258491F11F07B8E7B9BF2A38140,),PersistentGuid=00000000000000000000000000000000,bHidden=False,bNotConnectable=False,bDefaultValueIsReadOnly=False,bDefaultValueIsIgnored=False,bAdvancedView=False,)
+			   CustomProperties Pin (PinId=18C8362646DD324E0B3EE09FE2E9691F,PinName="then",Direction="EGPD_Output",PinType.PinCategory="exec",PinType.PinSubCategoryObject=None,PinType.PinSubCategoryMemberReference=(),PinType.bIsArray=False,PinType.bIsReference=False,PinType.bIsConst=False,PinType.bIsWeakPointer=False,PersistentGuid=00000000000000000000000000000000,bHidden=False,bNotConnectable=False,bDefaultValueIsReadOnly=False,bDefaultValueIsIgnored=False,bAdvancedView=False,)
+			   CustomProperties Pin (PinId=3E4C4C0C47F9C599C77C0A94DAE5CDBC,PinName="self",PinFriendlyName=NSLOCTEXT("K2Node", "Target", "Target"),PinType.PinCategory="object",PinType.PinSubCategoryObject=Class'/Script/Engine.KismetRenderingLibrary',PinType.PinSubCategoryMemberReference=(),PinType.bIsArray=False,PinType.bIsReference=False,PinType.bIsConst=False,PinType.bIsWeakPointer=False,DefaultObject="/Script/Engine.Default__KismetRenderingLibrary",PersistentGuid=00000000000000000000000000000000,bHidden=True,bNotConnectable=False,bDefaultValueIsReadOnly=False,bDefaultValueIsIgnored=False,bAdvancedView=False,)
+			   CustomProperties Pin (PinId=5EE9579647D9A511B63300AEE24E89BF,PinName="WorldContextObject",PinType.PinCategory="object",PinType.PinSubCategoryObject=Class'/Script/CoreUObject.Object',PinType.PinSubCategoryMemberReference=(),PinType.bIsArray=False,PinType.bIsReference=False,PinType.bIsConst=False,PinType.bIsWeakPointer=False,PersistentGuid=00000000000000000000000000000000,bHidden=True,bNotConnectable=False,bDefaultValueIsReadOnly=False,bDefaultValueIsIgnored=False,bAdvancedView=False,)
+			   CustomProperties Pin (PinId=B57C770C4D827E19FDE1F885B23E018C,PinName="TextureRenderTarget",PinType.PinCategory="object",PinType.PinSubCategoryObject=Class'/Script/Engine.TextureRenderTarget2D',PinType.PinSubCategoryMemberReference=(),PinType.bIsArray=False,PinType.bIsReference=False,PinType.bIsConst=False,PinType.bIsWeakPointer=False,LinkedTo=(K2Node_CallFunction_1529 E9A0065B436A442D29A940B6587482B6,),PersistentGuid=00000000000000000000000000000000,bHidden=False,bNotConnectable=False,bDefaultValueIsReadOnly=False,bDefaultValueIsIgnored=False,bAdvancedView=False,)
+			   CustomProperties Pin (PinId=D99B80D64C0171739FBC10A4E4C540EC,PinName="Material",PinType.PinCategory="object",PinType.PinSubCategoryObject=Class'/Script/Engine.MaterialInterface',PinType.PinSubCategoryMemberReference=(),PinType.bIsArray=False,PinType.bIsReference=False,PinType.bIsConst=False,PinType.bIsWeakPointer=False,LinkedTo=(K2Node_CallMaterialParameterCollectionFunction_116 0ABD7E4B428C7424D56A178963D9C137,),PersistentGuid=00000000000000000000000000000000,bHidden=False,bNotConnectable=False,bDefaultValueIsReadOnly=False,bDefaultValueIsIgnored=False,bAdvancedView=False,)
+			End Object
+			Begin Object Class=EdGraphNode_Comment Name="EdGraphNode_Comment_225"
+			   CommentDepth=-2
+			   NodePosX=208
+			   NodePosY=-3280
+			   NodeWidth=1728
+			   NodeHeight=528
+			   NodeComment="Take Damage force"
+			   NodeGuid=0FE9978145619358EE7C23B3434A2906
+			End Object
+			Begin Object Class=K2Node_CallFunction Name="K2Node_CallFunction_1536"
+			   FunctionReference=(MemberParent=Class'/Script/Engine.MaterialInstanceDynamic',MemberName="SetVectorParameterValue")
+			   NodePosX=1248
+			   NodePosY=-3230
+			   NodeGuid=34580CC742ED9BCE1B13D38DD6377872
+			   CustomProperties Pin (PinId=88C40CCB4DEBDA28F201A388C6DF16B9,PinName="execute",PinToolTip="\nExec",PinType.PinCategory="exec",PinType.PinSubCategoryObject=None,PinType.PinSubCategoryMemberReference=(),PinType.bIsArray=False,PinType.bIsReference=False,PinType.bIsConst=False,PinType.bIsWeakPointer=False,LinkedTo=(K2Node_CallMaterialParameterCollectionFunction_116 CE08AFD348EC9924552120BAB494E880,),PersistentGuid=00000000000000000000000000000000,bHidden=False,bNotConnectable=False,bDefaultValueIsReadOnly=False,bDefaultValueIsIgnored=False,bAdvancedView=False,)
+			   CustomProperties Pin (PinId=F5065258491F11F07B8E7B9BF2A38140,PinName="then",PinToolTip="\nExec",Direction="EGPD_Output",PinType.PinCategory="exec",PinType.PinSubCategoryObject=None,PinType.PinSubCategoryMemberReference=(),PinType.bIsArray=False,PinType.bIsReference=False,PinType.bIsConst=False,PinType.bIsWeakPointer=False,LinkedTo=(K2Node_CallFunction_1535 D89DC81641CA4093CA72DAAF5CACC388,),PersistentGuid=00000000000000000000000000000000,bHidden=False,bNotConnectable=False,bDefaultValueIsReadOnly=False,bDefaultValueIsIgnored=False,bAdvancedView=False,)
+			   CustomProperties Pin (PinId=FF3048F146D54F1C4AA4A48C221B5457,PinName="self",PinFriendlyName=NSLOCTEXT("K2Node", "Target", "Target"),PinToolTip="Target\nMaterial Instance Dynamic Reference",PinType.PinCategory="object",PinType.PinSubCategoryObject=Class'/Script/Engine.MaterialInstanceDynamic',PinType.PinSubCategoryMemberReference=(),PinType.bIsArray=False,PinType.bIsReference=False,PinType.bIsConst=False,PinType.bIsWeakPointer=False,LinkedTo=(K2Node_CallMaterialParameterCollectionFunction_116 0ABD7E4B428C7424D56A178963D9C137,),PersistentGuid=00000000000000000000000000000000,bHidden=False,bNotConnectable=False,bDefaultValueIsReadOnly=False,bDefaultValueIsIgnored=False,bAdvancedView=False,)
+			   CustomProperties Pin (PinId=50E108F844D08F285B9864B792806572,PinName="ParameterName",PinToolTip="Parameter Name\nName",PinType.PinCategory="name",PinType.PinSubCategoryObject=None,PinType.PinSubCategoryMemberReference=(),PinType.bIsArray=False,PinType.bIsReference=False,PinType.bIsConst=False,PinType.bIsWeakPointer=False,DefaultValue="ForcePosition",AutogeneratedDefaultValue="None",PersistentGuid=00000000000000000000000000000000,bHidden=False,bNotConnectable=False,bDefaultValueIsReadOnly=False,bDefaultValueIsIgnored=False,bAdvancedView=False,)
+			   CustomProperties Pin (PinId=1523A2CA42868EC87B90F39450925710,PinName="Value",PinToolTip="Value\nLinear Color Structure",PinType.PinCategory="struct",PinType.PinSubCategoryObject=ScriptStruct'/Script/CoreUObject.LinearColor',PinType.PinSubCategoryMemberReference=(),PinType.bIsArray=False,PinType.bIsReference=False,PinType.bIsConst=False,PinType.bIsWeakPointer=False,LinkedTo=(K2Node_CallFunction_1533 86266FF5416699BCD33F46A84671E715,),PersistentGuid=00000000000000000000000000000000,bHidden=False,bNotConnectable=False,bDefaultValueIsReadOnly=False,bDefaultValueIsIgnored=False,bAdvancedView=False,)
+			End Object
+			Begin Object Class=K2Node_CallMaterialParameterCollectionFunction Name="K2Node_CallMaterialParameterCollectionFunction_116"
+			   FunctionReference=(MemberParent=Class'/Script/Engine.KismetMaterialLibrary',MemberName="CreateDynamicMaterialInstance")
+			   NodePosX=896
+			   NodePosY=-3215
+			   NodeGuid=FB06673A4D849A470F76559CC498E165
+			   CustomProperties Pin (PinId=604EF317495C13D515FA428315721F21,PinName="execute",PinToolTip="\nExec",PinType.PinCategory="exec",PinType.PinSubCategoryObject=None,PinType.PinSubCategoryMemberReference=(),PinType.bIsArray=False,PinType.bIsReference=False,PinType.bIsConst=False,PinType.bIsWeakPointer=False,LinkedTo=(K2Node_CallFunction_1529 BF80119E4AE88655005DD5AD3A44C53F,),PersistentGuid=00000000000000000000000000000000,bHidden=False,bNotConnectable=False,bDefaultValueIsReadOnly=False,bDefaultValueIsIgnored=False,bAdvancedView=False,)
+			   CustomProperties Pin (PinId=CE08AFD348EC9924552120BAB494E880,PinName="then",PinToolTip="\nExec",Direction="EGPD_Output",PinType.PinCategory="exec",PinType.PinSubCategoryObject=None,PinType.PinSubCategoryMemberReference=(),PinType.bIsArray=False,PinType.bIsReference=False,PinType.bIsConst=False,PinType.bIsWeakPointer=False,LinkedTo=(K2Node_CallFunction_1536 88C40CCB4DEBDA28F201A388C6DF16B9,),PersistentGuid=00000000000000000000000000000000,bHidden=False,bNotConnectable=False,bDefaultValueIsReadOnly=False,bDefaultValueIsIgnored=False,bAdvancedView=False,)
+			   CustomProperties Pin (PinId=F354CE6841076CBD35B650AC5B015F77,PinName="self",PinFriendlyName=NSLOCTEXT("K2Node", "Target", "Target"),PinToolTip="Target\nKismet Material Library Reference",PinType.PinCategory="object",PinType.PinSubCategoryObject=Class'/Script/Engine.KismetMaterialLibrary',PinType.PinSubCategoryMemberReference=(),PinType.bIsArray=False,PinType.bIsReference=False,PinType.bIsConst=False,PinType.bIsWeakPointer=False,DefaultObject="/Script/Engine.Default__KismetMaterialLibrary",PersistentGuid=00000000000000000000000000000000,bHidden=True,bNotConnectable=False,bDefaultValueIsReadOnly=False,bDefaultValueIsIgnored=False,bAdvancedView=False,)
+			   CustomProperties Pin (PinId=15A9854149EF737BE52F289D6D38DC85,PinName="WorldContextObject",PinToolTip="World Context Object\nObject Reference",PinType.PinCategory="object",PinType.PinSubCategoryObject=Class'/Script/CoreUObject.Object',PinType.PinSubCategoryMemberReference=(),PinType.bIsArray=False,PinType.bIsReference=False,PinType.bIsConst=False,PinType.bIsWeakPointer=False,PersistentGuid=00000000000000000000000000000000,bHidden=True,bNotConnectable=False,bDefaultValueIsReadOnly=False,bDefaultValueIsIgnored=False,bAdvancedView=False,)
+			   CustomProperties Pin (PinId=F741E66F4E110D411EA953B54A502939,PinName="Parent",PinToolTip="Parent\nMaterial Interface Reference",PinType.PinCategory="object",PinType.PinSubCategoryObject=Class'/Script/Engine.MaterialInterface',PinType.PinSubCategoryMemberReference=(),PinType.bIsArray=False,PinType.bIsReference=False,PinType.bIsConst=False,PinType.bIsWeakPointer=False,DefaultObject="/Game/ExampleContent/BlueprintFluidSurface/ForceSplat.ForceSplat",PersistentGuid=00000000000000000000000000000000,bHidden=False,bNotConnectable=False,bDefaultValueIsReadOnly=False,bDefaultValueIsIgnored=False,bAdvancedView=False,)
+			   CustomProperties Pin (PinId=0ABD7E4B428C7424D56A178963D9C137,PinName="ReturnValue",PinToolTip="Return Value\nMaterial Instance Dynamic Reference",Direction="EGPD_Output",PinType.PinCategory="object",PinType.PinSubCategoryObject=Class'/Script/Engine.MaterialInstanceDynamic',PinType.PinSubCategoryMemberReference=(),PinType.bIsArray=False,PinType.bIsReference=False,PinType.bIsConst=False,PinType.bIsWeakPointer=False,LinkedTo=(K2Node_CallFunction_1536 FF3048F146D54F1C4AA4A48C221B5457,K2Node_CallFunction_1535 D99B80D64C0171739FBC10A4E4C540EC,),PersistentGuid=00000000000000000000000000000000,bHidden=False,bNotConnectable=False,bDefaultValueIsReadOnly=False,bDefaultValueIsIgnored=False,bAdvancedView=False,)
+			End Object
+			Begin Object Class=K2Node_VariableGet Name="K2Node_VariableGet_476"
+			   VariableReference=(MemberName="StaticMeshComponent",bSelfContext=True)
+			   NodePosX=544
+			   NodePosY=-2896
+			   NodeGuid=B01C78874B98077760295E973AB744FB
+			   CustomProperties Pin (PinId=78F118CD49FC76D027795D843DFE71FF,PinName="StaticMeshComponent",Direction="EGPD_Output",PinType.PinCategory="object",PinType.PinSubCategoryObject=Class'/Script/Engine.StaticMeshComponent',PinType.PinSubCategoryMemberReference=(),PinType.bIsArray=False,PinType.bIsReference=False,PinType.bIsConst=False,PinType.bIsWeakPointer=False,LinkedTo=(K2Node_CallFunction_1531 5A1010634C4DE3EAF3EF409F23EC1215,),PersistentGuid=00000000000000000000000000000000,bHidden=False,bNotConnectable=False,bDefaultValueIsReadOnly=False,bDefaultValueIsIgnored=False,bAdvancedView=False,)
+			   CustomProperties Pin (PinId=629DFDFB49CD19075CAD4B8F8CE8BD97,PinName="self",PinFriendlyName=NSLOCTEXT("K2Node", "Target", "Target"),PinType.PinCategory="object",PinType.PinSubCategoryObject=Class'/Script/Engine.StaticMeshActor',PinType.PinSubCategoryMemberReference=(),PinType.bIsArray=False,PinType.bIsReference=False,PinType.bIsConst=False,PinType.bIsWeakPointer=False,PersistentGuid=00000000000000000000000000000000,bHidden=True,bNotConnectable=False,bDefaultValueIsReadOnly=False,bDefaultValueIsIgnored=False,bAdvancedView=False,)
+			End Object
+			Begin Object Class=EdGraphNode_Comment Name="EdGraphNode_Comment_226"
+			   NodePosX=528
+			   NodePosY=-3008
+			   NodeWidth=1312
+			   NodeHeight=224
+			   NodeComment="Convert the hit location into texture space for this mesh"
+			   NodeGuid=8ABC00674CD138CE9F3EFF8BB266DFBD
+			End Object
+			END OBJECT
+	[/PARAM]
+[/OBJECT]
+[/EXCERPT:TakeDamage]
+
+[EXCERPT:Overlap]	
+[OBJECT:Snippet]
+	[PARAMLITERAL:id]
+		Overlap
+	[/PARAMLITERAL]
+	[PARAMLITERAL:type]
+		blueprint
+	[/PARAMLITERAL]
+	[PARAM:image]
+		![](Engine\Rendering\RenderTargets\BlueprintRenderTargets\HowTo\FluidSurface\FS_Begin_End_Overlap.png)
+	[/PARAM]
+	[PARAM:snippet]
+			BEGIN OBJECT
+							Begin Object Class=K2Node_Event Name="K2Node_Event_15"
+			   EventReference=(MemberParent=Class'/Script/Engine.Actor',MemberName="ReceiveActorBeginOverlap")
+			   bOverrideFunction=True
+			   NodePosX=-80
+			   NodePosY=-2096
+			   NodeGuid=4B16C97D43329443FD982BAAAF7377F0
+			   CustomProperties Pin (PinId=EDF91BE74B216C65079E9FB7F74F7FA7,PinName="OutputDelegate",Direction="EGPD_Output",PinType.PinCategory="delegate",PinType.PinSubCategoryObject=None,PinType.PinSubCategoryMemberReference=(MemberParent=Class'/Script/Engine.Actor',MemberName="ReceiveActorBeginOverlap"),PinType.bIsArray=False,PinType.bIsReference=False,PinType.bIsConst=False,PinType.bIsWeakPointer=False,PersistentGuid=00000000000000000000000000000000,bHidden=False,bNotConnectable=False,bDefaultValueIsReadOnly=False,bDefaultValueIsIgnored=False,bAdvancedView=False,)
+			   CustomProperties Pin (PinId=055914C14FA7D053484C2CB1F6F8EF52,PinName="then",Direction="EGPD_Output",PinType.PinCategory="exec",PinType.PinSubCategoryObject=None,PinType.PinSubCategoryMemberReference=(),PinType.bIsArray=False,PinType.bIsReference=False,PinType.bIsConst=False,PinType.bIsWeakPointer=False,LinkedTo=(K2Node_DynamicCast_1 B5D8DADE4D4ABD4C416EFAA3F65FA9B3,),PersistentGuid=00000000000000000000000000000000,bHidden=False,bNotConnectable=False,bDefaultValueIsReadOnly=False,bDefaultValueIsIgnored=False,bAdvancedView=False,)
+			   CustomProperties Pin (PinId=391A9F9F46A1A313B7CE58B7CDA1D848,PinName="OtherActor",PinToolTip="Other Actor\nActor Reference",Direction="EGPD_Output",PinType.PinCategory="object",PinType.PinSubCategoryObject=Class'/Script/Engine.Actor',PinType.PinSubCategoryMemberReference=(),PinType.bIsArray=False,PinType.bIsReference=False,PinType.bIsConst=False,PinType.bIsWeakPointer=False,LinkedTo=(K2Node_DynamicCast_1 D6C9840C4042EEA41862828875910F4F,K2Node_VariableSet_20 F4B88E7C4B4007FC9EFD4CA7888501CE,),PersistentGuid=00000000000000000000000000000000,bHidden=False,bNotConnectable=False,bDefaultValueIsReadOnly=False,bDefaultValueIsIgnored=False,bAdvancedView=False,)
+			End Object
+			Begin Object Class=K2Node_Event Name="K2Node_Event_28"
+			   EventReference=(MemberParent=Class'/Script/Engine.Actor',MemberName="ReceiveActorEndOverlap")
+			   bOverrideFunction=True
+			   NodePosX=-80
+			   NodePosY=-1856
+			   NodeGuid=36B7DFBE45F764675403CAB288F42C8E
+			   CustomProperties Pin (PinId=5623C22643914997B6F375BCD06FEC14,PinName="OutputDelegate",Direction="EGPD_Output",PinType.PinCategory="delegate",PinType.PinSubCategoryObject=None,PinType.PinSubCategoryMemberReference=(MemberParent=Class'/Script/Engine.Actor',MemberName="ReceiveActorEndOverlap"),PinType.bIsArray=False,PinType.bIsReference=False,PinType.bIsConst=False,PinType.bIsWeakPointer=False,PersistentGuid=00000000000000000000000000000000,bHidden=False,bNotConnectable=False,bDefaultValueIsReadOnly=False,bDefaultValueIsIgnored=False,bAdvancedView=False,)
+			   CustomProperties Pin (PinId=8A7D9BA3491EB4EF4B04849B8D9BFD31,PinName="then",Direction="EGPD_Output",PinType.PinCategory="exec",PinType.PinSubCategoryObject=None,PinType.PinSubCategoryMemberReference=(),PinType.bIsArray=False,PinType.bIsReference=False,PinType.bIsConst=False,PinType.bIsWeakPointer=False,LinkedTo=(K2Node_DynamicCast_0 06EB6749486B504A257F1BB37AA213B3,),PersistentGuid=00000000000000000000000000000000,bHidden=False,bNotConnectable=False,bDefaultValueIsReadOnly=False,bDefaultValueIsIgnored=False,bAdvancedView=False,)
+			   CustomProperties Pin (PinId=C29A42954CF639AE1F7B83898A2DECAB,PinName="OtherActor",PinToolTip="Other Actor\nActor Reference",Direction="EGPD_Output",PinType.PinCategory="object",PinType.PinSubCategoryObject=Class'/Script/Engine.Actor',PinType.PinSubCategoryMemberReference=(),PinType.bIsArray=False,PinType.bIsReference=False,PinType.bIsConst=False,PinType.bIsWeakPointer=False,LinkedTo=(K2Node_DynamicCast_0 8B424502457B0C7ECA3141831013FC29,),PersistentGuid=00000000000000000000000000000000,bHidden=False,bNotConnectable=False,bDefaultValueIsReadOnly=False,bDefaultValueIsIgnored=False,bAdvancedView=False,)
+			End Object
+			Begin Object Class=K2Node_VariableSet Name="K2Node_VariableSet_20"
+			   VariableReference=(MemberName="TouchingCharacter",MemberGuid=22B361484158D2692D5EDAB25083E8C8,bSelfContext=True)
+			   NodePosX=816
+			   NodePosY=-2016
+			   NodeGuid=0BDE15B94BE0EEBB7230DE822180E79A
+			   CustomProperties Pin (PinId=27C7F78B497DBF41964857B72A9FEBB0,PinName="execute",PinType.PinCategory="exec",PinType.PinSubCategoryObject=None,PinType.PinSubCategoryMemberReference=(),PinType.bIsArray=False,PinType.bIsReference=False,PinType.bIsConst=False,PinType.bIsWeakPointer=False,LinkedTo=(K2Node_IfThenElse_9 65A639B54E8372182D1387844E03C396,),PersistentGuid=00000000000000000000000000000000,bHidden=False,bNotConnectable=False,bDefaultValueIsReadOnly=False,bDefaultValueIsIgnored=False,bAdvancedView=False,)
+			   CustomProperties Pin (PinId=96FB9A954CCC6295F99E718AA70563B8,PinName="then",Direction="EGPD_Output",PinType.PinCategory="exec",PinType.PinSubCategoryObject=None,PinType.PinSubCategoryMemberReference=(),PinType.bIsArray=False,PinType.bIsReference=False,PinType.bIsConst=False,PinType.bIsWeakPointer=False,LinkedTo=(K2Node_VariableSet_54 1CB6F5A54FEB039F3367CD82E325B0FF,),PersistentGuid=00000000000000000000000000000000,bHidden=False,bNotConnectable=False,bDefaultValueIsReadOnly=False,bDefaultValueIsIgnored=False,bAdvancedView=False,)
+			   CustomProperties Pin (PinId=F4B88E7C4B4007FC9EFD4CA7888501CE,PinName="TouchingCharacter",PinType.PinCategory="object",PinType.PinSubCategoryObject=Class'/Script/Engine.Actor',PinType.PinSubCategoryMemberReference=(),PinType.bIsArray=False,PinType.bIsReference=False,PinType.bIsConst=False,PinType.bIsWeakPointer=False,LinkedTo=(K2Node_Event_15 391A9F9F46A1A313B7CE58B7CDA1D848,),PersistentGuid=00000000000000000000000000000000,bHidden=False,bNotConnectable=False,bDefaultValueIsReadOnly=False,bDefaultValueIsIgnored=False,bAdvancedView=False,)
+			   CustomProperties Pin (PinId=C41839C34AA37950BEA6FBA32BCA0BB5,PinName="Output_Get",PinToolTip="Retrieves the value of the variable, can use instead of a separate Get node",Direction="EGPD_Output",PinType.PinCategory="object",PinType.PinSubCategoryObject=Class'/Script/Engine.Actor',PinType.PinSubCategoryMemberReference=(),PinType.bIsArray=False,PinType.bIsReference=False,PinType.bIsConst=False,PinType.bIsWeakPointer=False,LinkedTo=(K2Node_CallFunction_470 5FE83FDE476C10859CF9D094B8DD9A2B,),PersistentGuid=00000000000000000000000000000000,bHidden=False,bNotConnectable=False,bDefaultValueIsReadOnly=False,bDefaultValueIsIgnored=False,bAdvancedView=False,)
+			   CustomProperties Pin (PinId=B043F0E245BC6EC5F08404AE0371CF01,PinName="self",PinFriendlyName=NSLOCTEXT("K2Node", "Target", "Target"),PinType.PinCategory="object",PinType.PinSubCategoryObject=BlueprintGeneratedClass'/Game/ExampleContent/BlueprintFluidSurface/WaterSurface_BP.WaterSurface_BP_C',PinType.PinSubCategoryMemberReference=(),PinType.bIsArray=False,PinType.bIsReference=False,PinType.bIsConst=False,PinType.bIsWeakPointer=False,PersistentGuid=00000000000000000000000000000000,bHidden=True,bNotConnectable=False,bDefaultValueIsReadOnly=False,bDefaultValueIsIgnored=False,bAdvancedView=False,)
+			End Object
+			Begin Object Class=K2Node_VariableSet Name="K2Node_VariableSet_29"
+			   VariableReference=(MemberName="TouchingCharacter",MemberGuid=22B361484158D2692D5EDAB25083E8C8,bSelfContext=True)
+			   NodePosX=816
+			   NodePosY=-1812
+			   NodeGuid=4373DC10415D0CFF3EFA0BA0DCF6DD91
+			   CustomProperties Pin (PinId=3372F24044B22493E4CDB99C2FEE9AC6,PinName="execute",PinType.PinCategory="exec",PinType.PinSubCategoryObject=None,PinType.PinSubCategoryMemberReference=(),PinType.bIsArray=False,PinType.bIsReference=False,PinType.bIsConst=False,PinType.bIsWeakPointer=False,LinkedTo=(K2Node_IfThenElse_8 8FBC10034397D4F24C0A29B09D795D95,),PersistentGuid=00000000000000000000000000000000,bHidden=False,bNotConnectable=False,bDefaultValueIsReadOnly=False,bDefaultValueIsIgnored=False,bAdvancedView=False,)
+			   CustomProperties Pin (PinId=906C6BE0479F2D1541DE8891DE93F022,PinName="then",Direction="EGPD_Output",PinType.PinCategory="exec",PinType.PinSubCategoryObject=None,PinType.PinSubCategoryMemberReference=(),PinType.bIsArray=False,PinType.bIsReference=False,PinType.bIsConst=False,PinType.bIsWeakPointer=False,PersistentGuid=00000000000000000000000000000000,bHidden=False,bNotConnectable=False,bDefaultValueIsReadOnly=False,bDefaultValueIsIgnored=False,bAdvancedView=False,)
+			   CustomProperties Pin (PinId=3A78422A4DA45950C9E41F9FD191455C,PinName="TouchingCharacter",PinType.PinCategory="object",PinType.PinSubCategoryObject=Class'/Script/Engine.Actor',PinType.PinSubCategoryMemberReference=(),PinType.bIsArray=False,PinType.bIsReference=False,PinType.bIsConst=False,PinType.bIsWeakPointer=False,PersistentGuid=00000000000000000000000000000000,bHidden=False,bNotConnectable=False,bDefaultValueIsReadOnly=False,bDefaultValueIsIgnored=False,bAdvancedView=False,)
+			   CustomProperties Pin (PinId=60C1EED7448A10D9D45EEF83C9143C9F,PinName="Output_Get",PinToolTip="Retrieves the value of the variable, can use instead of a separate Get node",Direction="EGPD_Output",PinType.PinCategory="object",PinType.PinSubCategoryObject=Class'/Script/Engine.Actor',PinType.PinSubCategoryMemberReference=(),PinType.bIsArray=False,PinType.bIsReference=False,PinType.bIsConst=False,PinType.bIsWeakPointer=False,PersistentGuid=00000000000000000000000000000000,bHidden=False,bNotConnectable=False,bDefaultValueIsReadOnly=False,bDefaultValueIsIgnored=False,bAdvancedView=False,)
+			   CustomProperties Pin (PinId=E6B19F2847FD61E8430CC3AE2F764698,PinName="self",PinFriendlyName=NSLOCTEXT("K2Node", "Target", "Target"),PinType.PinCategory="object",PinType.PinSubCategoryObject=BlueprintGeneratedClass'/Game/ExampleContent/BlueprintFluidSurface/WaterSurface_BP.WaterSurface_BP_C',PinType.PinSubCategoryMemberReference=(),PinType.bIsArray=False,PinType.bIsReference=False,PinType.bIsConst=False,PinType.bIsWeakPointer=False,PersistentGuid=00000000000000000000000000000000,bHidden=True,bNotConnectable=False,bDefaultValueIsReadOnly=False,bDefaultValueIsIgnored=False,bAdvancedView=False,)
+			End Object
+			Begin Object Class=K2Node_VariableSet Name="K2Node_VariableSet_54"
+			   VariableReference=(MemberName="LastTouchingActorPosition",MemberGuid=C849AC4745DA47383E1884AA8AA77BC9,bSelfContext=True)
+			   NodePosX=1280
+			   NodePosY=-2016
+			   NodeGuid=3543BBE04E182DFBCF8F55907EC879EF
+			   CustomProperties Pin (PinId=1CB6F5A54FEB039F3367CD82E325B0FF,PinName="execute",PinType.PinCategory="exec",PinType.PinSubCategoryObject=None,PinType.PinSubCategoryMemberReference=(),PinType.bIsArray=False,PinType.bIsReference=False,PinType.bIsConst=False,PinType.bIsWeakPointer=False,LinkedTo=(K2Node_VariableSet_20 96FB9A954CCC6295F99E718AA70563B8,),PersistentGuid=00000000000000000000000000000000,bHidden=False,bNotConnectable=False,bDefaultValueIsReadOnly=False,bDefaultValueIsIgnored=False,bAdvancedView=False,)
+			   CustomProperties Pin (PinId=FD830AF54885553AA9AE85922058B810,PinName="then",Direction="EGPD_Output",PinType.PinCategory="exec",PinType.PinSubCategoryObject=None,PinType.PinSubCategoryMemberReference=(),PinType.bIsArray=False,PinType.bIsReference=False,PinType.bIsConst=False,PinType.bIsWeakPointer=False,PersistentGuid=00000000000000000000000000000000,bHidden=False,bNotConnectable=False,bDefaultValueIsReadOnly=False,bDefaultValueIsIgnored=False,bAdvancedView=False,)
+			   CustomProperties Pin (PinId=2E930A3743304FE706034A8249E7A33B,PinName="LastTouchingActorPosition",PinType.PinCategory="struct",PinType.PinSubCategoryObject=ScriptStruct'/Script/CoreUObject.Vector',PinType.PinSubCategoryMemberReference=(),PinType.bIsArray=False,PinType.bIsReference=False,PinType.bIsConst=False,PinType.bIsWeakPointer=False,DefaultValue="0, 0, 0",AutogeneratedDefaultValue="0, 0, 0",LinkedTo=(K2Node_CallFunction_470 7B1CCB42494CE3435ABB058D8F03C923,),PersistentGuid=00000000000000000000000000000000,bHidden=False,bNotConnectable=False,bDefaultValueIsReadOnly=False,bDefaultValueIsIgnored=False,bAdvancedView=False,)
+			   CustomProperties Pin (PinId=7D61E06C423B89754E763DB02A6585DA,PinName="Output_Get",PinToolTip="Retrieves the value of the variable, can use instead of a separate Get node",Direction="EGPD_Output",PinType.PinCategory="struct",PinType.PinSubCategoryObject=ScriptStruct'/Script/CoreUObject.Vector',PinType.PinSubCategoryMemberReference=(),PinType.bIsArray=False,PinType.bIsReference=False,PinType.bIsConst=False,PinType.bIsWeakPointer=False,DefaultValue="0, 0, 0",AutogeneratedDefaultValue="0, 0, 0",PersistentGuid=00000000000000000000000000000000,bHidden=False,bNotConnectable=False,bDefaultValueIsReadOnly=False,bDefaultValueIsIgnored=False,bAdvancedView=False,)
+			   CustomProperties Pin (PinId=5181DED447177703A8F449BEF51888AD,PinName="self",PinFriendlyName=NSLOCTEXT("K2Node", "Target", "Target"),PinType.PinCategory="object",PinType.PinSubCategoryObject=BlueprintGeneratedClass'/Game/ExampleContent/BlueprintFluidSurface/WaterSurface_BP.WaterSurface_BP_C',PinType.PinSubCategoryMemberReference=(),PinType.bIsArray=False,PinType.bIsReference=False,PinType.bIsConst=False,PinType.bIsWeakPointer=False,PersistentGuid=00000000000000000000000000000000,bHidden=True,bNotConnectable=False,bDefaultValueIsReadOnly=False,bDefaultValueIsIgnored=False,bAdvancedView=False,)
+			End Object
+			Begin Object Class=K2Node_CallFunction Name="K2Node_CallFunction_470"
+			   bIsPureFunc=True
+			   bIsConstFunc=True
+			   FunctionReference=(MemberName="K2_GetActorLocation",bSelfContext=True)
+			   NodePosX=1056
+			   NodePosY=-1872
+			   NodeGuid=141955464E52ED6BAAE196A1A765EDFE
+			   CustomProperties Pin (PinId=5FE83FDE476C10859CF9D094B8DD9A2B,PinName="self",PinFriendlyName=NSLOCTEXT("K2Node", "Target", "Target"),PinToolTip="Target\nActor Reference",PinType.PinCategory="object",PinType.PinSubCategoryObject=Class'/Script/Engine.Actor',PinType.PinSubCategoryMemberReference=(),PinType.bIsArray=False,PinType.bIsReference=False,PinType.bIsConst=False,PinType.bIsWeakPointer=False,LinkedTo=(K2Node_VariableSet_20 C41839C34AA37950BEA6FBA32BCA0BB5,),PersistentGuid=00000000000000000000000000000000,bHidden=False,bNotConnectable=False,bDefaultValueIsReadOnly=False,bDefaultValueIsIgnored=False,bAdvancedView=False,)
+			   CustomProperties Pin (PinId=7B1CCB42494CE3435ABB058D8F03C923,PinName="ReturnValue",PinToolTip="Return Value\nVector ",Direction="EGPD_Output",PinType.PinCategory="struct",PinType.PinSubCategoryObject=ScriptStruct'/Script/CoreUObject.Vector',PinType.PinSubCategoryMemberReference=(),PinType.bIsArray=False,PinType.bIsReference=False,PinType.bIsConst=False,PinType.bIsWeakPointer=False,AutogeneratedDefaultValue="0, 0, 0",LinkedTo=(K2Node_VariableSet_54 2E930A3743304FE706034A8249E7A33B,),PersistentGuid=00000000000000000000000000000000,bHidden=False,bNotConnectable=False,bDefaultValueIsReadOnly=False,bDefaultValueIsIgnored=False,bAdvancedView=False,)
+			End Object
+			Begin Object Class=K2Node_DynamicCast Name="K2Node_DynamicCast_1"
+			   TargetType=Class'/Script/Engine.Pawn'
+			   NodePosX=192
+			   NodePosY=-2096
+			   NodeGuid=40CB32054FDA3D91875DC9A14D7C72D8
+			   CustomProperties Pin (PinId=B5D8DADE4D4ABD4C416EFAA3F65FA9B3,PinName="execute",PinType.PinCategory="exec",PinType.PinSubCategoryObject=None,PinType.PinSubCategoryMemberReference=(),PinType.bIsArray=False,PinType.bIsReference=False,PinType.bIsConst=False,PinType.bIsWeakPointer=False,LinkedTo=(K2Node_Event_15 055914C14FA7D053484C2CB1F6F8EF52,),PersistentGuid=00000000000000000000000000000000,bHidden=False,bNotConnectable=False,bDefaultValueIsReadOnly=False,bDefaultValueIsIgnored=False,bAdvancedView=False,)
+			   CustomProperties Pin (PinId=7AB01AEF4834E1CA7B12058A2A23DE21,PinName="then",Direction="EGPD_Output",PinType.PinCategory="exec",PinType.PinSubCategoryObject=None,PinType.PinSubCategoryMemberReference=(),PinType.bIsArray=False,PinType.bIsReference=False,PinType.bIsConst=False,PinType.bIsWeakPointer=False,LinkedTo=(K2Node_IfThenElse_9 67354A2D48CE7BA27BA653BD29FF3B5F,),PersistentGuid=00000000000000000000000000000000,bHidden=False,bNotConnectable=False,bDefaultValueIsReadOnly=False,bDefaultValueIsIgnored=False,bAdvancedView=False,)
+			   CustomProperties Pin (PinId=5338F64E4BA8333DFFDD918AD8C05EEA,PinName="CastFailed",Direction="EGPD_Output",PinType.PinCategory="exec",PinType.PinSubCategoryObject=None,PinType.PinSubCategoryMemberReference=(),PinType.bIsArray=False,PinType.bIsReference=False,PinType.bIsConst=False,PinType.bIsWeakPointer=False,LinkedTo=(K2Node_IfThenElse_9 67354A2D48CE7BA27BA653BD29FF3B5F,),PersistentGuid=00000000000000000000000000000000,bHidden=False,bNotConnectable=False,bDefaultValueIsReadOnly=False,bDefaultValueIsIgnored=False,bAdvancedView=False,)
+			   CustomProperties Pin (PinId=D6C9840C4042EEA41862828875910F4F,PinName="Object",PinType.PinCategory="object",PinType.PinSubCategoryObject=Class'/Script/CoreUObject.Object',PinType.PinSubCategoryMemberReference=(),PinType.bIsArray=False,PinType.bIsReference=False,PinType.bIsConst=False,PinType.bIsWeakPointer=False,LinkedTo=(K2Node_Event_15 391A9F9F46A1A313B7CE58B7CDA1D848,),PersistentGuid=00000000000000000000000000000000,bHidden=False,bNotConnectable=False,bDefaultValueIsReadOnly=False,bDefaultValueIsIgnored=False,bAdvancedView=False,)
+			   CustomProperties Pin (PinId=843D74A642EE387079F3538B3C52FD8A,PinName="AsPawn",Direction="EGPD_Output",PinType.PinCategory="object",PinType.PinSubCategoryObject=Class'/Script/Engine.Pawn',PinType.PinSubCategoryMemberReference=(),PinType.bIsArray=False,PinType.bIsReference=False,PinType.bIsConst=False,PinType.bIsWeakPointer=False,LinkedTo=(K2Node_CallFunction_296 AA2CC8364449C9DF0517E2B09AA2036C,),PersistentGuid=00000000000000000000000000000000,bHidden=False,bNotConnectable=False,bDefaultValueIsReadOnly=False,bDefaultValueIsIgnored=False,bAdvancedView=False,)
+			   CustomProperties Pin (PinId=BE40C6DF41E15C7514C30C8F0E80BAD5,PinName="bSuccess",Direction="EGPD_Output",PinType.PinCategory="bool",PinType.PinSubCategoryObject=None,PinType.PinSubCategoryMemberReference=(),PinType.bIsArray=False,PinType.bIsReference=False,PinType.bIsConst=False,PinType.bIsWeakPointer=False,PersistentGuid=00000000000000000000000000000000,bHidden=True,bNotConnectable=False,bDefaultValueIsReadOnly=False,bDefaultValueIsIgnored=False,bAdvancedView=False,)
+			End Object
+			Begin Object Class=K2Node_CallFunction Name="K2Node_CallFunction_296"
+			   bIsPureFunc=True
+			   FunctionReference=(MemberParent=Class'/Script/Engine.KismetMathLibrary',MemberName="EqualEqual_ObjectObject")
+			   NodePosX=400
+			   NodePosY=-1968
+			   NodeGuid=ACD04FEC4506753112C597AE71A3084C
+			   CustomProperties Pin (PinId=D6181CCF4089D8CDF00E0EA060935EBC,PinName="self",PinFriendlyName=NSLOCTEXT("K2Node", "Target", "Target"),PinToolTip="Target\nKismet Math Library Reference",PinType.PinCategory="object",PinType.PinSubCategoryObject=Class'/Script/Engine.KismetMathLibrary',PinType.PinSubCategoryMemberReference=(),PinType.bIsArray=False,PinType.bIsReference=False,PinType.bIsConst=False,PinType.bIsWeakPointer=False,DefaultObject="/Script/Engine.Default__KismetMathLibrary",PersistentGuid=00000000000000000000000000000000,bHidden=True,bNotConnectable=False,bDefaultValueIsReadOnly=False,bDefaultValueIsIgnored=False,bAdvancedView=False,)
+			   CustomProperties Pin (PinId=AA2CC8364449C9DF0517E2B09AA2036C,PinName="A",PinToolTip="A\nObject Reference",PinType.PinCategory="object",PinType.PinSubCategoryObject=Class'/Script/CoreUObject.Object',PinType.PinSubCategoryMemberReference=(),PinType.bIsArray=False,PinType.bIsReference=False,PinType.bIsConst=False,PinType.bIsWeakPointer=False,LinkedTo=(K2Node_DynamicCast_1 843D74A642EE387079F3538B3C52FD8A,),PersistentGuid=00000000000000000000000000000000,bHidden=False,bNotConnectable=False,bDefaultValueIsReadOnly=False,bDefaultValueIsIgnored=False,bAdvancedView=False,)
+			   CustomProperties Pin (PinId=E94368524B8BF23E1CB015A8979EC2C7,PinName="B",PinToolTip="B\nObject Reference",PinType.PinCategory="object",PinType.PinSubCategoryObject=Class'/Script/CoreUObject.Object',PinType.PinSubCategoryMemberReference=(),PinType.bIsArray=False,PinType.bIsReference=False,PinType.bIsConst=False,PinType.bIsWeakPointer=False,PersistentGuid=00000000000000000000000000000000,bHidden=False,bNotConnectable=False,bDefaultValueIsReadOnly=False,bDefaultValueIsIgnored=False,bAdvancedView=False,)
+			   CustomProperties Pin (PinId=3E86CDA140599882576083A249F7D50A,PinName="ReturnValue",PinToolTip="Return Value\nBoolean",Direction="EGPD_Output",PinType.PinCategory="bool",PinType.PinSubCategoryObject=None,PinType.PinSubCategoryMemberReference=(),PinType.bIsArray=False,PinType.bIsReference=False,PinType.bIsConst=False,PinType.bIsWeakPointer=False,AutogeneratedDefaultValue="false",LinkedTo=(K2Node_IfThenElse_9 843CF78648A6F4CB3A866BBBB6B1105E,),PersistentGuid=00000000000000000000000000000000,bHidden=False,bNotConnectable=False,bDefaultValueIsReadOnly=False,bDefaultValueIsIgnored=False,bAdvancedView=False,)
+			End Object
+			Begin Object Class=K2Node_IfThenElse Name="K2Node_IfThenElse_9"
+			   NodePosX=608
+			   NodePosY=-2096
+			   NodeGuid=580F3CC44B53E518F7077FAE0F257B18
+			   CustomProperties Pin (PinId=67354A2D48CE7BA27BA653BD29FF3B5F,PinName="execute",PinType.PinCategory="exec",PinType.PinSubCategoryObject=None,PinType.PinSubCategoryMemberReference=(),PinType.bIsArray=False,PinType.bIsReference=False,PinType.bIsConst=False,PinType.bIsWeakPointer=False,LinkedTo=(K2Node_DynamicCast_1 7AB01AEF4834E1CA7B12058A2A23DE21,K2Node_DynamicCast_1 5338F64E4BA8333DFFDD918AD8C05EEA,),PersistentGuid=00000000000000000000000000000000,bHidden=False,bNotConnectable=False,bDefaultValueIsReadOnly=False,bDefaultValueIsIgnored=False,bAdvancedView=False,)
+			   CustomProperties Pin (PinId=843CF78648A6F4CB3A866BBBB6B1105E,PinName="Condition",PinType.PinCategory="bool",PinType.PinSubCategoryObject=None,PinType.PinSubCategoryMemberReference=(),PinType.bIsArray=False,PinType.bIsReference=False,PinType.bIsConst=False,PinType.bIsWeakPointer=False,DefaultValue="false",AutogeneratedDefaultValue="true",LinkedTo=(K2Node_CallFunction_296 3E86CDA140599882576083A249F7D50A,),PersistentGuid=00000000000000000000000000000000,bHidden=False,bNotConnectable=False,bDefaultValueIsReadOnly=False,bDefaultValueIsIgnored=False,bAdvancedView=False,)
+			   CustomProperties Pin (PinId=B013B2B940F293EB568958951B67AC86,PinName="then",PinFriendlyName=NSLOCTEXT("K2Node", "true", "true"),Direction="EGPD_Output",PinType.PinCategory="exec",PinType.PinSubCategoryObject=None,PinType.PinSubCategoryMemberReference=(),PinType.bIsArray=False,PinType.bIsReference=False,PinType.bIsConst=False,PinType.bIsWeakPointer=False,PersistentGuid=00000000000000000000000000000000,bHidden=False,bNotConnectable=False,bDefaultValueIsReadOnly=False,bDefaultValueIsIgnored=False,bAdvancedView=False,)
+			   CustomProperties Pin (PinId=65A639B54E8372182D1387844E03C396,PinName="else",PinFriendlyName=NSLOCTEXT("K2Node", "false", "false"),Direction="EGPD_Output",PinType.PinCategory="exec",PinType.PinSubCategoryObject=None,PinType.PinSubCategoryMemberReference=(),PinType.bIsArray=False,PinType.bIsReference=False,PinType.bIsConst=False,PinType.bIsWeakPointer=False,LinkedTo=(K2Node_VariableSet_20 27C7F78B497DBF41964857B72A9FEBB0,),PersistentGuid=00000000000000000000000000000000,bHidden=False,bNotConnectable=False,bDefaultValueIsReadOnly=False,bDefaultValueIsIgnored=False,bAdvancedView=False,)
+			End Object
+			Begin Object Class=K2Node_DynamicCast Name="K2Node_DynamicCast_0"
+			   TargetType=Class'/Script/Engine.Pawn'
+			   NodePosX=160
+			   NodePosY=-1856
+			   NodeGuid=7C789A4242DEEA644D05D7B762DCE930
+			   CustomProperties Pin (PinId=06EB6749486B504A257F1BB37AA213B3,PinName="execute",PinType.PinCategory="exec",PinType.PinSubCategoryObject=None,PinType.PinSubCategoryMemberReference=(),PinType.bIsArray=False,PinType.bIsReference=False,PinType.bIsConst=False,PinType.bIsWeakPointer=False,LinkedTo=(K2Node_Event_28 8A7D9BA3491EB4EF4B04849B8D9BFD31,),PersistentGuid=00000000000000000000000000000000,bHidden=False,bNotConnectable=False,bDefaultValueIsReadOnly=False,bDefaultValueIsIgnored=False,bAdvancedView=False,)
+			   CustomProperties Pin (PinId=DFD2BEFC42BDDD3AACD4D6927945F824,PinName="then",Direction="EGPD_Output",PinType.PinCategory="exec",PinType.PinSubCategoryObject=None,PinType.PinSubCategoryMemberReference=(),PinType.bIsArray=False,PinType.bIsReference=False,PinType.bIsConst=False,PinType.bIsWeakPointer=False,LinkedTo=(K2Node_IfThenElse_8 1379B94A4AACE12905E0A19BEF4F3E85,),PersistentGuid=00000000000000000000000000000000,bHidden=False,bNotConnectable=False,bDefaultValueIsReadOnly=False,bDefaultValueIsIgnored=False,bAdvancedView=False,)
+			   CustomProperties Pin (PinId=9A1E90334AA40A869A540984E1B11AFD,PinName="CastFailed",Direction="EGPD_Output",PinType.PinCategory="exec",PinType.PinSubCategoryObject=None,PinType.PinSubCategoryMemberReference=(),PinType.bIsArray=False,PinType.bIsReference=False,PinType.bIsConst=False,PinType.bIsWeakPointer=False,LinkedTo=(K2Node_IfThenElse_8 1379B94A4AACE12905E0A19BEF4F3E85,),PersistentGuid=00000000000000000000000000000000,bHidden=False,bNotConnectable=False,bDefaultValueIsReadOnly=False,bDefaultValueIsIgnored=False,bAdvancedView=False,)
+			   CustomProperties Pin (PinId=8B424502457B0C7ECA3141831013FC29,PinName="Object",PinType.PinCategory="object",PinType.PinSubCategoryObject=Class'/Script/CoreUObject.Object',PinType.PinSubCategoryMemberReference=(),PinType.bIsArray=False,PinType.bIsReference=False,PinType.bIsConst=False,PinType.bIsWeakPointer=False,LinkedTo=(K2Node_Event_28 C29A42954CF639AE1F7B83898A2DECAB,),PersistentGuid=00000000000000000000000000000000,bHidden=False,bNotConnectable=False,bDefaultValueIsReadOnly=False,bDefaultValueIsIgnored=False,bAdvancedView=False,)
+			   CustomProperties Pin (PinId=98E69AC0456F3B063EA6C5AC5AFCDED2,PinName="AsPawn",Direction="EGPD_Output",PinType.PinCategory="object",PinType.PinSubCategoryObject=Class'/Script/Engine.Pawn',PinType.PinSubCategoryMemberReference=(),PinType.bIsArray=False,PinType.bIsReference=False,PinType.bIsConst=False,PinType.bIsWeakPointer=False,LinkedTo=(K2Node_CallFunction_295 D0D22CC6496D73BD12EE72AB42773A51,),PersistentGuid=00000000000000000000000000000000,bHidden=False,bNotConnectable=False,bDefaultValueIsReadOnly=False,bDefaultValueIsIgnored=False,bAdvancedView=False,)
+			   CustomProperties Pin (PinId=73FE87DD44BB23EADFCA389AB0C22243,PinName="bSuccess",Direction="EGPD_Output",PinType.PinCategory="bool",PinType.PinSubCategoryObject=None,PinType.PinSubCategoryMemberReference=(),PinType.bIsArray=False,PinType.bIsReference=False,PinType.bIsConst=False,PinType.bIsWeakPointer=False,PersistentGuid=00000000000000000000000000000000,bHidden=True,bNotConnectable=False,bDefaultValueIsReadOnly=False,bDefaultValueIsIgnored=False,bAdvancedView=False,)
+			End Object
+			Begin Object Class=K2Node_CallFunction Name="K2Node_CallFunction_295"
+			   bIsPureFunc=True
+			   FunctionReference=(MemberParent=Class'/Script/Engine.KismetMathLibrary',MemberName="EqualEqual_ObjectObject")
+			   NodePosX=368
+			   NodePosY=-1744
+			   NodeGuid=AB2AE0C346B34E9B1C0568BF4203E320
+			   CustomProperties Pin (PinId=6A731F084E6EED9720C42D901AB0EDDF,PinName="self",PinFriendlyName=NSLOCTEXT("K2Node", "Target", "Target"),PinToolTip="Target\nKismet Math Library Reference",PinType.PinCategory="object",PinType.PinSubCategoryObject=Class'/Script/Engine.KismetMathLibrary',PinType.PinSubCategoryMemberReference=(),PinType.bIsArray=False,PinType.bIsReference=False,PinType.bIsConst=False,PinType.bIsWeakPointer=False,DefaultObject="/Script/Engine.Default__KismetMathLibrary",PersistentGuid=00000000000000000000000000000000,bHidden=True,bNotConnectable=False,bDefaultValueIsReadOnly=False,bDefaultValueIsIgnored=False,bAdvancedView=False,)
+			   CustomProperties Pin (PinId=D0D22CC6496D73BD12EE72AB42773A51,PinName="A",PinToolTip="A\nObject Reference",PinType.PinCategory="object",PinType.PinSubCategoryObject=Class'/Script/CoreUObject.Object',PinType.PinSubCategoryMemberReference=(),PinType.bIsArray=False,PinType.bIsReference=False,PinType.bIsConst=False,PinType.bIsWeakPointer=False,LinkedTo=(K2Node_DynamicCast_0 98E69AC0456F3B063EA6C5AC5AFCDED2,),PersistentGuid=00000000000000000000000000000000,bHidden=False,bNotConnectable=False,bDefaultValueIsReadOnly=False,bDefaultValueIsIgnored=False,bAdvancedView=False,)
+			   CustomProperties Pin (PinId=3EE181B5429D9BE8E2AF6099C94E5317,PinName="B",PinToolTip="B\nObject Reference",PinType.PinCategory="object",PinType.PinSubCategoryObject=Class'/Script/CoreUObject.Object',PinType.PinSubCategoryMemberReference=(),PinType.bIsArray=False,PinType.bIsReference=False,PinType.bIsConst=False,PinType.bIsWeakPointer=False,PersistentGuid=00000000000000000000000000000000,bHidden=False,bNotConnectable=False,bDefaultValueIsReadOnly=False,bDefaultValueIsIgnored=False,bAdvancedView=False,)
+			   CustomProperties Pin (PinId=0A18A59443E7C0D120E7FBA8D4CA0D24,PinName="ReturnValue",PinToolTip="Return Value\nBoolean",Direction="EGPD_Output",PinType.PinCategory="bool",PinType.PinSubCategoryObject=None,PinType.PinSubCategoryMemberReference=(),PinType.bIsArray=False,PinType.bIsReference=False,PinType.bIsConst=False,PinType.bIsWeakPointer=False,AutogeneratedDefaultValue="false",LinkedTo=(K2Node_IfThenElse_8 49D7F9D041F6F18B5714018FAA632144,),PersistentGuid=00000000000000000000000000000000,bHidden=False,bNotConnectable=False,bDefaultValueIsReadOnly=False,bDefaultValueIsIgnored=False,bAdvancedView=False,)
+			End Object
+			Begin Object Class=K2Node_IfThenElse Name="K2Node_IfThenElse_8"
+			   NodePosX=608
+			   NodePosY=-1856
+			   NodeGuid=099773344687B4C73943F981520F32AC
+			   CustomProperties Pin (PinId=1379B94A4AACE12905E0A19BEF4F3E85,PinName="execute",PinType.PinCategory="exec",PinType.PinSubCategoryObject=None,PinType.PinSubCategoryMemberReference=(),PinType.bIsArray=False,PinType.bIsReference=False,PinType.bIsConst=False,PinType.bIsWeakPointer=False,LinkedTo=(K2Node_DynamicCast_0 DFD2BEFC42BDDD3AACD4D6927945F824,K2Node_DynamicCast_0 9A1E90334AA40A869A540984E1B11AFD,),PersistentGuid=00000000000000000000000000000000,bHidden=False,bNotConnectable=False,bDefaultValueIsReadOnly=False,bDefaultValueIsIgnored=False,bAdvancedView=False,)
+			   CustomProperties Pin (PinId=49D7F9D041F6F18B5714018FAA632144,PinName="Condition",PinType.PinCategory="bool",PinType.PinSubCategoryObject=None,PinType.PinSubCategoryMemberReference=(),PinType.bIsArray=False,PinType.bIsReference=False,PinType.bIsConst=False,PinType.bIsWeakPointer=False,DefaultValue="false",AutogeneratedDefaultValue="true",LinkedTo=(K2Node_CallFunction_295 0A18A59443E7C0D120E7FBA8D4CA0D24,),PersistentGuid=00000000000000000000000000000000,bHidden=False,bNotConnectable=False,bDefaultValueIsReadOnly=False,bDefaultValueIsIgnored=False,bAdvancedView=False,)
+			   CustomProperties Pin (PinId=276C22144FA57477D968698871BE9EC5,PinName="then",PinFriendlyName=NSLOCTEXT("K2Node", "true", "true"),Direction="EGPD_Output",PinType.PinCategory="exec",PinType.PinSubCategoryObject=None,PinType.PinSubCategoryMemberReference=(),PinType.bIsArray=False,PinType.bIsReference=False,PinType.bIsConst=False,PinType.bIsWeakPointer=False,PersistentGuid=00000000000000000000000000000000,bHidden=False,bNotConnectable=False,bDefaultValueIsReadOnly=False,bDefaultValueIsIgnored=False,bAdvancedView=False,)
+			   CustomProperties Pin (PinId=8FBC10034397D4F24C0A29B09D795D95,PinName="else",PinFriendlyName=NSLOCTEXT("K2Node", "false", "false"),Direction="EGPD_Output",PinType.PinCategory="exec",PinType.PinSubCategoryObject=None,PinType.PinSubCategoryMemberReference=(),PinType.bIsArray=False,PinType.bIsReference=False,PinType.bIsConst=False,PinType.bIsWeakPointer=False,LinkedTo=(K2Node_VariableSet_29 3372F24044B22493E4CDB99C2FEE9AC6,),PersistentGuid=00000000000000000000000000000000,bHidden=False,bNotConnectable=False,bDefaultValueIsReadOnly=False,bDefaultValueIsIgnored=False,bAdvancedView=False,)
+			End Object
+			Begin Object Class=EdGraphNode_Comment Name="EdGraphNode_Comment_7"
+			   NodePosX=-112
+			   NodePosY=-2160
+			   NodeWidth=1664
+			   NodeHeight=512
+			   NodeComment="Begin & End overlap to check to see if the character is touching to fluid surface actor. "
+			   NodeGuid=A4E802DA4A1A76C1EC1E608F07C2C02F
+			End Object
+			END OBJECT
+	[/PARAM]
+[/OBJECT]
+[/EXCERPT:Overlap]
+
+[EXCERPT:ShootWhenMouseDown]	
+[OBJECT:Snippet]
+	[PARAMLITERAL:id]
+		ShootWhenMouseDown
+	[/PARAMLITERAL]
+	[PARAMLITERAL:type]
+		blueprint
+	[/PARAMLITERAL]
+	[PARAM:image]
+		![](Engine\Rendering\RenderTargets\BlueprintRenderTargets\HowTo\FluidSurface\FS_Shoot_When_Mouse_Is_Down.png)
+	[/PARAM]
+	[PARAM:snippet]
+			BEGIN OBJECT
+							Begin Object Class=K2Node_Event Name="K2Node_Event_1"
+			   EventReference=(MemberParent=Class'/Script/Engine.Actor',MemberName="ReceiveTick")
+			   bOverrideFunction=True
+			   NodePosX=-1488
+			   NodePosY=-1056
+			   bCommentBubblePinned=True
+			   NodeGuid=EEF6C3BB443B5AC7BC7A61A5F48F78E3
+			   CustomProperties Pin (PinId=62DF0CEF40354973152CACB672685856,PinName="OutputDelegate",Direction="EGPD_Output",PinType.PinCategory="delegate",PinType.PinSubCategoryObject=None,PinType.PinSubCategoryMemberReference=(MemberParent=Class'/Script/Engine.Actor',MemberName="ReceiveTick"),PinType.bIsArray=False,PinType.bIsReference=False,PinType.bIsConst=False,PinType.bIsWeakPointer=False,PersistentGuid=00000000000000000000000000000000,bHidden=False,bNotConnectable=False,bDefaultValueIsReadOnly=False,bDefaultValueIsIgnored=False,bAdvancedView=False,)
+			   CustomProperties Pin (PinId=0E7FB0704C3F4FD4A65171938B5B6AE7,PinName="then",Direction="EGPD_Output",PinType.PinCategory="exec",PinType.PinSubCategoryObject=None,PinType.PinSubCategoryMemberReference=(),PinType.bIsArray=False,PinType.bIsReference=False,PinType.bIsConst=False,PinType.bIsWeakPointer=False,LinkedTo=(K2Node_IfThenElse_285 997FDA76493A40CEE900CD81079C2BFC,),PersistentGuid=00000000000000000000000000000000,bHidden=False,bNotConnectable=False,bDefaultValueIsReadOnly=False,bDefaultValueIsIgnored=False,bAdvancedView=False,)
+			   CustomProperties Pin (PinId=90C5BF9B4C7F93AF4E34319E2020E8D0,PinName="DeltaSeconds",PinToolTip="Delta Seconds\nFloat",Direction="EGPD_Output",PinType.PinCategory="float",PinType.PinSubCategoryObject=None,PinType.PinSubCategoryMemberReference=(),PinType.bIsArray=False,PinType.bIsReference=False,PinType.bIsConst=False,PinType.bIsWeakPointer=False,AutogeneratedDefaultValue="0.0",LinkedTo=(K2Node_CommutativeAssociativeBinaryOperator_6 29EFA03442ADC8249E9CC9B27F215866,),PersistentGuid=00000000000000000000000000000000,bHidden=False,bNotConnectable=False,bDefaultValueIsReadOnly=False,bDefaultValueIsIgnored=False,bAdvancedView=False,)
+			End Object
+			Begin Object Class=K2Node_CallFunction Name="K2Node_CallFunction_2679"
+			   FunctionReference=(MemberName="TraceFromCamera",MemberGuid=C96978014B119460D779C7A90DE17A56,bSelfContext=True)
+			   NodePosX=-1008
+			   NodePosY=-1088
+			   NodeGuid=AE88F4D04DC37ACF12961DAD21CA3BF3
+			   CustomProperties Pin (PinId=9F3F898442AC24B5F566A9A6E3D64A03,PinName="execute",PinToolTip="\nExec",PinType.PinCategory="exec",PinType.PinSubCategoryObject=None,PinType.PinSubCategoryMemberReference=(),PinType.bIsArray=False,PinType.bIsReference=False,PinType.bIsConst=False,PinType.bIsWeakPointer=False,LinkedTo=(K2Node_IfThenElse_285 358D51FF4B717F8D875E549790D394DB,),PersistentGuid=00000000000000000000000000000000,bHidden=False,bNotConnectable=False,bDefaultValueIsReadOnly=False,bDefaultValueIsIgnored=False,bAdvancedView=False,)
+			   CustomProperties Pin (PinId=BF6710D74C9EC10D6D5690B01743CC3D,PinName="then",PinToolTip="\nExec",Direction="EGPD_Output",PinType.PinCategory="exec",PinType.PinSubCategoryObject=None,PinType.PinSubCategoryMemberReference=(),PinType.bIsArray=False,PinType.bIsReference=False,PinType.bIsConst=False,PinType.bIsWeakPointer=False,LinkedTo=(K2Node_IfThenElse_3 2A4FE2894A9586E95A87B99E5295651E,),PersistentGuid=00000000000000000000000000000000,bHidden=False,bNotConnectable=False,bDefaultValueIsReadOnly=False,bDefaultValueIsIgnored=False,bAdvancedView=False,)
+			   CustomProperties Pin (PinId=53418ABD4C9A655761FB14B9217B591D,PinName="self",PinFriendlyName=NSLOCTEXT("K2Node", "Target", "Target"),PinToolTip="Target\nSelf Reference",PinType.PinCategory="object",PinType.PinSubCategory="self",PinType.PinSubCategoryObject=None,PinType.PinSubCategoryMemberReference=(),PinType.bIsArray=False,PinType.bIsReference=False,PinType.bIsConst=False,PinType.bIsWeakPointer=False,PersistentGuid=00000000000000000000000000000000,bHidden=False,bNotConnectable=False,bDefaultValueIsReadOnly=False,bDefaultValueIsIgnored=False,bAdvancedView=False,)
+			End Object
+			Begin Object Class=K2Node_VariableGet Name="K2Node_VariableGet_319"
+			   VariableReference=(MemberName="MouseDown",MemberGuid=209E1D534C15EFFA81760DA2C7F9D56C,bSelfContext=True)
+			   NodePosX=-1472
+			   NodePosY=-933
+			   NodeGuid=49A458A448873AF10EB149A1F18983E8
+			   CustomProperties Pin (PinId=200A09224BC64E9B688537BFADAE5937,PinName="MouseDown",Direction="EGPD_Output",PinType.PinCategory="bool",PinType.PinSubCategoryObject=None,PinType.PinSubCategoryMemberReference=(),PinType.bIsArray=False,PinType.bIsReference=False,PinType.bIsConst=False,PinType.bIsWeakPointer=False,AutogeneratedDefaultValue="false",LinkedTo=(K2Node_IfThenElse_285 DBE9930D445160424F99CAB156B1E26F,),PersistentGuid=00000000000000000000000000000000,bHidden=False,bNotConnectable=False,bDefaultValueIsReadOnly=False,bDefaultValueIsIgnored=False,bAdvancedView=False,)
+			   CustomProperties Pin (PinId=0C6449354ECF7788EE6878B67D1A2964,PinName="self",PinFriendlyName=NSLOCTEXT("K2Node", "Target", "Target"),PinType.PinCategory="object",PinType.PinSubCategoryObject=BlueprintGeneratedClass'/Game/ExampleContent/BlueprintFluidSurface/WaterSurface_BP.WaterSurface_BP_C',PinType.PinSubCategoryMemberReference=(),PinType.bIsArray=False,PinType.bIsReference=False,PinType.bIsConst=False,PinType.bIsWeakPointer=False,PersistentGuid=00000000000000000000000000000000,bHidden=True,bNotConnectable=False,bDefaultValueIsReadOnly=False,bDefaultValueIsIgnored=False,bAdvancedView=False,)
+			End Object
+			Begin Object Class=K2Node_IfThenElse Name="K2Node_IfThenElse_285"
+			   NodePosX=-1248
+			   NodePosY=-1056
+			   NodeGuid=1BFD2D08494E65E52D3F9F9D64B3B5A1
+			   CustomProperties Pin (PinId=997FDA76493A40CEE900CD81079C2BFC,PinName="execute",PinType.PinCategory="exec",PinType.PinSubCategoryObject=None,PinType.PinSubCategoryMemberReference=(),PinType.bIsArray=False,PinType.bIsReference=False,PinType.bIsConst=False,PinType.bIsWeakPointer=False,LinkedTo=(K2Node_Event_1 0E7FB0704C3F4FD4A65171938B5B6AE7,),PersistentGuid=00000000000000000000000000000000,bHidden=False,bNotConnectable=False,bDefaultValueIsReadOnly=False,bDefaultValueIsIgnored=False,bAdvancedView=False,)
+			   CustomProperties Pin (PinId=DBE9930D445160424F99CAB156B1E26F,PinName="Condition",PinType.PinCategory="bool",PinType.PinSubCategoryObject=None,PinType.PinSubCategoryMemberReference=(),PinType.bIsArray=False,PinType.bIsReference=False,PinType.bIsConst=False,PinType.bIsWeakPointer=False,DefaultValue="false",AutogeneratedDefaultValue="true",LinkedTo=(K2Node_VariableGet_319 200A09224BC64E9B688537BFADAE5937,),PersistentGuid=00000000000000000000000000000000,bHidden=False,bNotConnectable=False,bDefaultValueIsReadOnly=False,bDefaultValueIsIgnored=False,bAdvancedView=False,)
+			   CustomProperties Pin (PinId=358D51FF4B717F8D875E549790D394DB,PinName="then",PinFriendlyName=NSLOCTEXT("K2Node", "true", "true"),Direction="EGPD_Output",PinType.PinCategory="exec",PinType.PinSubCategoryObject=None,PinType.PinSubCategoryMemberReference=(),PinType.bIsArray=False,PinType.bIsReference=False,PinType.bIsConst=False,PinType.bIsWeakPointer=False,LinkedTo=(K2Node_CallFunction_2679 9F3F898442AC24B5F566A9A6E3D64A03,),PersistentGuid=00000000000000000000000000000000,bHidden=False,bNotConnectable=False,bDefaultValueIsReadOnly=False,bDefaultValueIsIgnored=False,bAdvancedView=False,)
+			   CustomProperties Pin (PinId=53B9B8674F2C0FDC804A088C1129D0E3,PinName="else",PinFriendlyName=NSLOCTEXT("K2Node", "false", "false"),Direction="EGPD_Output",PinType.PinCategory="exec",PinType.PinSubCategoryObject=None,PinType.PinSubCategoryMemberReference=(),PinType.bIsArray=False,PinType.bIsReference=False,PinType.bIsConst=False,PinType.bIsWeakPointer=False,LinkedTo=(K2Node_Knot_0 D43815774597065AA1121BB03EB92D96,),PersistentGuid=00000000000000000000000000000000,bHidden=False,bNotConnectable=False,bDefaultValueIsReadOnly=False,bDefaultValueIsIgnored=False,bAdvancedView=False,)
+			End Object
+			Begin Object Class=EdGraphNode_Comment Name="EdGraphNode_Comment_25"
+			   NodePosX=-1520
+			   NodePosY=-1136
+			   NodeWidth=768
+			   NodeHeight=272
+			   NodeComment="Shoot when mouse is down"
+			   NodeGuid=C130C7E64BDBE315AE16A8B480C7ED4E
+			End Object
+			Begin Object Class=K2Node_Knot Name="K2Node_Knot_0"
+			   NodePosX=-816
+			   NodePosY=-944
+			   NodeGuid=3BE6B7404E717FC34C49E78A7000CE62
+			   CustomProperties Pin (PinId=D43815774597065AA1121BB03EB92D96,PinName="InputPin",PinType.PinCategory="exec",PinType.PinSubCategoryObject=None,PinType.PinSubCategoryMemberReference=(),PinType.bIsArray=False,PinType.bIsReference=False,PinType.bIsConst=False,PinType.bIsWeakPointer=False,LinkedTo=(K2Node_IfThenElse_285 53B9B8674F2C0FDC804A088C1129D0E3,),PersistentGuid=00000000000000000000000000000000,bHidden=False,bNotConnectable=False,bDefaultValueIsReadOnly=False,bDefaultValueIsIgnored=True,bAdvancedView=False,)
+			   CustomProperties Pin (PinId=8433E1DB49AE64976C75FDB8EA405493,PinName="OutputPin",Direction="EGPD_Output",PinType.PinCategory="exec",PinType.PinSubCategoryObject=None,PinType.PinSubCategoryMemberReference=(),PinType.bIsArray=False,PinType.bIsReference=False,PinType.bIsConst=False,PinType.bIsWeakPointer=False,LinkedTo=(K2Node_IfThenElse_3 2A4FE2894A9586E95A87B99E5295651E,),PersistentGuid=00000000000000000000000000000000,bHidden=False,bNotConnectable=False,bDefaultValueIsReadOnly=False,bDefaultValueIsIgnored=False,bAdvancedView=False,)
+			End Object
+			END OBJECT
+	[/PARAM]
+[/OBJECT]
+[/EXCERPT:ShootWhenMouseDown]
+
+[EXCERPT:TouchingCarForce]	
+[OBJECT:Snippet]
+	[PARAMLITERAL:id]
+		TouchingCarForce
+	[/PARAMLITERAL]
+	[PARAMLITERAL:type]
+		blueprint
+	[/PARAMLITERAL]
+	[PARAM:image]
+		![](Engine\Rendering\RenderTargets\BlueprintRenderTargets\HowTo\FluidSurface\FS_Touching_Char_FOrce.png)
+	[/PARAM]
+	[PARAM:snippet]
+			BEGIN OBJECT
+							Begin Object Class=K2Node_VariableGet Name="K2Node_VariableGet_42"
+			   VariableReference=(MemberName="TouchingCharacter",MemberGuid=22B361484158D2692D5EDAB25083E8C8,bSelfContext=True)
+			   NodePosX=-208
+			   NodePosY=-1305
+			   NodeGuid=72AF85D049CFEAAA74E0B3BC3F65AC78
+			   CustomProperties Pin (PinId=927BAB27420DC48B836265AE156AFD69,PinName="TouchingCharacter",Direction="EGPD_Output",PinType.PinCategory="object",PinType.PinSubCategoryObject=Class'/Script/Engine.Actor',PinType.PinSubCategoryMemberReference=(),PinType.bIsArray=False,PinType.bIsReference=False,PinType.bIsConst=False,PinType.bIsWeakPointer=False,LinkedTo=(K2Node_CallFunction_468 A44FEAF742DF3DB2567A3D8FADC37CF4,K2Node_CallFunction_469 D038F0104F1573B586497F9506F37961,),PersistentGuid=00000000000000000000000000000000,bHidden=False,bNotConnectable=False,bDefaultValueIsReadOnly=False,bDefaultValueIsIgnored=False,bAdvancedView=False,)
+			   CustomProperties Pin (PinId=D589AC5D410BFCAC7D1773858CFFB84A,PinName="self",PinFriendlyName=NSLOCTEXT("K2Node", "Target", "Target"),PinType.PinCategory="object",PinType.PinSubCategoryObject=BlueprintGeneratedClass'/Game/ExampleContent/BlueprintFluidSurface/WaterSurface_BP.WaterSurface_BP_C',PinType.PinSubCategoryMemberReference=(),PinType.bIsArray=False,PinType.bIsReference=False,PinType.bIsConst=False,PinType.bIsWeakPointer=False,PersistentGuid=00000000000000000000000000000000,bHidden=True,bNotConnectable=False,bDefaultValueIsReadOnly=False,bDefaultValueIsIgnored=False,bAdvancedView=False,)
+			End Object
+			Begin Object Class=K2Node_CallFunction Name="K2Node_CallFunction_468"
+			   bIsPureFunc=True
+			   bIsConstFunc=True
+			   FunctionReference=(MemberName="K2_GetActorLocation",bSelfContext=True)
+			   NodePosX=16
+			   NodePosY=-1344
+			   NodeGuid=5BD9206F487D66FCF28CC5BBB93942B1
+			   CustomProperties Pin (PinId=A44FEAF742DF3DB2567A3D8FADC37CF4,PinName="self",PinFriendlyName=NSLOCTEXT("K2Node", "Target", "Target"),PinToolTip="Target\nActor Reference",PinType.PinCategory="object",PinType.PinSubCategoryObject=Class'/Script/Engine.Actor',PinType.PinSubCategoryMemberReference=(),PinType.bIsArray=False,PinType.bIsReference=False,PinType.bIsConst=False,PinType.bIsWeakPointer=False,LinkedTo=(K2Node_VariableGet_42 927BAB27420DC48B836265AE156AFD69,),PersistentGuid=00000000000000000000000000000000,bHidden=False,bNotConnectable=False,bDefaultValueIsReadOnly=False,bDefaultValueIsIgnored=False,bAdvancedView=False,)
+			   CustomProperties Pin (PinId=58B1ADC0421583D3D49F459C2E09B61A,PinName="ReturnValue",PinToolTip="Return Value\nVector ",Direction="EGPD_Output",PinType.PinCategory="struct",PinType.PinSubCategoryObject=ScriptStruct'/Script/CoreUObject.Vector',PinType.PinSubCategoryMemberReference=(),PinType.bIsArray=False,PinType.bIsReference=False,PinType.bIsConst=False,PinType.bIsWeakPointer=False,AutogeneratedDefaultValue="0, 0, 0",LinkedTo=(K2Node_VariableSet_53 DE9F63DA4A6324DDA9E2F0B41DDEF5AB,K2Node_CallFunction_753 757CE9264CA706BECDD331B89650EC21,K2Node_CallFunction_551 E6DC3D054EAFEFE16464859E4A63782A,),PersistentGuid=00000000000000000000000000000000,bHidden=False,bNotConnectable=False,bDefaultValueIsReadOnly=False,bDefaultValueIsIgnored=False,bAdvancedView=False,)
+			End Object
+			Begin Object Class=K2Node_IfThenElse Name="K2Node_IfThenElse_3"
+			   NodePosX=192
+			   NodePosY=-1120
+			   NodeGuid=15ED5A6B48CF0FF35E6EBABC248D2536
+			   CustomProperties Pin (PinId=2A4FE2894A9586E95A87B99E5295651E,PinName="execute",PinType.PinCategory="exec",PinType.PinSubCategoryObject=None,PinType.PinSubCategoryMemberReference=(),PinType.bIsArray=False,PinType.bIsReference=False,PinType.bIsConst=False,PinType.bIsWeakPointer=False,LinkedTo=(K2Node_CallFunction_2679 BF6710D74C9EC10D6D5690B01743CC3D,K2Node_Knot_0 8433E1DB49AE64976C75FDB8EA405493,),PersistentGuid=00000000000000000000000000000000,bHidden=False,bNotConnectable=False,bDefaultValueIsReadOnly=False,bDefaultValueIsIgnored=False,bAdvancedView=False,)
+			   CustomProperties Pin (PinId=03CCB32E44D37449E82EB3A1C583B4D6,PinName="Condition",PinType.PinCategory="bool",PinType.PinSubCategoryObject=None,PinType.PinSubCategoryMemberReference=(),PinType.bIsArray=False,PinType.bIsReference=False,PinType.bIsConst=False,PinType.bIsWeakPointer=False,DefaultValue="false",AutogeneratedDefaultValue="true",LinkedTo=(K2Node_CallFunction_469 4DFAF8054A4E3B9DADE086AA4BE011F0,),PersistentGuid=00000000000000000000000000000000,bHidden=False,bNotConnectable=False,bDefaultValueIsReadOnly=False,bDefaultValueIsIgnored=False,bAdvancedView=False,)
+			   CustomProperties Pin (PinId=F46E8B0E462777612D45B5BA6EA979CD,PinName="then",PinFriendlyName=NSLOCTEXT("K2Node", "true", "true"),Direction="EGPD_Output",PinType.PinCategory="exec",PinType.PinSubCategoryObject=None,PinType.PinSubCategoryMemberReference=(),PinType.bIsArray=False,PinType.bIsReference=False,PinType.bIsConst=False,PinType.bIsWeakPointer=False,LinkedTo=(K2Node_VariableSet_19 879DBE52479842492E01529E9B944C4A,),PersistentGuid=00000000000000000000000000000000,bHidden=False,bNotConnectable=False,bDefaultValueIsReadOnly=False,bDefaultValueIsIgnored=False,bAdvancedView=False,)
+			   CustomProperties Pin (PinId=3552DB0C4941E65DA8FFB4AA61FF3BBD,PinName="else",PinFriendlyName=NSLOCTEXT("K2Node", "false", "false"),Direction="EGPD_Output",PinType.PinCategory="exec",PinType.PinSubCategoryObject=None,PinType.PinSubCategoryMemberReference=(),PinType.bIsArray=False,PinType.bIsReference=False,PinType.bIsConst=False,PinType.bIsWeakPointer=False,LinkedTo=(K2Node_CallFunction_21 EC2F25A34C8219299EF4AD9A874581F8,),PersistentGuid=00000000000000000000000000000000,bHidden=False,bNotConnectable=False,bDefaultValueIsReadOnly=False,bDefaultValueIsIgnored=False,bAdvancedView=False,)
+			End Object
+			Begin Object Class=K2Node_CallFunction Name="K2Node_CallFunction_469"
+			   bIsPureFunc=True
+			   FunctionReference=(MemberParent=Class'/Script/Engine.KismetMathLibrary',MemberName="EqualEqual_ObjectObject")
+			   NodePosX=-32
+			   NodePosY=-1216
+			   NodeGuid=31F166C940EEF0A5531F30B87613D8A6
+			   CustomProperties Pin (PinId=4005B9544C8A3E412C14B7BF81D904C5,PinName="self",PinFriendlyName=NSLOCTEXT("K2Node", "Target", "Target"),PinToolTip="Target\nKismet Math Library Reference",PinType.PinCategory="object",PinType.PinSubCategoryObject=Class'/Script/Engine.KismetMathLibrary',PinType.PinSubCategoryMemberReference=(),PinType.bIsArray=False,PinType.bIsReference=False,PinType.bIsConst=False,PinType.bIsWeakPointer=False,DefaultObject="/Script/Engine.Default__KismetMathLibrary",PersistentGuid=00000000000000000000000000000000,bHidden=True,bNotConnectable=False,bDefaultValueIsReadOnly=False,bDefaultValueIsIgnored=False,bAdvancedView=False,)
+			   CustomProperties Pin (PinId=D038F0104F1573B586497F9506F37961,PinName="A",PinToolTip="A\nObject Reference",PinType.PinCategory="object",PinType.PinSubCategoryObject=Class'/Script/CoreUObject.Object',PinType.PinSubCategoryMemberReference=(),PinType.bIsArray=False,PinType.bIsReference=False,PinType.bIsConst=False,PinType.bIsWeakPointer=False,LinkedTo=(K2Node_VariableGet_42 927BAB27420DC48B836265AE156AFD69,),PersistentGuid=00000000000000000000000000000000,bHidden=False,bNotConnectable=False,bDefaultValueIsReadOnly=False,bDefaultValueIsIgnored=False,bAdvancedView=False,)
+			   CustomProperties Pin (PinId=232428D94517DE33FD6A1290DD6AE0C0,PinName="B",PinToolTip="B\nObject Reference",PinType.PinCategory="object",PinType.PinSubCategoryObject=Class'/Script/CoreUObject.Object',PinType.PinSubCategoryMemberReference=(),PinType.bIsArray=False,PinType.bIsReference=False,PinType.bIsConst=False,PinType.bIsWeakPointer=False,PersistentGuid=00000000000000000000000000000000,bHidden=False,bNotConnectable=False,bDefaultValueIsReadOnly=False,bDefaultValueIsIgnored=False,bAdvancedView=False,)
+			   CustomProperties Pin (PinId=4DFAF8054A4E3B9DADE086AA4BE011F0,PinName="ReturnValue",PinToolTip="Return Value\nBoolean",Direction="EGPD_Output",PinType.PinCategory="bool",PinType.PinSubCategoryObject=None,PinType.PinSubCategoryMemberReference=(),PinType.bIsArray=False,PinType.bIsReference=False,PinType.bIsConst=False,PinType.bIsWeakPointer=False,AutogeneratedDefaultValue="false",LinkedTo=(K2Node_IfThenElse_3 03CCB32E44D37449E82EB3A1C583B4D6,),PersistentGuid=00000000000000000000000000000000,bHidden=False,bNotConnectable=False,bDefaultValueIsReadOnly=False,bDefaultValueIsIgnored=False,bAdvancedView=False,)
+			End Object
+			Begin Object Class=K2Node_VariableSet Name="K2Node_VariableSet_53"
+			   VariableReference=(MemberName="LastTouchingActorPosition",MemberGuid=C849AC4745DA47383E1884AA8AA77BC9,bSelfContext=True)
+			   NodePosX=2128
+			   NodePosY=-1104
+			   NodeGuid=F78F09404F324F17C98A6491FB6BF2FB
+			   CustomProperties Pin (PinId=E5EF78B644B6CB8EAD2E229407D67DE2,PinName="execute",PinType.PinCategory="exec",PinType.PinSubCategoryObject=None,PinType.PinSubCategoryMemberReference=(),PinType.bIsArray=False,PinType.bIsReference=False,PinType.bIsConst=False,PinType.bIsWeakPointer=False,LinkedTo=(K2Node_CallFunction_144 C312F52A41D4BEDBB220AFA395DD1F9F,),PersistentGuid=00000000000000000000000000000000,bHidden=False,bNotConnectable=False,bDefaultValueIsReadOnly=False,bDefaultValueIsIgnored=False,bAdvancedView=False,)
+			   CustomProperties Pin (PinId=140B85434191688AB8F2A1A2C156091A,PinName="then",Direction="EGPD_Output",PinType.PinCategory="exec",PinType.PinSubCategoryObject=None,PinType.PinSubCategoryMemberReference=(),PinType.bIsArray=False,PinType.bIsReference=False,PinType.bIsConst=False,PinType.bIsWeakPointer=False,LinkedTo=(K2Node_VariableSet_19 879DBE52479842492E01529E9B944C4A,),PersistentGuid=00000000000000000000000000000000,bHidden=False,bNotConnectable=False,bDefaultValueIsReadOnly=False,bDefaultValueIsIgnored=False,bAdvancedView=False,)
+			   CustomProperties Pin (PinId=DE9F63DA4A6324DDA9E2F0B41DDEF5AB,PinName="LastTouchingActorPosition",PinType.PinCategory="struct",PinType.PinSubCategoryObject=ScriptStruct'/Script/CoreUObject.Vector',PinType.PinSubCategoryMemberReference=(),PinType.bIsArray=False,PinType.bIsReference=False,PinType.bIsConst=False,PinType.bIsWeakPointer=False,DefaultValue="0, 0, 0",AutogeneratedDefaultValue="0, 0, 0",LinkedTo=(K2Node_CallFunction_468 58B1ADC0421583D3D49F459C2E09B61A,),PersistentGuid=00000000000000000000000000000000,bHidden=False,bNotConnectable=False,bDefaultValueIsReadOnly=False,bDefaultValueIsIgnored=False,bAdvancedView=False,)
+			   CustomProperties Pin (PinId=1C5B504644AA0A3B5BD2F7B068F8DE96,PinName="Output_Get",PinToolTip="Retrieves the value of the variable, can use instead of a separate Get node",Direction="EGPD_Output",PinType.PinCategory="struct",PinType.PinSubCategoryObject=ScriptStruct'/Script/CoreUObject.Vector',PinType.PinSubCategoryMemberReference=(),PinType.bIsArray=False,PinType.bIsReference=False,PinType.bIsConst=False,PinType.bIsWeakPointer=False,DefaultValue="0, 0, 0",AutogeneratedDefaultValue="0, 0, 0",PersistentGuid=00000000000000000000000000000000,bHidden=False,bNotConnectable=False,bDefaultValueIsReadOnly=False,bDefaultValueIsIgnored=False,bAdvancedView=False,)
+			   CustomProperties Pin (PinId=DBF466D643EF53401B12A6AD5D5F86ED,PinName="self",PinFriendlyName=NSLOCTEXT("K2Node", "Target", "Target"),PinType.PinCategory="object",PinType.PinSubCategoryObject=BlueprintGeneratedClass'/Game/ExampleContent/BlueprintFluidSurface/WaterSurface_BP.WaterSurface_BP_C',PinType.PinSubCategoryMemberReference=(),PinType.bIsArray=False,PinType.bIsReference=False,PinType.bIsConst=False,PinType.bIsWeakPointer=False,PersistentGuid=00000000000000000000000000000000,bHidden=True,bNotConnectable=False,bDefaultValueIsReadOnly=False,bDefaultValueIsIgnored=False,bAdvancedView=False,)
+			End Object
+			Begin Object Class=K2Node_CallFunction Name="K2Node_CallFunction_21"
+			   FunctionReference=(MemberName="GetHeightRT",MemberGuid=2EEA338F4CA878E7228A10BB2DD63E00,bSelfContext=True)
+			   NodePosX=512
+			   NodePosY=-1104
+			   NodeGuid=3D0426964A8819269E65E8989258DF7B
+			   CustomProperties Pin (PinId=EC2F25A34C8219299EF4AD9A874581F8,PinName="execute",PinType.PinCategory="exec",PinType.PinSubCategoryObject=None,PinType.PinSubCategoryMemberReference=(),PinType.bIsArray=False,PinType.bIsReference=False,PinType.bIsConst=False,PinType.bIsWeakPointer=False,LinkedTo=(K2Node_IfThenElse_3 3552DB0C4941E65DA8FFB4AA61FF3BBD,),PersistentGuid=00000000000000000000000000000000,bHidden=False,bNotConnectable=False,bDefaultValueIsReadOnly=False,bDefaultValueIsIgnored=False,bAdvancedView=False,)
+			   CustomProperties Pin (PinId=95EA835544AB1FA43D2999A887130F63,PinName="then",Direction="EGPD_Output",PinType.PinCategory="exec",PinType.PinSubCategoryObject=None,PinType.PinSubCategoryMemberReference=(),PinType.bIsArray=False,PinType.bIsReference=False,PinType.bIsConst=False,PinType.bIsWeakPointer=False,LinkedTo=(K2Node_CallMaterialParameterCollectionFunction_21 319E141F49655CE8105AC0B2A2955CA9,),PersistentGuid=00000000000000000000000000000000,bHidden=False,bNotConnectable=False,bDefaultValueIsReadOnly=False,bDefaultValueIsIgnored=False,bAdvancedView=False,)
+			   CustomProperties Pin (PinId=A902316E45FADD64A2401EBA8075724D,PinName="self",PinFriendlyName=NSLOCTEXT("K2Node", "Target", "Target"),PinType.PinCategory="object",PinType.PinSubCategory="self",PinType.PinSubCategoryObject=None,PinType.PinSubCategoryMemberReference=(),PinType.bIsArray=False,PinType.bIsReference=False,PinType.bIsConst=False,PinType.bIsWeakPointer=False,PersistentGuid=00000000000000000000000000000000,bHidden=False,bNotConnectable=False,bDefaultValueIsReadOnly=False,bDefaultValueIsIgnored=False,bAdvancedView=False,)
+			   CustomProperties Pin (PinId=31FE1BBB4CDD998BFD5D2BAA5D4A3EA0,PinName="Index",PinType.PinCategory="int",PinType.PinSubCategoryObject=None,PinType.PinSubCategoryMemberReference=(),PinType.bIsArray=False,PinType.bIsReference=False,PinType.bIsConst=False,PinType.bIsWeakPointer=False,DefaultValue="0",AutogeneratedDefaultValue="0",LinkedTo=(K2Node_VariableGet_146 C83EBC6B40A3F46564CC9A833477D9E2,),PersistentGuid=00000000000000000000000000000000,bHidden=False,bNotConnectable=False,bDefaultValueIsReadOnly=False,bDefaultValueIsIgnored=False,bAdvancedView=False,)
+			   CustomProperties Pin (PinId=F349F78749FD3E5ADB14DE8F5E164315,PinName="HeightRT",Direction="EGPD_Output",PinType.PinCategory="object",PinType.PinSubCategoryObject=Class'/Script/Engine.TextureRenderTarget2D',PinType.PinSubCategoryMemberReference=(),PinType.bIsArray=False,PinType.bIsReference=False,PinType.bIsConst=False,PinType.bIsWeakPointer=False,LinkedTo=(K2Node_CallFunction_144 D2975F3E4A669D9FC42692A4B0EFFDE5,),PersistentGuid=00000000000000000000000000000000,bHidden=False,bNotConnectable=False,bDefaultValueIsReadOnly=False,bDefaultValueIsIgnored=False,bAdvancedView=False,)
+			End Object
+			Begin Object Class=K2Node_VariableGet Name="K2Node_VariableGet_146"
+			   VariableReference=(MemberName="HeightState",MemberGuid=8E10E24F4154A224C2065A89B367BD62,bSelfContext=True)
+			   NodePosX=368
+			   NodePosY=-1008
+			   NodeGuid=BAC7A5774397B56B380A53AED1E8B0C1
+			   CustomProperties Pin (PinId=C83EBC6B40A3F46564CC9A833477D9E2,PinName="HeightState",Direction="EGPD_Output",PinType.PinCategory="int",PinType.PinSubCategoryObject=None,PinType.PinSubCategoryMemberReference=(),PinType.bIsArray=False,PinType.bIsReference=False,PinType.bIsConst=False,PinType.bIsWeakPointer=False,AutogeneratedDefaultValue="0",LinkedTo=(K2Node_CallFunction_21 31FE1BBB4CDD998BFD5D2BAA5D4A3EA0,),PersistentGuid=00000000000000000000000000000000,bHidden=False,bNotConnectable=False,bDefaultValueIsReadOnly=False,bDefaultValueIsIgnored=False,bAdvancedView=False,)
+			   CustomProperties Pin (PinId=321C448F4298CBF159EF6F99050E49CD,PinName="self",PinFriendlyName=NSLOCTEXT("K2Node", "Target", "Target"),PinType.PinCategory="object",PinType.PinSubCategoryObject=BlueprintGeneratedClass'/Game/ExampleContent/BlueprintFluidSurface/WaterSurface_BP.WaterSurface_BP_C',PinType.PinSubCategoryMemberReference=(),PinType.bIsArray=False,PinType.bIsReference=False,PinType.bIsConst=False,PinType.bIsWeakPointer=False,PersistentGuid=00000000000000000000000000000000,bHidden=True,bNotConnectable=False,bDefaultValueIsReadOnly=False,bDefaultValueIsIgnored=False,bAdvancedView=False,)
+			End Object
+			Begin Object Class=K2Node_CallMaterialParameterCollectionFunction Name="K2Node_CallMaterialParameterCollectionFunction_21"
+			   FunctionReference=(MemberParent=Class'/Script/Engine.KismetMaterialLibrary',MemberName="CreateDynamicMaterialInstance")
+			   NodePosX=816
+			   NodePosY=-1120
+			   NodeGuid=D12CDA354B2B17173E40588644F1EC4E
+			   CustomProperties Pin (PinId=319E141F49655CE8105AC0B2A2955CA9,PinName="execute",PinToolTip="\nExec",PinType.PinCategory="exec",PinType.PinSubCategoryObject=None,PinType.PinSubCategoryMemberReference=(),PinType.bIsArray=False,PinType.bIsReference=False,PinType.bIsConst=False,PinType.bIsWeakPointer=False,LinkedTo=(K2Node_CallFunction_21 95EA835544AB1FA43D2999A887130F63,),PersistentGuid=00000000000000000000000000000000,bHidden=False,bNotConnectable=False,bDefaultValueIsReadOnly=False,bDefaultValueIsIgnored=False,bAdvancedView=False,)
+			   CustomProperties Pin (PinId=0D2AE76648F65D9BB655518B1E728178,PinName="then",PinToolTip="\nExec",Direction="EGPD_Output",PinType.PinCategory="exec",PinType.PinSubCategoryObject=None,PinType.PinSubCategoryMemberReference=(),PinType.bIsArray=False,PinType.bIsReference=False,PinType.bIsConst=False,PinType.bIsWeakPointer=False,LinkedTo=(K2Node_CallFunction_1036 6D1A8B1E4CDD559685223E87CBC7D3C6,),PersistentGuid=00000000000000000000000000000000,bHidden=False,bNotConnectable=False,bDefaultValueIsReadOnly=False,bDefaultValueIsIgnored=False,bAdvancedView=False,)
+			   CustomProperties Pin (PinId=0773178D477AA8E5F2E014B5C84BBA82,PinName="self",PinFriendlyName=NSLOCTEXT("K2Node", "Target", "Target"),PinToolTip="Target\nKismet Material Library Reference",PinType.PinCategory="object",PinType.PinSubCategoryObject=Class'/Script/Engine.KismetMaterialLibrary',PinType.PinSubCategoryMemberReference=(),PinType.bIsArray=False,PinType.bIsReference=False,PinType.bIsConst=False,PinType.bIsWeakPointer=False,DefaultObject="/Script/Engine.Default__KismetMaterialLibrary",PersistentGuid=00000000000000000000000000000000,bHidden=True,bNotConnectable=False,bDefaultValueIsReadOnly=False,bDefaultValueIsIgnored=False,bAdvancedView=False,)
+			   CustomProperties Pin (PinId=FD9828C0491A00E2AA47DDAF6BFC2DE9,PinName="WorldContextObject",PinToolTip="World Context Object\nObject Reference",PinType.PinCategory="object",PinType.PinSubCategoryObject=Class'/Script/CoreUObject.Object',PinType.PinSubCategoryMemberReference=(),PinType.bIsArray=False,PinType.bIsReference=False,PinType.bIsConst=False,PinType.bIsWeakPointer=False,PersistentGuid=00000000000000000000000000000000,bHidden=True,bNotConnectable=False,bDefaultValueIsReadOnly=False,bDefaultValueIsIgnored=False,bAdvancedView=False,)
+			   CustomProperties Pin (PinId=5E2FC1964A86D962DA7C2DA7ECCAEA22,PinName="Parent",PinToolTip="Parent\nMaterial Interface Reference",PinType.PinCategory="object",PinType.PinSubCategoryObject=Class'/Script/Engine.MaterialInterface',PinType.PinSubCategoryMemberReference=(),PinType.bIsArray=False,PinType.bIsReference=False,PinType.bIsConst=False,PinType.bIsWeakPointer=False,DefaultObject="/Game/ExampleContent/BlueprintFluidSurface/ForceSplat.ForceSplat",PersistentGuid=00000000000000000000000000000000,bHidden=False,bNotConnectable=False,bDefaultValueIsReadOnly=False,bDefaultValueIsIgnored=False,bAdvancedView=False,)
+			   CustomProperties Pin (PinId=A3A63E4B40398A36A1C1C5A923BDC83E,PinName="ReturnValue",PinToolTip="Return Value\nMaterial Instance Dynamic Reference",Direction="EGPD_Output",PinType.PinCategory="object",PinType.PinSubCategoryObject=Class'/Script/Engine.MaterialInstanceDynamic',PinType.PinSubCategoryMemberReference=(),PinType.bIsArray=False,PinType.bIsReference=False,PinType.bIsConst=False,PinType.bIsWeakPointer=False,LinkedTo=(K2Node_CallFunction_1036 D48996374313B923726C3CA3E5A4FDC3,K2Node_CallFunction_2200 0D70796D4BD1BFABCED6F08008A7F1AE,K2Node_CallFunction_144 9E1B05684D5F28E5E0D4E49A7BC18BA8,),PersistentGuid=00000000000000000000000000000000,bHidden=False,bNotConnectable=False,bDefaultValueIsReadOnly=False,bDefaultValueIsIgnored=False,bAdvancedView=False,)
+			End Object
+			Begin Object Class=K2Node_CallFunction Name="K2Node_CallFunction_1036"
+			   FunctionReference=(MemberParent=Class'/Script/Engine.MaterialInstanceDynamic',MemberName="SetVectorParameterValue")
+			   NodePosX=1200
+			   NodePosY=-1136
+			   NodeGuid=BDE44EFB46F8636A12824F8797A30DB4
+			   CustomProperties Pin (PinId=6D1A8B1E4CDD559685223E87CBC7D3C6,PinName="execute",PinToolTip="\nExec",PinType.PinCategory="exec",PinType.PinSubCategoryObject=None,PinType.PinSubCategoryMemberReference=(),PinType.bIsArray=False,PinType.bIsReference=False,PinType.bIsConst=False,PinType.bIsWeakPointer=False,LinkedTo=(K2Node_CallMaterialParameterCollectionFunction_21 0D2AE76648F65D9BB655518B1E728178,),PersistentGuid=00000000000000000000000000000000,bHidden=False,bNotConnectable=False,bDefaultValueIsReadOnly=False,bDefaultValueIsIgnored=False,bAdvancedView=False,)
+			   CustomProperties Pin (PinId=BD299B3A459706D661AE5CB3B94FC4E5,PinName="then",PinToolTip="\nExec",Direction="EGPD_Output",PinType.PinCategory="exec",PinType.PinSubCategoryObject=None,PinType.PinSubCategoryMemberReference=(),PinType.bIsArray=False,PinType.bIsReference=False,PinType.bIsConst=False,PinType.bIsWeakPointer=False,LinkedTo=(K2Node_CallFunction_2200 0722514B4563FDB08CA516BC8FADE4E3,),PersistentGuid=00000000000000000000000000000000,bHidden=False,bNotConnectable=False,bDefaultValueIsReadOnly=False,bDefaultValueIsIgnored=False,bAdvancedView=False,)
+			   CustomProperties Pin (PinId=D48996374313B923726C3CA3E5A4FDC3,PinName="self",PinFriendlyName=NSLOCTEXT("K2Node", "Target", "Target"),PinToolTip="Target\nMaterial Instance Dynamic Reference",PinType.PinCategory="object",PinType.PinSubCategoryObject=Class'/Script/Engine.MaterialInstanceDynamic',PinType.PinSubCategoryMemberReference=(),PinType.bIsArray=False,PinType.bIsReference=False,PinType.bIsConst=False,PinType.bIsWeakPointer=False,LinkedTo=(K2Node_CallMaterialParameterCollectionFunction_21 A3A63E4B40398A36A1C1C5A923BDC83E,),PersistentGuid=00000000000000000000000000000000,bHidden=False,bNotConnectable=False,bDefaultValueIsReadOnly=False,bDefaultValueIsIgnored=False,bAdvancedView=False,)
+			   CustomProperties Pin (PinId=32EEEB0B4A32FB6B0FE792B1C387A5F7,PinName="ParameterName",PinToolTip="Parameter Name\nName",PinType.PinCategory="name",PinType.PinSubCategoryObject=None,PinType.PinSubCategoryMemberReference=(),PinType.bIsArray=False,PinType.bIsReference=False,PinType.bIsConst=False,PinType.bIsWeakPointer=False,DefaultValue="ForcePosition",AutogeneratedDefaultValue="None",PersistentGuid=00000000000000000000000000000000,bHidden=False,bNotConnectable=False,bDefaultValueIsReadOnly=False,bDefaultValueIsIgnored=False,bAdvancedView=False,)
+			   CustomProperties Pin (PinId=76AFEFAD4255ADFAEB01C6A4E0FA5BDD,PinName="Value",PinToolTip="Value\nLinear Color Structure",PinType.PinCategory="struct",PinType.PinSubCategoryObject=ScriptStruct'/Script/CoreUObject.LinearColor',PinType.PinSubCategoryMemberReference=(),PinType.bIsArray=False,PinType.bIsReference=False,PinType.bIsConst=False,PinType.bIsWeakPointer=False,LinkedTo=(K2Node_CallFunction_2595 B88496B04962010CC02558BD9CFBC38B,),PersistentGuid=00000000000000000000000000000000,bHidden=False,bNotConnectable=False,bDefaultValueIsReadOnly=False,bDefaultValueIsIgnored=False,bAdvancedView=False,)
+			End Object
+			Begin Object Class=K2Node_CallFunction Name="K2Node_CallFunction_551"
+			   bIsPureFunc=True
+			   FunctionReference=(MemberParent=Class'/Script/Engine.KismetMathLibrary',MemberName="TransformLocation")
+			   NodePosX=576
+			   NodePosY=-944
+			   NodeGuid=3E97BC8D4BB274ACA15BD4A477665C39
+			   CustomProperties Pin (PinId=2658CE464BFFCDF0CC3A7CBAD96A0F9B,PinName="self",PinFriendlyName=NSLOCTEXT("K2Node", "Target", "Target"),PinToolTip="Target\nKismet Math Library Reference",PinType.PinCategory="object",PinType.PinSubCategoryObject=Class'/Script/Engine.KismetMathLibrary',PinType.PinSubCategoryMemberReference=(),PinType.bIsArray=False,PinType.bIsReference=False,PinType.bIsConst=False,PinType.bIsWeakPointer=False,DefaultObject="/Script/Engine.Default__KismetMathLibrary",PersistentGuid=00000000000000000000000000000000,bHidden=True,bNotConnectable=False,bDefaultValueIsReadOnly=False,bDefaultValueIsIgnored=False,bAdvancedView=False,)
+			   CustomProperties Pin (PinId=8E09869B41E1BBCF45AC9FB29C23D7D0,PinName="T",PinToolTip="T\nTransform  (by ref)",PinType.PinCategory="struct",PinType.PinSubCategoryObject=ScriptStruct'/Script/CoreUObject.Transform',PinType.PinSubCategoryMemberReference=(),PinType.bIsArray=False,PinType.bIsReference=True,PinType.bIsConst=True,PinType.bIsWeakPointer=False,LinkedTo=(K2Node_CallFunction_553 F24B5D2E450BFB9AEFA28C9702C5AF99,),PersistentGuid=00000000000000000000000000000000,bHidden=False,bNotConnectable=False,bDefaultValueIsReadOnly=False,bDefaultValueIsIgnored=True,bAdvancedView=False,)
+			   CustomProperties Pin (PinId=E6DC3D054EAFEFE16464859E4A63782A,PinName="Location",PinToolTip="Location\nVector ",PinType.PinCategory="struct",PinType.PinSubCategoryObject=ScriptStruct'/Script/CoreUObject.Vector',PinType.PinSubCategoryMemberReference=(),PinType.bIsArray=False,PinType.bIsReference=False,PinType.bIsConst=False,PinType.bIsWeakPointer=False,DefaultValue="0, 0, 0",AutogeneratedDefaultValue="0, 0, 0",LinkedTo=(K2Node_CallFunction_468 58B1ADC0421583D3D49F459C2E09B61A,),PersistentGuid=00000000000000000000000000000000,bHidden=False,bNotConnectable=False,bDefaultValueIsReadOnly=False,bDefaultValueIsIgnored=False,bAdvancedView=False,)
+			   CustomProperties Pin (PinId=6EB290BC41F4A130C5A6DD823E31F7ED,PinName="ReturnValue",PinToolTip="Return Value\nVector ",Direction="EGPD_Output",PinType.PinCategory="struct",PinType.PinSubCategoryObject=ScriptStruct'/Script/CoreUObject.Vector',PinType.PinSubCategoryMemberReference=(),PinType.bIsArray=False,PinType.bIsReference=False,PinType.bIsConst=False,PinType.bIsWeakPointer=False,AutogeneratedDefaultValue="0, 0, 0",LinkedTo=(K2Node_CallFunction_750 51592DFF4F2D15E8F3FC8A9A2F7C1656,),PersistentGuid=00000000000000000000000000000000,bHidden=False,bNotConnectable=False,bDefaultValueIsReadOnly=False,bDefaultValueIsIgnored=False,bAdvancedView=False,)
+			End Object
+			Begin Object Class=K2Node_CallFunction Name="K2Node_CallFunction_552"
+			   bIsPureFunc=True
+			   bIsConstFunc=True
+			   FunctionReference=(MemberParent=Class'/Script/Engine.SceneComponent',MemberName="K2_GetComponentToWorld")
+			   NodePosX=32
+			   NodePosY=-960
+			   NodeGuid=25C1CE9E425BE69501C0F5A498AB22D1
+			   CustomProperties Pin (PinId=1B9D37834BED70BF9F32EF88300929F5,PinName="self",PinFriendlyName=NSLOCTEXT("K2Node", "Target", "Target"),PinToolTip="Target\nScene Component Reference",PinType.PinCategory="object",PinType.PinSubCategoryObject=Class'/Script/Engine.SceneComponent',PinType.PinSubCategoryMemberReference=(),PinType.bIsArray=False,PinType.bIsReference=False,PinType.bIsConst=False,PinType.bIsWeakPointer=False,LinkedTo=(K2Node_VariableGet_30 415E187944F468E03347EFBF681ADBBF,),PersistentGuid=00000000000000000000000000000000,bHidden=False,bNotConnectable=False,bDefaultValueIsReadOnly=False,bDefaultValueIsIgnored=False,bAdvancedView=False,)
+			   CustomProperties Pin (PinId=169DCEB3496FE9E1900005B95AF4C88D,PinName="ReturnValue",PinToolTip="Return Value\nTransform ",Direction="EGPD_Output",PinType.PinCategory="struct",PinType.PinSubCategoryObject=ScriptStruct'/Script/CoreUObject.Transform',PinType.PinSubCategoryMemberReference=(),PinType.bIsArray=False,PinType.bIsReference=False,PinType.bIsConst=False,PinType.bIsWeakPointer=False,LinkedTo=(K2Node_CallFunction_553 571DF80A435DC232007F8AADAFF695AA,),PersistentGuid=00000000000000000000000000000000,bHidden=False,bNotConnectable=False,bDefaultValueIsReadOnly=False,bDefaultValueIsIgnored=False,bAdvancedView=False,)
+			End Object
+			Begin Object Class=K2Node_CallFunction Name="K2Node_CallFunction_553"
+			   bIsPureFunc=True
+			   FunctionReference=(MemberParent=Class'/Script/Engine.KismetMathLibrary',MemberName="InvertTransform")
+			   NodePosX=320
+			   NodePosY=-944
+			   NodeGuid=899F21B24CE7D1FA879DC0B7C7503BDA
+			   CustomProperties Pin (PinId=1C8F060D4F7F190277735098D64087AA,PinName="self",PinFriendlyName=NSLOCTEXT("K2Node", "Target", "Target"),PinToolTip="Target\nKismet Math Library Reference",PinType.PinCategory="object",PinType.PinSubCategoryObject=Class'/Script/Engine.KismetMathLibrary',PinType.PinSubCategoryMemberReference=(),PinType.bIsArray=False,PinType.bIsReference=False,PinType.bIsConst=False,PinType.bIsWeakPointer=False,DefaultObject="/Script/Engine.Default__KismetMathLibrary",PersistentGuid=00000000000000000000000000000000,bHidden=True,bNotConnectable=False,bDefaultValueIsReadOnly=False,bDefaultValueIsIgnored=False,bAdvancedView=False,)
+			   CustomProperties Pin (PinId=571DF80A435DC232007F8AADAFF695AA,PinName="T",PinToolTip="T\nTransform  (by ref)\n\nThe transform you wish to invert",PinType.PinCategory="struct",PinType.PinSubCategoryObject=ScriptStruct'/Script/CoreUObject.Transform',PinType.PinSubCategoryMemberReference=(),PinType.bIsArray=False,PinType.bIsReference=True,PinType.bIsConst=True,PinType.bIsWeakPointer=False,LinkedTo=(K2Node_CallFunction_552 169DCEB3496FE9E1900005B95AF4C88D,),PersistentGuid=00000000000000000000000000000000,bHidden=False,bNotConnectable=False,bDefaultValueIsReadOnly=False,bDefaultValueIsIgnored=True,bAdvancedView=False,)
+			   CustomProperties Pin (PinId=F24B5D2E450BFB9AEFA28C9702C5AF99,PinName="ReturnValue",PinToolTip="Return Value\nTransform \n\nThe inverse of T.",Direction="EGPD_Output",PinType.PinCategory="struct",PinType.PinSubCategoryObject=ScriptStruct'/Script/CoreUObject.Transform',PinType.PinSubCategoryMemberReference=(),PinType.bIsArray=False,PinType.bIsReference=False,PinType.bIsConst=False,PinType.bIsWeakPointer=False,LinkedTo=(K2Node_CallFunction_551 8E09869B41E1BBCF45AC9FB29C23D7D0,),PersistentGuid=00000000000000000000000000000000,bHidden=False,bNotConnectable=False,bDefaultValueIsReadOnly=False,bDefaultValueIsIgnored=False,bAdvancedView=False,)
+			End Object
+			Begin Object Class=K2Node_CallFunction Name="K2Node_CallFunction_2595"
+			   bIsPureFunc=True
+			   FunctionReference=(MemberParent=Class'/Script/Engine.KismetMathLibrary',MemberName="Conv_VectorToLinearColor")
+			   NodePosX=976
+			   NodePosY=-912
+			   NodeGuid=19F8084643ABFFF45FF8878C082120F2
+			   CustomProperties Pin (PinId=EED9FCB5467B69439D0DDC849B9C32B0,PinName="self",PinFriendlyName=NSLOCTEXT("K2Node", "Target", "Target"),PinType.PinCategory="object",PinType.PinSubCategoryObject=Class'/Script/Engine.KismetMathLibrary',PinType.PinSubCategoryMemberReference=(),PinType.bIsArray=False,PinType.bIsReference=False,PinType.bIsConst=False,PinType.bIsWeakPointer=False,DefaultObject="/Script/Engine.Default__KismetMathLibrary",PersistentGuid=00000000000000000000000000000000,bHidden=True,bNotConnectable=False,bDefaultValueIsReadOnly=False,bDefaultValueIsIgnored=False,bAdvancedView=False,)
+			   CustomProperties Pin (PinId=DEA26338429AED62B1FE2E905F85720A,PinName="InVec",PinType.PinCategory="struct",PinType.PinSubCategoryObject=ScriptStruct'/Script/CoreUObject.Vector',PinType.PinSubCategoryMemberReference=(),PinType.bIsArray=False,PinType.bIsReference=False,PinType.bIsConst=False,PinType.bIsWeakPointer=False,DefaultValue="0, 0, 0",AutogeneratedDefaultValue="0, 0, 0",LinkedTo=(K2Node_CallFunction_750 05ABE7DC481D5CE9B683148BBC283658,),PersistentGuid=00000000000000000000000000000000,bHidden=False,bNotConnectable=False,bDefaultValueIsReadOnly=False,bDefaultValueIsIgnored=False,bAdvancedView=False,)
+			   CustomProperties Pin (PinId=B88496B04962010CC02558BD9CFBC38B,PinName="ReturnValue",Direction="EGPD_Output",PinType.PinCategory="struct",PinType.PinSubCategoryObject=ScriptStruct'/Script/CoreUObject.LinearColor',PinType.PinSubCategoryMemberReference=(),PinType.bIsArray=False,PinType.bIsReference=False,PinType.bIsConst=False,PinType.bIsWeakPointer=False,LinkedTo=(K2Node_CallFunction_1036 76AFEFAD4255ADFAEB01C6A4E0FA5BDD,),PersistentGuid=00000000000000000000000000000000,bHidden=False,bNotConnectable=False,bDefaultValueIsReadOnly=False,bDefaultValueIsIgnored=False,bAdvancedView=False,)
+			End Object
+			Begin Object Class=K2Node_CallFunction Name="K2Node_CallFunction_750"
+			   bIsPureFunc=True
+			   FunctionReference=(MemberParent=Class'/Script/Engine.KismetMathLibrary',MemberName="Divide_VectorFloat")
+			   NodePosX=816
+			   NodePosY=-848
+			   NodeGuid=0170C3F04A490BCC1D9EB493F2DF5BCE
+			   CustomProperties Pin (PinId=329B30E844989AB1610815BDA88D129E,PinName="self",PinFriendlyName=NSLOCTEXT("K2Node", "Target", "Target"),PinToolTip="Target\nKismet Math Library Reference",PinType.PinCategory="object",PinType.PinSubCategoryObject=Class'/Script/Engine.KismetMathLibrary',PinType.PinSubCategoryMemberReference=(),PinType.bIsArray=False,PinType.bIsReference=False,PinType.bIsConst=False,PinType.bIsWeakPointer=False,DefaultObject="/Script/Engine.Default__KismetMathLibrary",PersistentGuid=00000000000000000000000000000000,bHidden=True,bNotConnectable=False,bDefaultValueIsReadOnly=False,bDefaultValueIsIgnored=False,bAdvancedView=False,)
+			   CustomProperties Pin (PinId=51592DFF4F2D15E8F3FC8A9A2F7C1656,PinName="A",PinToolTip="A\nVector ",PinType.PinCategory="struct",PinType.PinSubCategoryObject=ScriptStruct'/Script/CoreUObject.Vector',PinType.PinSubCategoryMemberReference=(),PinType.bIsArray=False,PinType.bIsReference=False,PinType.bIsConst=False,PinType.bIsWeakPointer=False,DefaultValue="0, 0, 0",AutogeneratedDefaultValue="0, 0, 0",LinkedTo=(K2Node_CallFunction_551 6EB290BC41F4A130C5A6DD823E31F7ED,),PersistentGuid=00000000000000000000000000000000,bHidden=False,bNotConnectable=False,bDefaultValueIsReadOnly=False,bDefaultValueIsIgnored=False,bAdvancedView=False,)
+			   CustomProperties Pin (PinId=F66F6A0343D4D7795DF1D6B0C4320B84,PinName="B",PinToolTip="B\nFloat",PinType.PinCategory="float",PinType.PinSubCategoryObject=None,PinType.PinSubCategoryMemberReference=(),PinType.bIsArray=False,PinType.bIsReference=False,PinType.bIsConst=False,PinType.bIsWeakPointer=False,DefaultValue="1000",AutogeneratedDefaultValue="1.000000",PersistentGuid=00000000000000000000000000000000,bHidden=False,bNotConnectable=False,bDefaultValueIsReadOnly=False,bDefaultValueIsIgnored=False,bAdvancedView=False,)
+			   CustomProperties Pin (PinId=05ABE7DC481D5CE9B683148BBC283658,PinName="ReturnValue",PinToolTip="Return Value\nVector ",Direction="EGPD_Output",PinType.PinCategory="struct",PinType.PinSubCategoryObject=ScriptStruct'/Script/CoreUObject.Vector',PinType.PinSubCategoryMemberReference=(),PinType.bIsArray=False,PinType.bIsReference=False,PinType.bIsConst=False,PinType.bIsWeakPointer=False,AutogeneratedDefaultValue="0, 0, 0",LinkedTo=(K2Node_CallFunction_2595 DEA26338429AED62B1FE2E905F85720A,),PersistentGuid=00000000000000000000000000000000,bHidden=False,bNotConnectable=False,bDefaultValueIsReadOnly=False,bDefaultValueIsIgnored=False,bAdvancedView=False,)
+			End Object
+			Begin Object Class=K2Node_CallFunction Name="K2Node_CallFunction_753"
+			   bIsPureFunc=True
+			   FunctionReference=(MemberParent=Class'/Script/Engine.KismetMathLibrary',MemberName="Subtract_VectorVector")
+			   NodePosX=544
+			   NodePosY=-743
+			   NodeGuid=13124A384A53589660ECD7A5B699B217
+			   CustomProperties Pin (PinId=CE8AE61E4FD0B98ED3F0A0AD52389A6A,PinName="self",PinFriendlyName=NSLOCTEXT("K2Node", "Target", "Target"),PinToolTip="Target\nKismet Math Library Reference",PinType.PinCategory="object",PinType.PinSubCategoryObject=Class'/Script/Engine.KismetMathLibrary',PinType.PinSubCategoryMemberReference=(),PinType.bIsArray=False,PinType.bIsReference=False,PinType.bIsConst=False,PinType.bIsWeakPointer=False,DefaultObject="/Script/Engine.Default__KismetMathLibrary",PersistentGuid=00000000000000000000000000000000,bHidden=True,bNotConnectable=False,bDefaultValueIsReadOnly=False,bDefaultValueIsIgnored=False,bAdvancedView=False,)
+			   CustomProperties Pin (PinId=757CE9264CA706BECDD331B89650EC21,PinName="A",PinToolTip="A\nVector ",PinType.PinCategory="struct",PinType.PinSubCategoryObject=ScriptStruct'/Script/CoreUObject.Vector',PinType.PinSubCategoryMemberReference=(),PinType.bIsArray=False,PinType.bIsReference=False,PinType.bIsConst=False,PinType.bIsWeakPointer=False,DefaultValue="0, 0, 0",AutogeneratedDefaultValue="0, 0, 0",LinkedTo=(K2Node_CallFunction_468 58B1ADC0421583D3D49F459C2E09B61A,),PersistentGuid=00000000000000000000000000000000,bHidden=False,bNotConnectable=False,bDefaultValueIsReadOnly=False,bDefaultValueIsIgnored=False,bAdvancedView=False,)
+			   CustomProperties Pin (PinId=C60FCEFC4EB001C70279E88658AED2D6,PinName="B",PinToolTip="B\nVector ",PinType.PinCategory="struct",PinType.PinSubCategoryObject=ScriptStruct'/Script/CoreUObject.Vector',PinType.PinSubCategoryMemberReference=(),PinType.bIsArray=False,PinType.bIsReference=False,PinType.bIsConst=False,PinType.bIsWeakPointer=False,DefaultValue="0, 0, 0",AutogeneratedDefaultValue="0, 0, 0",LinkedTo=(K2Node_VariableGet_1423 7A20A38A4B4BA9F8557850B65C3E44EC,),PersistentGuid=00000000000000000000000000000000,bHidden=False,bNotConnectable=False,bDefaultValueIsReadOnly=False,bDefaultValueIsIgnored=False,bAdvancedView=False,)
+			   CustomProperties Pin (PinId=638E64934D938A933B279583B5F1FBBE,PinName="ReturnValue",PinToolTip="Return Value\nVector ",Direction="EGPD_Output",PinType.PinCategory="struct",PinType.PinSubCategoryObject=ScriptStruct'/Script/CoreUObject.Vector',PinType.PinSubCategoryMemberReference=(),PinType.bIsArray=False,PinType.bIsReference=False,PinType.bIsConst=False,PinType.bIsWeakPointer=False,AutogeneratedDefaultValue="0, 0, 0",LinkedTo=(K2Node_CallFunction_1151 02460B1B421926AE0062D5BD1FE3D740,),PersistentGuid=00000000000000000000000000000000,bHidden=False,bNotConnectable=False,bDefaultValueIsReadOnly=False,bDefaultValueIsIgnored=False,bAdvancedView=False,)
+			End Object
+			Begin Object Class=K2Node_VariableGet Name="K2Node_VariableGet_1423"
+			   VariableReference=(MemberName="LastTouchingActorPosition",MemberGuid=C849AC4745DA47383E1884AA8AA77BC9,bSelfContext=True)
+			   NodePosX=256
+			   NodePosY=-717
+			   NodeGuid=D79A67B94503ECCB5AD12D95E212055B
+			   CustomProperties Pin (PinId=7A20A38A4B4BA9F8557850B65C3E44EC,PinName="LastTouchingActorPosition",Direction="EGPD_Output",PinType.PinCategory="struct",PinType.PinSubCategoryObject=ScriptStruct'/Script/CoreUObject.Vector',PinType.PinSubCategoryMemberReference=(),PinType.bIsArray=False,PinType.bIsReference=False,PinType.bIsConst=False,PinType.bIsWeakPointer=False,AutogeneratedDefaultValue="0, 0, 0",LinkedTo=(K2Node_CallFunction_753 C60FCEFC4EB001C70279E88658AED2D6,),PersistentGuid=00000000000000000000000000000000,bHidden=False,bNotConnectable=False,bDefaultValueIsReadOnly=False,bDefaultValueIsIgnored=False,bAdvancedView=False,)
+			   CustomProperties Pin (PinId=249EF3D54F840055518B9E8AAC68F979,PinName="self",PinFriendlyName=NSLOCTEXT("K2Node", "Target", "Target"),PinType.PinCategory="object",PinType.PinSubCategoryObject=BlueprintGeneratedClass'/Game/ExampleContent/BlueprintFluidSurface/WaterSurface_BP.WaterSurface_BP_C',PinType.PinSubCategoryMemberReference=(),PinType.bIsArray=False,PinType.bIsReference=False,PinType.bIsConst=False,PinType.bIsWeakPointer=False,PersistentGuid=00000000000000000000000000000000,bHidden=True,bNotConnectable=False,bDefaultValueIsReadOnly=False,bDefaultValueIsIgnored=False,bAdvancedView=False,)
+			End Object
+			Begin Object Class=EdGraphNode_Comment Name="EdGraphNode_Comment_0"
+			   NodePosX=-240
+			   NodePosY=-1408
+			   NodeWidth=2672
+			   NodeHeight=752
+			   NodeComment="Touching character force"
+			   NodeGuid=A8ABD0654FE8CA0C9652109D475B7DC5
+			End Object
+			Begin Object Class=K2Node_CallFunction Name="K2Node_CallFunction_2200"
+			   FunctionReference=(MemberParent=Class'/Script/Engine.MaterialInstanceDynamic',MemberName="SetScalarParameterValue")
+			   NodePosX=1504
+			   NodePosY=-1136
+			   ErrorType=1
+			   ErrorMsg="Error This blueprint (self) is not a MaterialInstanceDynamic, therefore \' Target \' must have a connection."
+			   NodeGuid=D2602A384EA73C9F3634E5ADCB2F202F
+			   CustomProperties Pin (PinId=0722514B4563FDB08CA516BC8FADE4E3,PinName="execute",PinToolTip="\nExec",PinType.PinCategory="exec",PinType.PinSubCategoryObject=None,PinType.PinSubCategoryMemberReference=(),PinType.bIsArray=False,PinType.bIsReference=False,PinType.bIsConst=False,PinType.bIsWeakPointer=False,LinkedTo=(K2Node_CallFunction_1036 BD299B3A459706D661AE5CB3B94FC4E5,),PersistentGuid=00000000000000000000000000000000,bHidden=False,bNotConnectable=False,bDefaultValueIsReadOnly=False,bDefaultValueIsIgnored=False,bAdvancedView=False,)
+			   CustomProperties Pin (PinId=5DCF70DD4DE82CDBFC76FE9B1899F671,PinName="then",PinToolTip="\nExec",Direction="EGPD_Output",PinType.PinCategory="exec",PinType.PinSubCategoryObject=None,PinType.PinSubCategoryMemberReference=(),PinType.bIsArray=False,PinType.bIsReference=False,PinType.bIsConst=False,PinType.bIsWeakPointer=False,LinkedTo=(K2Node_CallFunction_144 F1A930D74213A3622C8A7A8A43919075,),PersistentGuid=00000000000000000000000000000000,bHidden=False,bNotConnectable=False,bDefaultValueIsReadOnly=False,bDefaultValueIsIgnored=False,bAdvancedView=False,)
+			   CustomProperties Pin (PinId=0D70796D4BD1BFABCED6F08008A7F1AE,PinName="self",PinFriendlyName=NSLOCTEXT("K2Node", "Target", "Target"),PinToolTip="Target\nMaterial Instance Dynamic Reference",PinType.PinCategory="object",PinType.PinSubCategoryObject=Class'/Script/Engine.MaterialInstanceDynamic',PinType.PinSubCategoryMemberReference=(),PinType.bIsArray=False,PinType.bIsReference=False,PinType.bIsConst=False,PinType.bIsWeakPointer=False,LinkedTo=(K2Node_CallMaterialParameterCollectionFunction_21 A3A63E4B40398A36A1C1C5A923BDC83E,),PersistentGuid=00000000000000000000000000000000,bHidden=False,bNotConnectable=False,bDefaultValueIsReadOnly=False,bDefaultValueIsIgnored=False,bAdvancedView=False,)
+			   CustomProperties Pin (PinId=3A722E3F41D286785C1C66BAF6D35FE9,PinName="ParameterName",PinToolTip="Parameter Name\nName",PinType.PinCategory="name",PinType.PinSubCategoryObject=None,PinType.PinSubCategoryMemberReference=(),PinType.bIsArray=False,PinType.bIsReference=False,PinType.bIsConst=False,PinType.bIsWeakPointer=False,DefaultValue="ForceStrength",AutogeneratedDefaultValue="None",PersistentGuid=00000000000000000000000000000000,bHidden=False,bNotConnectable=False,bDefaultValueIsReadOnly=False,bDefaultValueIsIgnored=False,bAdvancedView=False,)
+			   CustomProperties Pin (PinId=CA6F95FD46A623E4DF50B99810CCD5D4,PinName="Value",PinToolTip="Value\nFloat",PinType.PinCategory="float",PinType.PinSubCategoryObject=None,PinType.PinSubCategoryMemberReference=(),PinType.bIsArray=False,PinType.bIsReference=False,PinType.bIsConst=False,PinType.bIsWeakPointer=False,DefaultValue="0.0",AutogeneratedDefaultValue="0.0",LinkedTo=(K2Node_CommutativeAssociativeBinaryOperator_38 7DA4112E4DDFB97EF3EAEE98509CC6EB,),PersistentGuid=00000000000000000000000000000000,bHidden=False,bNotConnectable=False,bDefaultValueIsReadOnly=False,bDefaultValueIsIgnored=False,bAdvancedView=False,)
+			End Object
+			Begin Object Class=K2Node_CallFunction Name="K2Node_CallFunction_1151"
+			   bIsPureFunc=True
+			   FunctionReference=(MemberParent=Class'/Script/Engine.KismetMathLibrary',MemberName="VSize")
+			   NodePosX=768
+			   NodePosY=-752
+			   NodeGuid=0B4BA4194C61E26B5E6C8BB6B35D1B3C
+			   CustomProperties Pin (PinId=26B8A18A4D37D401D34CAD860FBA328A,PinName="self",PinFriendlyName=NSLOCTEXT("K2Node", "Target", "Target"),PinToolTip="Target\nKismet Math Library Reference",PinType.PinCategory="object",PinType.PinSubCategoryObject=Class'/Script/Engine.KismetMathLibrary',PinType.PinSubCategoryMemberReference=(),PinType.bIsArray=False,PinType.bIsReference=False,PinType.bIsConst=False,PinType.bIsWeakPointer=False,DefaultObject="/Script/Engine.Default__KismetMathLibrary",PersistentGuid=00000000000000000000000000000000,bHidden=True,bNotConnectable=False,bDefaultValueIsReadOnly=False,bDefaultValueIsIgnored=False,bAdvancedView=False,)
+			   CustomProperties Pin (PinId=02460B1B421926AE0062D5BD1FE3D740,PinName="A",PinToolTip="A\nVector ",PinType.PinCategory="struct",PinType.PinSubCategoryObject=ScriptStruct'/Script/CoreUObject.Vector',PinType.PinSubCategoryMemberReference=(),PinType.bIsArray=False,PinType.bIsReference=False,PinType.bIsConst=False,PinType.bIsWeakPointer=False,DefaultValue="0, 0, 0",AutogeneratedDefaultValue="0, 0, 0",LinkedTo=(K2Node_CallFunction_753 638E64934D938A933B279583B5F1FBBE,),PersistentGuid=00000000000000000000000000000000,bHidden=False,bNotConnectable=False,bDefaultValueIsReadOnly=False,bDefaultValueIsIgnored=False,bAdvancedView=False,)
+			   CustomProperties Pin (PinId=5416EA2343C806C3E7C0918963FC7C27,PinName="ReturnValue",PinToolTip="Return Value\nFloat",Direction="EGPD_Output",PinType.PinCategory="float",PinType.PinSubCategoryObject=None,PinType.PinSubCategoryMemberReference=(),PinType.bIsArray=False,PinType.bIsReference=False,PinType.bIsConst=False,PinType.bIsWeakPointer=False,AutogeneratedDefaultValue="0.0",LinkedTo=(K2Node_CommutativeAssociativeBinaryOperator_38 662ED13A412CBBC7C621F790F181C1B8,),PersistentGuid=00000000000000000000000000000000,bHidden=False,bNotConnectable=False,bDefaultValueIsReadOnly=False,bDefaultValueIsIgnored=False,bAdvancedView=False,)
+			End Object
+			Begin Object Class=K2Node_CommutativeAssociativeBinaryOperator Name="K2Node_CommutativeAssociativeBinaryOperator_38"
+			   bIsPureFunc=True
+			   FunctionReference=(MemberParent=Class'/Script/Engine.KismetMathLibrary',MemberName="Multiply_FloatFloat")
+			   NodePosX=1024
+			   NodePosY=-728
+			   NodeGuid=EA8245864E5831B0FAD39BB98F742158
+			   CustomProperties Pin (PinId=3446CE904F772B6E674D1B9B59378F06,PinName="self",PinFriendlyName=NSLOCTEXT("K2Node", "Target", "Target"),PinToolTip="Target\nKismet Math Library Reference",PinType.PinCategory="object",PinType.PinSubCategoryObject=Class'/Script/Engine.KismetMathLibrary',PinType.PinSubCategoryMemberReference=(),PinType.bIsArray=False,PinType.bIsReference=False,PinType.bIsConst=False,PinType.bIsWeakPointer=False,DefaultObject="/Script/Engine.Default__KismetMathLibrary",PersistentGuid=00000000000000000000000000000000,bHidden=True,bNotConnectable=False,bDefaultValueIsReadOnly=False,bDefaultValueIsIgnored=False,bAdvancedView=False,)
+			   CustomProperties Pin (PinId=662ED13A412CBBC7C621F790F181C1B8,PinName="A",PinToolTip="A\nFloat",PinType.PinCategory="float",PinType.PinSubCategoryObject=None,PinType.PinSubCategoryMemberReference=(),PinType.bIsArray=False,PinType.bIsReference=False,PinType.bIsConst=False,PinType.bIsWeakPointer=False,DefaultValue="0.0",AutogeneratedDefaultValue="0.0",LinkedTo=(K2Node_CallFunction_1151 5416EA2343C806C3E7C0918963FC7C27,),PersistentGuid=00000000000000000000000000000000,bHidden=False,bNotConnectable=False,bDefaultValueIsReadOnly=False,bDefaultValueIsIgnored=False,bAdvancedView=False,)
+			   CustomProperties Pin (PinId=3E874D544B96409DDA6C7C8574202267,PinName="B",PinToolTip="B\nFloat",PinType.PinCategory="float",PinType.PinSubCategoryObject=None,PinType.PinSubCategoryMemberReference=(),PinType.bIsArray=False,PinType.bIsReference=False,PinType.bIsConst=False,PinType.bIsWeakPointer=False,DefaultValue="0.006",AutogeneratedDefaultValue="0.0",PersistentGuid=00000000000000000000000000000000,bHidden=False,bNotConnectable=False,bDefaultValueIsReadOnly=False,bDefaultValueIsIgnored=False,bAdvancedView=False,)
+			   CustomProperties Pin (PinId=7DA4112E4DDFB97EF3EAEE98509CC6EB,PinName="ReturnValue",PinToolTip="Return Value\nFloat",Direction="EGPD_Output",PinType.PinCategory="float",PinType.PinSubCategoryObject=None,PinType.PinSubCategoryMemberReference=(),PinType.bIsArray=False,PinType.bIsReference=False,PinType.bIsConst=False,PinType.bIsWeakPointer=False,AutogeneratedDefaultValue="0.0",LinkedTo=(K2Node_CallFunction_2200 CA6F95FD46A623E4DF50B99810CCD5D4,),PersistentGuid=00000000000000000000000000000000,bHidden=False,bNotConnectable=False,bDefaultValueIsReadOnly=False,bDefaultValueIsIgnored=False,bAdvancedView=False,)
+			End Object
+			Begin Object Class=K2Node_CallFunction Name="K2Node_CallFunction_144"
+			   FunctionReference=(MemberParent=Class'/Script/Engine.KismetRenderingLibrary',MemberName="DrawMaterialToRenderTarget")
+			   NodePosX=1840
+			   NodePosY=-1120
+			   NodeGuid=169B9B0F41707A32ADB76D9764C2BDA4
+			   CustomProperties Pin (PinId=F1A930D74213A3622C8A7A8A43919075,PinName="execute",PinToolTip="\nExec",PinType.PinCategory="exec",PinType.PinSubCategoryObject=None,PinType.PinSubCategoryMemberReference=(),PinType.bIsArray=False,PinType.bIsReference=False,PinType.bIsConst=False,PinType.bIsWeakPointer=False,LinkedTo=(K2Node_CallFunction_2200 5DCF70DD4DE82CDBFC76FE9B1899F671,),PersistentGuid=00000000000000000000000000000000,bHidden=False,bNotConnectable=False,bDefaultValueIsReadOnly=False,bDefaultValueIsIgnored=False,bAdvancedView=False,)
+			   CustomProperties Pin (PinId=C312F52A41D4BEDBB220AFA395DD1F9F,PinName="then",PinToolTip="\nExec",Direction="EGPD_Output",PinType.PinCategory="exec",PinType.PinSubCategoryObject=None,PinType.PinSubCategoryMemberReference=(),PinType.bIsArray=False,PinType.bIsReference=False,PinType.bIsConst=False,PinType.bIsWeakPointer=False,LinkedTo=(K2Node_VariableSet_53 E5EF78B644B6CB8EAD2E229407D67DE2,),PersistentGuid=00000000000000000000000000000000,bHidden=False,bNotConnectable=False,bDefaultValueIsReadOnly=False,bDefaultValueIsIgnored=False,bAdvancedView=False,)
+			   CustomProperties Pin (PinId=095F2BC84E3C55957C90D287FBE23C2C,PinName="self",PinFriendlyName=NSLOCTEXT("K2Node", "Target", "Target"),PinToolTip="Target\nKismet Rendering Library Reference",PinType.PinCategory="object",PinType.PinSubCategoryObject=Class'/Script/Engine.KismetRenderingLibrary',PinType.PinSubCategoryMemberReference=(),PinType.bIsArray=False,PinType.bIsReference=False,PinType.bIsConst=False,PinType.bIsWeakPointer=False,DefaultObject="/Script/Engine.Default__KismetRenderingLibrary",PersistentGuid=00000000000000000000000000000000,bHidden=True,bNotConnectable=False,bDefaultValueIsReadOnly=False,bDefaultValueIsIgnored=False,bAdvancedView=False,)
+			   CustomProperties Pin (PinId=F5D046AE4AFF1E38DD6B8FBBDD4CCD2F,PinName="WorldContextObject",PinToolTip="World Context Object\nObject Reference",PinType.PinCategory="object",PinType.PinSubCategoryObject=Class'/Script/CoreUObject.Object',PinType.PinSubCategoryMemberReference=(),PinType.bIsArray=False,PinType.bIsReference=False,PinType.bIsConst=False,PinType.bIsWeakPointer=False,PersistentGuid=00000000000000000000000000000000,bHidden=True,bNotConnectable=False,bDefaultValueIsReadOnly=False,bDefaultValueIsIgnored=False,bAdvancedView=False,)
+			   CustomProperties Pin (PinId=D2975F3E4A669D9FC42692A4B0EFFDE5,PinName="TextureRenderTarget",PinToolTip="Texture Render Target\nTexture Render Target 2D Reference",PinType.PinCategory="object",PinType.PinSubCategoryObject=Class'/Script/Engine.TextureRenderTarget2D',PinType.PinSubCategoryMemberReference=(),PinType.bIsArray=False,PinType.bIsReference=False,PinType.bIsConst=False,PinType.bIsWeakPointer=False,LinkedTo=(K2Node_CallFunction_21 F349F78749FD3E5ADB14DE8F5E164315,),PersistentGuid=00000000000000000000000000000000,bHidden=False,bNotConnectable=False,bDefaultValueIsReadOnly=False,bDefaultValueIsIgnored=False,bAdvancedView=False,)
+			   CustomProperties Pin (PinId=9E1B05684D5F28E5E0D4E49A7BC18BA8,PinName="Material",PinToolTip="Material\nMaterial Interface Reference",PinType.PinCategory="object",PinType.PinSubCategoryObject=Class'/Script/Engine.MaterialInterface',PinType.PinSubCategoryMemberReference=(),PinType.bIsArray=False,PinType.bIsReference=False,PinType.bIsConst=False,PinType.bIsWeakPointer=False,LinkedTo=(K2Node_CallMaterialParameterCollectionFunction_21 A3A63E4B40398A36A1C1C5A923BDC83E,),PersistentGuid=00000000000000000000000000000000,bHidden=False,bNotConnectable=False,bDefaultValueIsReadOnly=False,bDefaultValueIsIgnored=False,bAdvancedView=False,)
+			End Object
+			Begin Object Class=K2Node_VariableGet Name="K2Node_VariableGet_30"
+			   VariableReference=(MemberName="StaticMeshComponent",bSelfContext=True)
+			   NodePosX=-192
+			   NodePosY=-921
+			   NodeGuid=94D2CBA84FE43135A50871A3BD52E35C
+			   CustomProperties Pin (PinId=415E187944F468E03347EFBF681ADBBF,PinName="StaticMeshComponent",Direction="EGPD_Output",PinType.PinCategory="object",PinType.PinSubCategoryObject=Class'/Script/Engine.StaticMeshComponent',PinType.PinSubCategoryMemberReference=(),PinType.bIsArray=False,PinType.bIsReference=False,PinType.bIsConst=False,PinType.bIsWeakPointer=False,LinkedTo=(K2Node_CallFunction_552 1B9D37834BED70BF9F32EF88300929F5,),PersistentGuid=00000000000000000000000000000000,bHidden=False,bNotConnectable=False,bDefaultValueIsReadOnly=False,bDefaultValueIsIgnored=False,bAdvancedView=False,)
+			   CustomProperties Pin (PinId=175759E54B2B5D77C2264784024C7B9F,PinName="self",PinFriendlyName=NSLOCTEXT("K2Node", "Target", "Target"),PinType.PinCategory="object",PinType.PinSubCategoryObject=Class'/Script/Engine.StaticMeshActor',PinType.PinSubCategoryMemberReference=(),PinType.bIsArray=False,PinType.bIsReference=False,PinType.bIsConst=False,PinType.bIsWeakPointer=False,PersistentGuid=00000000000000000000000000000000,bHidden=True,bNotConnectable=False,bDefaultValueIsReadOnly=False,bDefaultValueIsIgnored=False,bAdvancedView=False,)
+			End Object
+			END OBJECT
+	[/PARAM]
+[/OBJECT]
+[/EXCERPT:TouchingCarForce]
+
+[EXCERPT:FixedFrameRate]	
+[OBJECT:Snippet]
+	[PARAMLITERAL:id]
+		FixedFrameRate
+	[/PARAMLITERAL]
+	[PARAMLITERAL:type]
+		blueprint
+	[/PARAMLITERAL]
+	[PARAM:image]
+		![](Engine\Rendering\RenderTargets\BlueprintRenderTargets\HowTo\FluidSurface\FS_Fixed_Time_Step.png)
+	[/PARAM]
+	[PARAM:snippet]
+			BEGIN OBJECT
+							Begin Object Class=K2Node_MacroInstance Name="K2Node_MacroInstance_0"
+			   MacroGraphReference=(MacroGraph=EdGraph'/Engine/EditorBlueprintResources/StandardMacros.StandardMacros:WhileLoop',GraphBlueprint=Blueprint'/Engine/EditorBlueprintResources/StandardMacros.StandardMacros',GraphGuid=FA93B260444755CD702C21A123E9A987)
+			   NodePosX=352
+			   NodePosY=-195
+			   NodeGuid=88735E2B4AD0C02554CAC0A78ABEFD44
+			   CustomProperties Pin (PinId=1E2B254C4CB4780A2C6517AC5DB1C2FA,PinName="execute",PinType.PinCategory="exec",PinType.PinSubCategoryObject=None,PinType.PinSubCategoryMemberReference=(),PinType.bIsArray=False,PinType.bIsReference=False,PinType.bIsConst=False,PinType.bIsWeakPointer=False,LinkedTo=(K2Node_VariableSet_19 C721D3C14D14A62AFEB73ABEF9C8E538,),PersistentGuid=00000000000000000000000000000000,bHidden=False,bNotConnectable=False,bDefaultValueIsReadOnly=False,bDefaultValueIsIgnored=False,bAdvancedView=False,)
+			   CustomProperties Pin (PinId=BD1907DE4D5A572E52719F8D43F5E6BA,PinName="Condition",PinType.PinCategory="bool",PinType.PinSubCategoryObject=None,PinType.PinSubCategoryMemberReference=(),PinType.bIsArray=False,PinType.bIsReference=False,PinType.bIsConst=False,PinType.bIsWeakPointer=False,DefaultValue="false",LinkedTo=(K2Node_CallFunction_154 0C400D12426215BF5E21C38BE44DC076,),PersistentGuid=00000000000000000000000000000000,bHidden=False,bNotConnectable=False,bDefaultValueIsReadOnly=False,bDefaultValueIsIgnored=False,bAdvancedView=False,)
+			   CustomProperties Pin (PinId=9C4D858646A3B2A7F83F04A2ED9032E8,PinName="LoopBody",Direction="EGPD_Output",PinType.PinCategory="exec",PinType.PinSubCategoryObject=None,PinType.PinSubCategoryMemberReference=(),PinType.bIsArray=False,PinType.bIsReference=False,PinType.bIsConst=False,PinType.bIsWeakPointer=False,LinkedTo=(K2Node_VariableSet_33 637CF0754418216D6E6A309BA7C87180,),PersistentGuid=00000000000000000000000000000000,bHidden=False,bNotConnectable=False,bDefaultValueIsReadOnly=False,bDefaultValueIsIgnored=False,bAdvancedView=False,)
+			   CustomProperties Pin (PinId=3F8D0C52472E83BA32ACAABBA33F8AD7,PinName="Completed",Direction="EGPD_Output",PinType.PinCategory="exec",PinType.PinSubCategoryObject=None,PinType.PinSubCategoryMemberReference=(),PinType.bIsArray=False,PinType.bIsReference=False,PinType.bIsConst=False,PinType.bIsWeakPointer=False,LinkedTo=(K2Node_CallMaterialParameterCollectionFunction_112 7481116A43E6DECBD5AE68A3BAA018B4,),PersistentGuid=00000000000000000000000000000000,bHidden=False,bNotConnectable=False,bDefaultValueIsReadOnly=False,bDefaultValueIsIgnored=False,bAdvancedView=False,)
+			End Object
+			Begin Object Class=K2Node_VariableGet Name="K2Node_VariableGet_0"
+			   VariableReference=(MemberName="TimeAccumulator",MemberGuid=22F269EB45A7199C7F1F88A43649D4BA,bSelfContext=True)
+			   NodePosX=-336
+			   NodePosY=-32
+			   NodeGuid=83EF520F478D2070E2E1FE8875AF200A
+			   CustomProperties Pin (PinId=28509CC34C4A97AAC0C14D9C0CA1E545,PinName="TimeAccumulator",Direction="EGPD_Output",PinType.PinCategory="float",PinType.PinSubCategoryObject=None,PinType.PinSubCategoryMemberReference=(),PinType.bIsArray=False,PinType.bIsReference=False,PinType.bIsConst=False,PinType.bIsWeakPointer=False,AutogeneratedDefaultValue="0.0",LinkedTo=(K2Node_CommutativeAssociativeBinaryOperator_6 45D2B9EC46637488EB5F89A65EB93942,),PersistentGuid=00000000000000000000000000000000,bHidden=False,bNotConnectable=False,bDefaultValueIsReadOnly=False,bDefaultValueIsIgnored=False,bAdvancedView=False,)
+			   CustomProperties Pin (PinId=A0CC64D74ED3F8D0D8DDFA9D54B3589F,PinName="self",PinFriendlyName=NSLOCTEXT("K2Node", "Target", "Target"),PinType.PinCategory="object",PinType.PinSubCategoryObject=BlueprintGeneratedClass'/Game/ExampleContent/BlueprintFluidSurface/WaterSurface_BP.WaterSurface_BP_C',PinType.PinSubCategoryMemberReference=(),PinType.bIsArray=False,PinType.bIsReference=False,PinType.bIsConst=False,PinType.bIsWeakPointer=False,PersistentGuid=00000000000000000000000000000000,bHidden=True,bNotConnectable=False,bDefaultValueIsReadOnly=False,bDefaultValueIsIgnored=False,bAdvancedView=False,)
+			End Object
+			Begin Object Class=K2Node_CommutativeAssociativeBinaryOperator Name="K2Node_CommutativeAssociativeBinaryOperator_6"
+			   bIsPureFunc=True
+			   FunctionReference=(MemberParent=Class'/Script/Engine.KismetMathLibrary',MemberName="Add_FloatFloat")
+			   NodePosX=-144
+			   NodePosY=-144
+			   NodeGuid=68822A2844A7AB05263B97BBABF808B2
+			   CustomProperties Pin (PinId=FE87E5D646A2D7F0B70A9491114DB5B8,PinName="self",PinFriendlyName=NSLOCTEXT("K2Node", "Target", "Target"),PinToolTip="Target\nKismet Math Library Reference",PinType.PinCategory="object",PinType.PinSubCategoryObject=Class'/Script/Engine.KismetMathLibrary',PinType.PinSubCategoryMemberReference=(),PinType.bIsArray=False,PinType.bIsReference=False,PinType.bIsConst=False,PinType.bIsWeakPointer=False,DefaultObject="/Script/Engine.Default__KismetMathLibrary",PersistentGuid=00000000000000000000000000000000,bHidden=True,bNotConnectable=False,bDefaultValueIsReadOnly=False,bDefaultValueIsIgnored=False,bAdvancedView=False,)
+			   CustomProperties Pin (PinId=29EFA03442ADC8249E9CC9B27F215866,PinName="A",PinToolTip="A\nFloat",PinType.PinCategory="float",PinType.PinSubCategoryObject=None,PinType.PinSubCategoryMemberReference=(),PinType.bIsArray=False,PinType.bIsReference=False,PinType.bIsConst=False,PinType.bIsWeakPointer=False,DefaultValue="0.0",AutogeneratedDefaultValue="0.0",LinkedTo=(K2Node_Event_1 90C5BF9B4C7F93AF4E34319E2020E8D0,),PersistentGuid=00000000000000000000000000000000,bHidden=False,bNotConnectable=False,bDefaultValueIsReadOnly=False,bDefaultValueIsIgnored=False,bAdvancedView=False,)
+			   CustomProperties Pin (PinId=45D2B9EC46637488EB5F89A65EB93942,PinName="B",PinToolTip="B\nFloat",PinType.PinCategory="float",PinType.PinSubCategoryObject=None,PinType.PinSubCategoryMemberReference=(),PinType.bIsArray=False,PinType.bIsReference=False,PinType.bIsConst=False,PinType.bIsWeakPointer=False,DefaultValue="0.0",AutogeneratedDefaultValue="1.000000",LinkedTo=(K2Node_VariableGet_0 28509CC34C4A97AAC0C14D9C0CA1E545,),PersistentGuid=00000000000000000000000000000000,bHidden=False,bNotConnectable=False,bDefaultValueIsReadOnly=False,bDefaultValueIsIgnored=False,bAdvancedView=False,)
+			   CustomProperties Pin (PinId=6EDD3417402E56E6EBD02DBE59A96B06,PinName="ReturnValue",PinToolTip="Return Value\nFloat",Direction="EGPD_Output",PinType.PinCategory="float",PinType.PinSubCategoryObject=None,PinType.PinSubCategoryMemberReference=(),PinType.bIsArray=False,PinType.bIsReference=False,PinType.bIsConst=False,PinType.bIsWeakPointer=False,AutogeneratedDefaultValue="0.0",LinkedTo=(K2Node_VariableSet_19 00D9BE81439571117EDDCDB52FA98CC0,),PersistentGuid=00000000000000000000000000000000,bHidden=False,bNotConnectable=False,bDefaultValueIsReadOnly=False,bDefaultValueIsIgnored=False,bAdvancedView=False,)
+			End Object
+			Begin Object Class=K2Node_VariableSet Name="K2Node_VariableSet_19"
+			   VariableReference=(MemberName="TimeAccumulator",MemberGuid=22F269EB45A7199C7F1F88A43649D4BA,bSelfContext=True)
+			   NodePosX=64
+			   NodePosY=-179
+			   NodeGuid=5D7E5EE9496EDBC58FDE289031AB00BB
+			   CustomProperties Pin (PinId=879DBE52479842492E01529E9B944C4A,PinName="execute",PinType.PinCategory="exec",PinType.PinSubCategoryObject=None,PinType.PinSubCategoryMemberReference=(),PinType.bIsArray=False,PinType.bIsReference=False,PinType.bIsConst=False,PinType.bIsWeakPointer=False,LinkedTo=(K2Node_IfThenElse_3 F46E8B0E462777612D45B5BA6EA979CD,K2Node_VariableSet_53 140B85434191688AB8F2A1A2C156091A,),PersistentGuid=00000000000000000000000000000000,bHidden=False,bNotConnectable=False,bDefaultValueIsReadOnly=False,bDefaultValueIsIgnored=False,bAdvancedView=False,)
+			   CustomProperties Pin (PinId=C721D3C14D14A62AFEB73ABEF9C8E538,PinName="then",Direction="EGPD_Output",PinType.PinCategory="exec",PinType.PinSubCategoryObject=None,PinType.PinSubCategoryMemberReference=(),PinType.bIsArray=False,PinType.bIsReference=False,PinType.bIsConst=False,PinType.bIsWeakPointer=False,LinkedTo=(K2Node_MacroInstance_0 1E2B254C4CB4780A2C6517AC5DB1C2FA,),PersistentGuid=00000000000000000000000000000000,bHidden=False,bNotConnectable=False,bDefaultValueIsReadOnly=False,bDefaultValueIsIgnored=False,bAdvancedView=False,)
+			   CustomProperties Pin (PinId=00D9BE81439571117EDDCDB52FA98CC0,PinName="TimeAccumulator",PinType.PinCategory="float",PinType.PinSubCategoryObject=None,PinType.PinSubCategoryMemberReference=(),PinType.bIsArray=False,PinType.bIsReference=False,PinType.bIsConst=False,PinType.bIsWeakPointer=False,DefaultValue="0.0",AutogeneratedDefaultValue="0.0",LinkedTo=(K2Node_CommutativeAssociativeBinaryOperator_6 6EDD3417402E56E6EBD02DBE59A96B06,),PersistentGuid=00000000000000000000000000000000,bHidden=False,bNotConnectable=False,bDefaultValueIsReadOnly=False,bDefaultValueIsIgnored=False,bAdvancedView=False,)
+			   CustomProperties Pin (PinId=398F67D0486B98D01B2D5DA8E82AA32D,PinName="Output_Get",PinToolTip="Retrieves the value of the variable, can use instead of a separate Get node",Direction="EGPD_Output",PinType.PinCategory="float",PinType.PinSubCategoryObject=None,PinType.PinSubCategoryMemberReference=(),PinType.bIsArray=False,PinType.bIsReference=False,PinType.bIsConst=False,PinType.bIsWeakPointer=False,DefaultValue="0.0",AutogeneratedDefaultValue="0.0",PersistentGuid=00000000000000000000000000000000,bHidden=False,bNotConnectable=False,bDefaultValueIsReadOnly=False,bDefaultValueIsIgnored=False,bAdvancedView=False,)
+			   CustomProperties Pin (PinId=A0C749764619C345AEF53AA755C9E837,PinName="self",PinFriendlyName=NSLOCTEXT("K2Node", "Target", "Target"),PinType.PinCategory="object",PinType.PinSubCategoryObject=BlueprintGeneratedClass'/Game/ExampleContent/BlueprintFluidSurface/WaterSurface_BP.WaterSurface_BP_C',PinType.PinSubCategoryMemberReference=(),PinType.bIsArray=False,PinType.bIsReference=False,PinType.bIsConst=False,PinType.bIsWeakPointer=False,PersistentGuid=00000000000000000000000000000000,bHidden=True,bNotConnectable=False,bDefaultValueIsReadOnly=False,bDefaultValueIsIgnored=False,bAdvancedView=False,)
+			End Object
+			Begin Object Class=K2Node_VariableGet Name="K2Node_VariableGet_26"
+			   VariableReference=(MemberName="TimeAccumulator",MemberGuid=22F269EB45A7199C7F1F88A43649D4BA,bSelfContext=True)
+			   NodePosX=16
+			   NodePosY=48
+			   NodeGuid=AA885ED7479C475DC5B398ABB59518C7
+			   CustomProperties Pin (PinId=4A75E9044154F4DCBB16DAA7FBA47FC6,PinName="TimeAccumulator",Direction="EGPD_Output",PinType.PinCategory="float",PinType.PinSubCategoryObject=None,PinType.PinSubCategoryMemberReference=(),PinType.bIsArray=False,PinType.bIsReference=False,PinType.bIsConst=False,PinType.bIsWeakPointer=False,AutogeneratedDefaultValue="0.0",LinkedTo=(K2Node_CallFunction_154 5DB4B65147AF76A4FF2DDF9F56F610B9,),PersistentGuid=00000000000000000000000000000000,bHidden=False,bNotConnectable=False,bDefaultValueIsReadOnly=False,bDefaultValueIsIgnored=False,bAdvancedView=False,)
+			   CustomProperties Pin (PinId=BC20D83044D2E4FCFA3541900F635B14,PinName="self",PinFriendlyName=NSLOCTEXT("K2Node", "Target", "Target"),PinType.PinCategory="object",PinType.PinSubCategoryObject=BlueprintGeneratedClass'/Game/ExampleContent/BlueprintFluidSurface/WaterSurface_BP.WaterSurface_BP_C',PinType.PinSubCategoryMemberReference=(),PinType.bIsArray=False,PinType.bIsReference=False,PinType.bIsConst=False,PinType.bIsWeakPointer=False,PersistentGuid=00000000000000000000000000000000,bHidden=True,bNotConnectable=False,bDefaultValueIsReadOnly=False,bDefaultValueIsIgnored=False,bAdvancedView=False,)
+			End Object
+			Begin Object Class=K2Node_CallFunction Name="K2Node_CallFunction_154"
+			   bIsPureFunc=True
+			   FunctionReference=(MemberParent=Class'/Script/Engine.KismetMathLibrary',MemberName="GreaterEqual_FloatFloat")
+			   NodePosX=192
+			   NodePosY=46
+			   NodeGuid=FC83298141EE57B1A33762BA693B9B13
+			   CustomProperties Pin (PinId=981F348D4CF817F7332BA8B45EFB2E85,PinName="self",PinFriendlyName=NSLOCTEXT("K2Node", "Target", "Target"),PinType.PinCategory="object",PinType.PinSubCategoryObject=Class'/Script/Engine.KismetMathLibrary',PinType.PinSubCategoryMemberReference=(),PinType.bIsArray=False,PinType.bIsReference=False,PinType.bIsConst=False,PinType.bIsWeakPointer=False,DefaultObject="/Script/Engine.Default__KismetMathLibrary",PersistentGuid=00000000000000000000000000000000,bHidden=True,bNotConnectable=False,bDefaultValueIsReadOnly=False,bDefaultValueIsIgnored=False,bAdvancedView=False,)
+			   CustomProperties Pin (PinId=5DB4B65147AF76A4FF2DDF9F56F610B9,PinName="A",PinType.PinCategory="float",PinType.PinSubCategoryObject=None,PinType.PinSubCategoryMemberReference=(),PinType.bIsArray=False,PinType.bIsReference=False,PinType.bIsConst=False,PinType.bIsWeakPointer=False,DefaultValue="0.0",AutogeneratedDefaultValue="0.0",LinkedTo=(K2Node_VariableGet_26 4A75E9044154F4DCBB16DAA7FBA47FC6,),PersistentGuid=00000000000000000000000000000000,bHidden=False,bNotConnectable=False,bDefaultValueIsReadOnly=False,bDefaultValueIsIgnored=False,bAdvancedView=False,)
+			   CustomProperties Pin (PinId=E33421FE40CB8A0A1368ADB1B47B0452,PinName="B",PinType.PinCategory="float",PinType.PinSubCategoryObject=None,PinType.PinSubCategoryMemberReference=(),PinType.bIsArray=False,PinType.bIsReference=False,PinType.bIsConst=False,PinType.bIsWeakPointer=False,DefaultValue="0.0",AutogeneratedDefaultValue="0.0",LinkedTo=(K2Node_CallFunction_293 A74B6C834197914FB6926AB85388A126,),PersistentGuid=00000000000000000000000000000000,bHidden=False,bNotConnectable=False,bDefaultValueIsReadOnly=False,bDefaultValueIsIgnored=False,bAdvancedView=False,)
+			   CustomProperties Pin (PinId=0C400D12426215BF5E21C38BE44DC076,PinName="ReturnValue",Direction="EGPD_Output",PinType.PinCategory="bool",PinType.PinSubCategoryObject=None,PinType.PinSubCategoryMemberReference=(),PinType.bIsArray=False,PinType.bIsReference=False,PinType.bIsConst=False,PinType.bIsWeakPointer=False,AutogeneratedDefaultValue="false",LinkedTo=(K2Node_MacroInstance_0 BD1907DE4D5A572E52719F8D43F5E6BA,),PersistentGuid=00000000000000000000000000000000,bHidden=False,bNotConnectable=False,bDefaultValueIsReadOnly=False,bDefaultValueIsIgnored=False,bAdvancedView=False,)
+			End Object
+			Begin Object Class=K2Node_VariableGet Name="K2Node_VariableGet_40"
+			   VariableReference=(MemberName="UpdateRate",MemberGuid=F0A8DF984EC251A8D15B1D90939E6891,bSelfContext=True)
+			   NodePosX=-352
+			   NodePosY=144
+			   NodeGuid=BC84A4A142A3F2F513DEBA92A27C18DA
+			   CustomProperties Pin (PinId=C90C323C416365E3F2070D8B127B4E5A,PinName="UpdateRate",Direction="EGPD_Output",PinType.PinCategory="float",PinType.PinSubCategoryObject=None,PinType.PinSubCategoryMemberReference=(),PinType.bIsArray=False,PinType.bIsReference=False,PinType.bIsConst=False,PinType.bIsWeakPointer=False,AutogeneratedDefaultValue="0.0",LinkedTo=(K2Node_CommutativeAssociativeBinaryOperator_34 B59E67724D090A3249D543A3DB75B9E5,),PersistentGuid=00000000000000000000000000000000,bHidden=False,bNotConnectable=False,bDefaultValueIsReadOnly=False,bDefaultValueIsIgnored=False,bAdvancedView=False,)
+			   CustomProperties Pin (PinId=4AD4475245E8B1120686D2A003F26818,PinName="self",PinFriendlyName=NSLOCTEXT("K2Node", "Target", "Target"),PinType.PinCategory="object",PinType.PinSubCategoryObject=BlueprintGeneratedClass'/Game/ExampleContent/BlueprintFluidSurface/WaterSurface_BP.WaterSurface_BP_C',PinType.PinSubCategoryMemberReference=(),PinType.bIsArray=False,PinType.bIsReference=False,PinType.bIsConst=False,PinType.bIsWeakPointer=False,PersistentGuid=00000000000000000000000000000000,bHidden=True,bNotConnectable=False,bDefaultValueIsReadOnly=False,bDefaultValueIsIgnored=False,bAdvancedView=False,)
+			End Object
+			Begin Object Class=K2Node_CallFunction Name="K2Node_CallFunction_293"
+			   bIsPureFunc=True
+			   FunctionReference=(MemberParent=Class'/Script/Engine.KismetMathLibrary',MemberName="Divide_FloatFloat")
+			   NodePosX=16
+			   NodePosY=115
+			   NodeGuid=F60556964E9941C0F9AC2599370E427C
+			   CustomProperties Pin (PinId=1BC2041F4BBCD86AD66C869598010E91,PinName="self",PinFriendlyName=NSLOCTEXT("K2Node", "Target", "Target"),PinType.PinCategory="object",PinType.PinSubCategoryObject=Class'/Script/Engine.KismetMathLibrary',PinType.PinSubCategoryMemberReference=(),PinType.bIsArray=False,PinType.bIsReference=False,PinType.bIsConst=False,PinType.bIsWeakPointer=False,DefaultObject="/Script/Engine.Default__KismetMathLibrary",PersistentGuid=00000000000000000000000000000000,bHidden=True,bNotConnectable=False,bDefaultValueIsReadOnly=False,bDefaultValueIsIgnored=False,bAdvancedView=False,)
+			   CustomProperties Pin (PinId=226E921B4E91BB13C2AF3988DF3C5E25,PinName="A",PinType.PinCategory="float",PinType.PinSubCategoryObject=None,PinType.PinSubCategoryMemberReference=(),PinType.bIsArray=False,PinType.bIsReference=False,PinType.bIsConst=False,PinType.bIsWeakPointer=False,DefaultValue="1",AutogeneratedDefaultValue="0.0",PersistentGuid=00000000000000000000000000000000,bHidden=False,bNotConnectable=False,bDefaultValueIsReadOnly=False,bDefaultValueIsIgnored=False,bAdvancedView=False,)
+			   CustomProperties Pin (PinId=066CC8F84BD984A41327A6BC78E45C56,PinName="B",PinType.PinCategory="float",PinType.PinSubCategoryObject=None,PinType.PinSubCategoryMemberReference=(),PinType.bIsArray=False,PinType.bIsReference=False,PinType.bIsConst=False,PinType.bIsWeakPointer=False,DefaultValue="0.0",AutogeneratedDefaultValue="1.000000",LinkedTo=(K2Node_CommutativeAssociativeBinaryOperator_34 C23F5B5743FEC4F10B3DB184AAB1E956,),PersistentGuid=00000000000000000000000000000000,bHidden=False,bNotConnectable=False,bDefaultValueIsReadOnly=False,bDefaultValueIsIgnored=False,bAdvancedView=False,)
+			   CustomProperties Pin (PinId=A74B6C834197914FB6926AB85388A126,PinName="ReturnValue",Direction="EGPD_Output",PinType.PinCategory="float",PinType.PinSubCategoryObject=None,PinType.PinSubCategoryMemberReference=(),PinType.bIsArray=False,PinType.bIsReference=False,PinType.bIsConst=False,PinType.bIsWeakPointer=False,AutogeneratedDefaultValue="0.0",LinkedTo=(K2Node_CallFunction_154 E33421FE40CB8A0A1368ADB1B47B0452,K2Node_CallFunction_294 DAB307B4407D3E82A090208FC0194C92,),PersistentGuid=00000000000000000000000000000000,bHidden=False,bNotConnectable=False,bDefaultValueIsReadOnly=False,bDefaultValueIsIgnored=False,bAdvancedView=False,)
+			End Object
+			Begin Object Class=K2Node_VariableGet Name="K2Node_VariableGet_41"
+			   VariableReference=(MemberName="TimeAccumulator",MemberGuid=22F269EB45A7199C7F1F88A43649D4BA,bSelfContext=True)
+			   NodePosX=304
+			   NodePosY=80
+			   NodeGuid=77EF15094F2412DB606B0080F01A7FCA
+			   CustomProperties Pin (PinId=DB15DED9442AA8D4395AB5B95C2211EC,PinName="TimeAccumulator",Direction="EGPD_Output",PinType.PinCategory="float",PinType.PinSubCategoryObject=None,PinType.PinSubCategoryMemberReference=(),PinType.bIsArray=False,PinType.bIsReference=False,PinType.bIsConst=False,PinType.bIsWeakPointer=False,AutogeneratedDefaultValue="0.0",LinkedTo=(K2Node_CallFunction_294 460457E947D9E825D747E1B19FAACC93,),PersistentGuid=00000000000000000000000000000000,bHidden=False,bNotConnectable=False,bDefaultValueIsReadOnly=False,bDefaultValueIsIgnored=False,bAdvancedView=False,)
+			   CustomProperties Pin (PinId=8AF546324DB31BA34646148D703E5E77,PinName="self",PinFriendlyName=NSLOCTEXT("K2Node", "Target", "Target"),PinType.PinCategory="object",PinType.PinSubCategoryObject=BlueprintGeneratedClass'/Game/ExampleContent/BlueprintFluidSurface/WaterSurface_BP.WaterSurface_BP_C',PinType.PinSubCategoryMemberReference=(),PinType.bIsArray=False,PinType.bIsReference=False,PinType.bIsConst=False,PinType.bIsWeakPointer=False,PersistentGuid=00000000000000000000000000000000,bHidden=True,bNotConnectable=False,bDefaultValueIsReadOnly=False,bDefaultValueIsIgnored=False,bAdvancedView=False,)
+			End Object
+			Begin Object Class=K2Node_CallFunction Name="K2Node_CallFunction_294"
+			   bIsPureFunc=True
+			   FunctionReference=(MemberParent=Class'/Script/Engine.KismetMathLibrary',MemberName="Subtract_FloatFloat")
+			   NodePosX=512
+			   NodePosY=104
+			   NodeGuid=DE5E942B4E69983EEA4D4583D2EA1205
+			   CustomProperties Pin (PinId=3A6D9E354798FD7C7BD2C89C59A2BAC6,PinName="self",PinFriendlyName=NSLOCTEXT("K2Node", "Target", "Target"),PinType.PinCategory="object",PinType.PinSubCategoryObject=Class'/Script/Engine.KismetMathLibrary',PinType.PinSubCategoryMemberReference=(),PinType.bIsArray=False,PinType.bIsReference=False,PinType.bIsConst=False,PinType.bIsWeakPointer=False,DefaultObject="/Script/Engine.Default__KismetMathLibrary",PersistentGuid=00000000000000000000000000000000,bHidden=True,bNotConnectable=False,bDefaultValueIsReadOnly=False,bDefaultValueIsIgnored=False,bAdvancedView=False,)
+			   CustomProperties Pin (PinId=460457E947D9E825D747E1B19FAACC93,PinName="A",PinType.PinCategory="float",PinType.PinSubCategoryObject=None,PinType.PinSubCategoryMemberReference=(),PinType.bIsArray=False,PinType.bIsReference=False,PinType.bIsConst=False,PinType.bIsWeakPointer=False,DefaultValue="0.0",AutogeneratedDefaultValue="0.0",LinkedTo=(K2Node_VariableGet_41 DB15DED9442AA8D4395AB5B95C2211EC,),PersistentGuid=00000000000000000000000000000000,bHidden=False,bNotConnectable=False,bDefaultValueIsReadOnly=False,bDefaultValueIsIgnored=False,bAdvancedView=False,)
+			   CustomProperties Pin (PinId=DAB307B4407D3E82A090208FC0194C92,PinName="B",PinType.PinCategory="float",PinType.PinSubCategoryObject=None,PinType.PinSubCategoryMemberReference=(),PinType.bIsArray=False,PinType.bIsReference=False,PinType.bIsConst=False,PinType.bIsWeakPointer=False,DefaultValue="0.0",AutogeneratedDefaultValue="1.000000",LinkedTo=(K2Node_CallFunction_293 A74B6C834197914FB6926AB85388A126,),PersistentGuid=00000000000000000000000000000000,bHidden=False,bNotConnectable=False,bDefaultValueIsReadOnly=False,bDefaultValueIsIgnored=False,bAdvancedView=False,)
+			   CustomProperties Pin (PinId=2B8CF4094C6D57C2FF4C91A791100C37,PinName="ReturnValue",Direction="EGPD_Output",PinType.PinCategory="float",PinType.PinSubCategoryObject=None,PinType.PinSubCategoryMemberReference=(),PinType.bIsArray=False,PinType.bIsReference=False,PinType.bIsConst=False,PinType.bIsWeakPointer=False,AutogeneratedDefaultValue="0.0",LinkedTo=(K2Node_VariableSet_33 9C95A14C4ABDEA88BCE7FABDD01914B8,),PersistentGuid=00000000000000000000000000000000,bHidden=False,bNotConnectable=False,bDefaultValueIsReadOnly=False,bDefaultValueIsIgnored=False,bAdvancedView=False,)
+			End Object
+			Begin Object Class=K2Node_VariableSet Name="K2Node_VariableSet_33"
+			   VariableReference=(MemberName="TimeAccumulator",MemberGuid=22F269EB45A7199C7F1F88A43649D4BA,bSelfContext=True)
+			   NodePosX=672
+			   NodePosY=-179
+			   NodeGuid=D978573346FCE963005F3396590650E1
+			   CustomProperties Pin (PinId=637CF0754418216D6E6A309BA7C87180,PinName="execute",PinType.PinCategory="exec",PinType.PinSubCategoryObject=None,PinType.PinSubCategoryMemberReference=(),PinType.bIsArray=False,PinType.bIsReference=False,PinType.bIsConst=False,PinType.bIsWeakPointer=False,LinkedTo=(K2Node_MacroInstance_0 9C4D858646A3B2A7F83F04A2ED9032E8,),PersistentGuid=00000000000000000000000000000000,bHidden=False,bNotConnectable=False,bDefaultValueIsReadOnly=False,bDefaultValueIsIgnored=False,bAdvancedView=False,)
+			   CustomProperties Pin (PinId=6D73F5F142F9D57CDE67AF95BA4B5CE0,PinName="then",Direction="EGPD_Output",PinType.PinCategory="exec",PinType.PinSubCategoryObject=None,PinType.PinSubCategoryMemberReference=(),PinType.bIsArray=False,PinType.bIsReference=False,PinType.bIsConst=False,PinType.bIsWeakPointer=False,LinkedTo=(K2Node_VariableSet_14 67D26A9B44D03912C7A167B0C90C861D,),PersistentGuid=00000000000000000000000000000000,bHidden=False,bNotConnectable=False,bDefaultValueIsReadOnly=False,bDefaultValueIsIgnored=False,bAdvancedView=False,)
+			   CustomProperties Pin (PinId=9C95A14C4ABDEA88BCE7FABDD01914B8,PinName="TimeAccumulator",PinType.PinCategory="float",PinType.PinSubCategoryObject=None,PinType.PinSubCategoryMemberReference=(),PinType.bIsArray=False,PinType.bIsReference=False,PinType.bIsConst=False,PinType.bIsWeakPointer=False,DefaultValue="0.0",AutogeneratedDefaultValue="0.0",LinkedTo=(K2Node_CallFunction_294 2B8CF4094C6D57C2FF4C91A791100C37,),PersistentGuid=00000000000000000000000000000000,bHidden=False,bNotConnectable=False,bDefaultValueIsReadOnly=False,bDefaultValueIsIgnored=False,bAdvancedView=False,)
+			   CustomProperties Pin (PinId=30AB000446E0B8F39FDC2CB1456D977B,PinName="Output_Get",PinToolTip="Retrieves the value of the variable, can use instead of a separate Get node",Direction="EGPD_Output",PinType.PinCategory="float",PinType.PinSubCategoryObject=None,PinType.PinSubCategoryMemberReference=(),PinType.bIsArray=False,PinType.bIsReference=False,PinType.bIsConst=False,PinType.bIsWeakPointer=False,DefaultValue="0.0",AutogeneratedDefaultValue="0.0",PersistentGuid=00000000000000000000000000000000,bHidden=False,bNotConnectable=False,bDefaultValueIsReadOnly=False,bDefaultValueIsIgnored=False,bAdvancedView=False,)
+			   CustomProperties Pin (PinId=BFAEF4ED423CCB16DAC3F78DE4639819,PinName="self",PinFriendlyName=NSLOCTEXT("K2Node", "Target", "Target"),PinType.PinCategory="object",PinType.PinSubCategoryObject=BlueprintGeneratedClass'/Game/ExampleContent/BlueprintFluidSurface/WaterSurface_BP.WaterSurface_BP_C',PinType.PinSubCategoryMemberReference=(),PinType.bIsArray=False,PinType.bIsReference=False,PinType.bIsConst=False,PinType.bIsWeakPointer=False,PersistentGuid=00000000000000000000000000000000,bHidden=True,bNotConnectable=False,bDefaultValueIsReadOnly=False,bDefaultValueIsIgnored=False,bAdvancedView=False,)
+			End Object
+			Begin Object Class=EdGraphNode_Comment Name="EdGraphNode_Comment_2"
+			   NodePosX=-368
+			   NodePosY=-240
+			   NodeWidth=1280
+			   NodeHeight=512
+			   NodeComment="Fixed timestep, so fluid speed is framerate independent"
+			   NodeGuid=67D771444DBB8730324897B77658CBC6
+			End Object
+			Begin Object Class=K2Node_CommutativeAssociativeBinaryOperator Name="K2Node_CommutativeAssociativeBinaryOperator_34"
+			   bIsPureFunc=True
+			   FunctionReference=(MemberParent=Class'/Script/Engine.KismetMathLibrary',MemberName="FMin")
+			   NodePosX=-176
+			   NodePosY=144
+			   NodeGuid=A143E63D4E9344CBEAB0B6B02A9A8F98
+			   CustomProperties Pin (PinId=10A2646C46F59CFC930AA19476D7EFE4,PinName="self",PinFriendlyName=NSLOCTEXT("K2Node", "Target", "Target"),PinType.PinCategory="object",PinType.PinSubCategoryObject=Class'/Script/Engine.KismetMathLibrary',PinType.PinSubCategoryMemberReference=(),PinType.bIsArray=False,PinType.bIsReference=False,PinType.bIsConst=False,PinType.bIsWeakPointer=False,DefaultObject="/Script/Engine.Default__KismetMathLibrary",PersistentGuid=00000000000000000000000000000000,bHidden=True,bNotConnectable=False,bDefaultValueIsReadOnly=False,bDefaultValueIsIgnored=False,bAdvancedView=False,)
+			   CustomProperties Pin (PinId=B59E67724D090A3249D543A3DB75B9E5,PinName="A",PinType.PinCategory="float",PinType.PinSubCategoryObject=None,PinType.PinSubCategoryMemberReference=(),PinType.bIsArray=False,PinType.bIsReference=False,PinType.bIsConst=False,PinType.bIsWeakPointer=False,DefaultValue="0.0",AutogeneratedDefaultValue="0.0",LinkedTo=(K2Node_VariableGet_40 C90C323C416365E3F2070D8B127B4E5A,),PersistentGuid=00000000000000000000000000000000,bHidden=False,bNotConnectable=False,bDefaultValueIsReadOnly=False,bDefaultValueIsIgnored=False,bAdvancedView=False,)
+			   CustomProperties Pin (PinId=B8A5FA754444744D152FD798754A4F88,PinName="B",PinType.PinCategory="float",PinType.PinSubCategoryObject=None,PinType.PinSubCategoryMemberReference=(),PinType.bIsArray=False,PinType.bIsReference=False,PinType.bIsConst=False,PinType.bIsWeakPointer=False,DefaultValue="120",AutogeneratedDefaultValue="0.0",PersistentGuid=00000000000000000000000000000000,bHidden=False,bNotConnectable=False,bDefaultValueIsReadOnly=False,bDefaultValueIsIgnored=False,bAdvancedView=False,)
+			   CustomProperties Pin (PinId=C23F5B5743FEC4F10B3DB184AAB1E956,PinName="ReturnValue",Direction="EGPD_Output",PinType.PinCategory="float",PinType.PinSubCategoryObject=None,PinType.PinSubCategoryMemberReference=(),PinType.bIsArray=False,PinType.bIsReference=False,PinType.bIsConst=False,PinType.bIsWeakPointer=False,AutogeneratedDefaultValue="0.0",LinkedTo=(K2Node_CallFunction_293 066CC8F84BD984A41327A6BC78E45C56,),PersistentGuid=00000000000000000000000000000000,bHidden=False,bNotConnectable=False,bDefaultValueIsReadOnly=False,bDefaultValueIsIgnored=False,bAdvancedView=False,)
+			End Object
+			END OBJECT
+	[/PARAM]
+[/OBJECT]
+[/EXCERPT:FixedFrameRate]
+
+[EXCERPT:CycleActiveHeightfield]	
+[OBJECT:Snippet]
+	[PARAMLITERAL:id]
+		CycleActiveHeightfield
+	[/PARAMLITERAL]
+	[PARAMLITERAL:type]
+		blueprint
+	[/PARAMLITERAL]
+	[PARAM:image]
+		![](Engine\Rendering\RenderTargets\BlueprintRenderTargets\HowTo\FluidSurface\FS_Cycle_Active_Heightfield.png)
+	[/PARAM]
+	[PARAM:snippet]
+			BEGIN OBJECT
+							Begin Object Class=K2Node_VariableGet Name="K2Node_VariableGet_19"
+			   VariableReference=(MemberName="HeightState",MemberGuid=DFF4A74440AF3247B05F028FA72FDFF9,bSelfContext=True)
+			   NodePosX=1059
+			   NodePosY=160
+			   NodeGuid=AAA9032549E53180E150E9A55DB283B5
+			   CustomProperties Pin (PinId=EEB54EFA427F3D4B4674DC8D6620862B,PinName="HeightState",Direction="EGPD_Output",PinType.PinCategory="int",PinType.PinSubCategoryObject=None,PinType.PinSubCategoryMemberReference=(),PinType.bIsArray=False,PinType.bIsReference=False,PinType.bIsConst=False,PinType.bIsWeakPointer=False,AutogeneratedDefaultValue="0",LinkedTo=(K2Node_CommutativeAssociativeBinaryOperator_1 155D15B04B9603A253D24DAA458A93F7,),PersistentGuid=00000000000000000000000000000000,bHidden=False,bNotConnectable=False,bDefaultValueIsReadOnly=False,bDefaultValueIsIgnored=False,bAdvancedView=False,)
+			   CustomProperties Pin (PinId=0966BECA4FE561FB7CD6B1AB49B5A38D,PinName="self",PinFriendlyName=NSLOCTEXT("K2Node", "Target", "Target"),PinType.PinCategory="object",PinType.PinSubCategoryObject=BlueprintGeneratedClass'/Game/ExampleContent/BlueprintFluidSurface/WaterSurface_BP.WaterSurface_BP_C',PinType.PinSubCategoryMemberReference=(),PinType.bIsArray=False,PinType.bIsReference=False,PinType.bIsConst=False,PinType.bIsWeakPointer=False,PersistentGuid=00000000000000000000000000000000,bHidden=True,bNotConnectable=False,bDefaultValueIsReadOnly=False,bDefaultValueIsIgnored=False,bAdvancedView=False,)
+			End Object
+			Begin Object Class=K2Node_CommutativeAssociativeBinaryOperator Name="K2Node_CommutativeAssociativeBinaryOperator_1"
+			   bIsPureFunc=True
+			   FunctionReference=(MemberParent=Class'/Script/Engine.KismetMathLibrary',MemberName="Add_IntInt")
+			   NodePosX=1251
+			   NodePosY=160
+			   NodeGuid=8C00357842B050ABF4357FB07035E065
+			   CustomProperties Pin (PinId=7E8934E44D2BC9DB2915A894640F1AC7,PinName="self",PinFriendlyName=NSLOCTEXT("K2Node", "Target", "Target"),PinType.PinCategory="object",PinType.PinSubCategoryObject=Class'/Script/Engine.KismetMathLibrary',PinType.PinSubCategoryMemberReference=(),PinType.bIsArray=False,PinType.bIsReference=False,PinType.bIsConst=False,PinType.bIsWeakPointer=False,DefaultObject="/Script/Engine.Default__KismetMathLibrary",PersistentGuid=00000000000000000000000000000000,bHidden=True,bNotConnectable=False,bDefaultValueIsReadOnly=False,bDefaultValueIsIgnored=False,bAdvancedView=False,)
+			   CustomProperties Pin (PinId=155D15B04B9603A253D24DAA458A93F7,PinName="A",PinType.PinCategory="int",PinType.PinSubCategoryObject=None,PinType.PinSubCategoryMemberReference=(),PinType.bIsArray=False,PinType.bIsReference=False,PinType.bIsConst=False,PinType.bIsWeakPointer=False,DefaultValue="0",AutogeneratedDefaultValue="0",LinkedTo=(K2Node_VariableGet_19 EEB54EFA427F3D4B4674DC8D6620862B,),PersistentGuid=00000000000000000000000000000000,bHidden=False,bNotConnectable=False,bDefaultValueIsReadOnly=False,bDefaultValueIsIgnored=False,bAdvancedView=False,)
+			   CustomProperties Pin (PinId=FBC5B5F5440356E3750501AC4E9918C8,PinName="B",PinType.PinCategory="int",PinType.PinSubCategoryObject=None,PinType.PinSubCategoryMemberReference=(),PinType.bIsArray=False,PinType.bIsReference=False,PinType.bIsConst=False,PinType.bIsWeakPointer=False,DefaultValue="1",AutogeneratedDefaultValue="1",PersistentGuid=00000000000000000000000000000000,bHidden=False,bNotConnectable=False,bDefaultValueIsReadOnly=False,bDefaultValueIsIgnored=False,bAdvancedView=False,)
+			   CustomProperties Pin (PinId=72A6B65C481A76A80CA5E7A48B54430B,PinName="ReturnValue",Direction="EGPD_Output",PinType.PinCategory="int",PinType.PinSubCategoryObject=None,PinType.PinSubCategoryMemberReference=(),PinType.bIsArray=False,PinType.bIsReference=False,PinType.bIsConst=False,PinType.bIsWeakPointer=False,AutogeneratedDefaultValue="0",LinkedTo=(K2Node_CallFunction_347 B2AE70D549846DAAB569619D71CCF7FC,),PersistentGuid=00000000000000000000000000000000,bHidden=False,bNotConnectable=False,bDefaultValueIsReadOnly=False,bDefaultValueIsIgnored=False,bAdvancedView=False,)
+			End Object
+			Begin Object Class=K2Node_CallFunction Name="K2Node_CallFunction_347"
+			   bIsPureFunc=True
+			   FunctionReference=(MemberParent=Class'/Script/Engine.KismetMathLibrary',MemberName="Percent_IntInt")
+			   NodePosX=1523
+			   NodePosY=160
+			   NodeGuid=078F1A8A4ED591BD3A034C91F332A16D
+			   CustomProperties Pin (PinId=F6312BF44720BF89EA7C4DA17BCD7FC6,PinName="self",PinFriendlyName=NSLOCTEXT("K2Node", "Target", "Target"),PinType.PinCategory="object",PinType.PinSubCategoryObject=Class'/Script/Engine.KismetMathLibrary',PinType.PinSubCategoryMemberReference=(),PinType.bIsArray=False,PinType.bIsReference=False,PinType.bIsConst=False,PinType.bIsWeakPointer=False,DefaultObject="/Script/Engine.Default__KismetMathLibrary",PersistentGuid=00000000000000000000000000000000,bHidden=True,bNotConnectable=False,bDefaultValueIsReadOnly=False,bDefaultValueIsIgnored=False,bAdvancedView=False,)
+			   CustomProperties Pin (PinId=B2AE70D549846DAAB569619D71CCF7FC,PinName="A",PinType.PinCategory="int",PinType.PinSubCategoryObject=None,PinType.PinSubCategoryMemberReference=(),PinType.bIsArray=False,PinType.bIsReference=False,PinType.bIsConst=False,PinType.bIsWeakPointer=False,DefaultValue="0",AutogeneratedDefaultValue="0",LinkedTo=(K2Node_CommutativeAssociativeBinaryOperator_1 72A6B65C481A76A80CA5E7A48B54430B,),PersistentGuid=00000000000000000000000000000000,bHidden=False,bNotConnectable=False,bDefaultValueIsReadOnly=False,bDefaultValueIsIgnored=False,bAdvancedView=False,)
+			   CustomProperties Pin (PinId=C1457227427808235772D7A6D4C3301A,PinName="B",PinType.PinCategory="int",PinType.PinSubCategoryObject=None,PinType.PinSubCategoryMemberReference=(),PinType.bIsArray=False,PinType.bIsReference=False,PinType.bIsConst=False,PinType.bIsWeakPointer=False,DefaultValue="3",AutogeneratedDefaultValue="1",PersistentGuid=00000000000000000000000000000000,bHidden=False,bNotConnectable=False,bDefaultValueIsReadOnly=False,bDefaultValueIsIgnored=False,bAdvancedView=False,)
+			   CustomProperties Pin (PinId=9D78D81C48F5A14FD7DFD8857D8AA3D3,PinName="ReturnValue",Direction="EGPD_Output",PinType.PinCategory="int",PinType.PinSubCategoryObject=None,PinType.PinSubCategoryMemberReference=(),PinType.bIsArray=False,PinType.bIsReference=False,PinType.bIsConst=False,PinType.bIsWeakPointer=False,AutogeneratedDefaultValue="0",LinkedTo=(K2Node_VariableSet_14 0DF2ABE447568D05FC81AEA70FEEDD3A,),PersistentGuid=00000000000000000000000000000000,bHidden=False,bNotConnectable=False,bDefaultValueIsReadOnly=False,bDefaultValueIsIgnored=False,bAdvancedView=False,)
+			End Object
+			Begin Object Class=K2Node_VariableSet Name="K2Node_VariableSet_14"
+			   VariableReference=(MemberName="HeightState",MemberGuid=DFF4A74440AF3247B05F028FA72FDFF9,bSelfContext=True)
+			   NodePosX=1651
+			   NodePosY=80
+			   NodeGuid=F43BCB684C40B58986C03B998E3A5A26
+			   CustomProperties Pin (PinId=67D26A9B44D03912C7A167B0C90C861D,PinName="execute",PinType.PinCategory="exec",PinType.PinSubCategoryObject=None,PinType.PinSubCategoryMemberReference=(),PinType.bIsArray=False,PinType.bIsReference=False,PinType.bIsConst=False,PinType.bIsWeakPointer=False,LinkedTo=(K2Node_VariableSet_33 6D73F5F142F9D57CDE67AF95BA4B5CE0,),PersistentGuid=00000000000000000000000000000000,bHidden=False,bNotConnectable=False,bDefaultValueIsReadOnly=False,bDefaultValueIsIgnored=False,bAdvancedView=False,)
+			   CustomProperties Pin (PinId=0B8AB4AA4E5048CF854A00B1212F57A3,PinName="then",Direction="EGPD_Output",PinType.PinCategory="exec",PinType.PinSubCategoryObject=None,PinType.PinSubCategoryMemberReference=(),PinType.bIsArray=False,PinType.bIsReference=False,PinType.bIsConst=False,PinType.bIsWeakPointer=False,LinkedTo=(K2Node_CallMaterialParameterCollectionFunction_42 FB7358C34A3C30CEF1AE7EA6D31A58D4,),PersistentGuid=00000000000000000000000000000000,bHidden=False,bNotConnectable=False,bDefaultValueIsReadOnly=False,bDefaultValueIsIgnored=False,bAdvancedView=False,)
+			   CustomProperties Pin (PinId=0DF2ABE447568D05FC81AEA70FEEDD3A,PinName="HeightState",PinType.PinCategory="int",PinType.PinSubCategoryObject=None,PinType.PinSubCategoryMemberReference=(),PinType.bIsArray=False,PinType.bIsReference=False,PinType.bIsConst=False,PinType.bIsWeakPointer=False,DefaultValue="0",AutogeneratedDefaultValue="0",LinkedTo=(K2Node_CallFunction_347 9D78D81C48F5A14FD7DFD8857D8AA3D3,),PersistentGuid=00000000000000000000000000000000,bHidden=False,bNotConnectable=False,bDefaultValueIsReadOnly=False,bDefaultValueIsIgnored=False,bAdvancedView=False,)
+			   CustomProperties Pin (PinId=D3FEE36F450B306D728843A10CAF338B,PinName="Output_Get",PinToolTip="Retrieves the value of the variable, can use instead of a separate Get node",Direction="EGPD_Output",PinType.PinCategory="int",PinType.PinSubCategoryObject=None,PinType.PinSubCategoryMemberReference=(),PinType.bIsArray=False,PinType.bIsReference=False,PinType.bIsConst=False,PinType.bIsWeakPointer=False,AutogeneratedDefaultValue="0",PersistentGuid=00000000000000000000000000000000,bHidden=False,bNotConnectable=False,bDefaultValueIsReadOnly=False,bDefaultValueIsIgnored=False,bAdvancedView=False,)
+			   CustomProperties Pin (PinId=145D5CE24B4A4A67D1EF429DA4B6E69F,PinName="self",PinFriendlyName=NSLOCTEXT("K2Node", "Target", "Target"),PinType.PinCategory="object",PinType.PinSubCategoryObject=BlueprintGeneratedClass'/Game/ExampleContent/BlueprintFluidSurface/WaterSurface_BP.WaterSurface_BP_C',PinType.PinSubCategoryMemberReference=(),PinType.bIsArray=False,PinType.bIsReference=False,PinType.bIsConst=False,PinType.bIsWeakPointer=False,PersistentGuid=00000000000000000000000000000000,bHidden=True,bNotConnectable=False,bDefaultValueIsReadOnly=False,bDefaultValueIsIgnored=False,bAdvancedView=False,)
+			End Object
+			Begin Object Class=EdGraphNode_Comment Name="EdGraphNode_Comment_1"
+			   NodePosX=1040
+			   NodePosY=16
+			   NodeWidth=832
+			   NodeHeight=256
+			   NodeComment="Cycle active heightfield"
+			   NodeGuid=B74D3A60401B182BB2BC4EBC70EA7F54
+			End Object
+			END OBJECT
+	[/PARAM]
+[/OBJECT]
+[/EXCERPT:CycleActiveHeightfield]
+
+[EXCERPT:ApplySim]	
+[OBJECT:Snippet]
+	[PARAMLITERAL:id]
+		ApplySim
+	[/PARAMLITERAL]
+	[PARAMLITERAL:type]
+		blueprint
+	[/PARAMLITERAL]
+	[PARAM:image]
+		![](Engine\Rendering\RenderTargets\BlueprintRenderTargets\HowTo\FluidSurface\FS_ApplySim.png)
+	[/PARAM]
+	[PARAM:snippet]
+			BEGIN OBJECT
+							Begin Object Class=K2Node_CallFunction Name="K2Node_CallFunction_249"
+			   FunctionReference=(MemberParent=Class'/Script/Engine.Canvas',MemberName="K2_DrawMaterial")
+			   NodePosX=4480
+			   NodePosY=112
+			   NodeGuid=8334A7F3438DDAACAE3BA594527F2A18
+			   CustomProperties Pin (PinId=A92EB2EA4CEB745BB289AFA4E8D3E083,PinName="execute",PinToolTip="\nExec",PinType.PinCategory="exec",PinType.PinSubCategoryObject=None,PinType.PinSubCategoryMemberReference=(),PinType.bIsArray=False,PinType.bIsReference=False,PinType.bIsConst=False,PinType.bIsWeakPointer=False,LinkedTo=(K2Node_CallFunction_247 02B0A2C04D69D8337E659FB041A48B3D,),PersistentGuid=00000000000000000000000000000000,bHidden=False,bNotConnectable=False,bDefaultValueIsReadOnly=False,bDefaultValueIsIgnored=False,bAdvancedView=False,)
+			   CustomProperties Pin (PinId=CDC0C36240D9C55CD3C4E0989677D9CB,PinName="then",PinToolTip="\nExec",Direction="EGPD_Output",PinType.PinCategory="exec",PinType.PinSubCategoryObject=None,PinType.PinSubCategoryMemberReference=(),PinType.bIsArray=False,PinType.bIsReference=False,PinType.bIsConst=False,PinType.bIsWeakPointer=False,LinkedTo=(K2Node_CallFunction_248 826FD10A429E7982FC6C50A368A2EF8B,),PersistentGuid=00000000000000000000000000000000,bHidden=False,bNotConnectable=False,bDefaultValueIsReadOnly=False,bDefaultValueIsIgnored=False,bAdvancedView=False,)
+			   CustomProperties Pin (PinId=8DDBD5C941067A40EC75EBBB4E6BCD74,PinName="self",PinFriendlyName=NSLOCTEXT("K2Node", "Target", "Target"),PinToolTip="Target\nCanvas Reference",PinType.PinCategory="object",PinType.PinSubCategoryObject=Class'/Script/Engine.Canvas',PinType.PinSubCategoryMemberReference=(),PinType.bIsArray=False,PinType.bIsReference=False,PinType.bIsConst=False,PinType.bIsWeakPointer=False,LinkedTo=(K2Node_CallFunction_247 DADFC7E349838F790BC569A556130F8B,),PersistentGuid=00000000000000000000000000000000,bHidden=False,bNotConnectable=False,bDefaultValueIsReadOnly=False,bDefaultValueIsIgnored=False,bAdvancedView=False,)
+			   CustomProperties Pin (PinId=98FDDE614179E45081FD66B0BF94C79E,PinName="RenderMaterial",PinToolTip="Render Material\nMaterial Interface Reference\n\nMaterial to use when rendering. Remember that only the emissive channel is able to be rendered as no lighting is performed when rendering to the Canvas.",PinType.PinCategory="object",PinType.PinSubCategoryObject=Class'/Script/Engine.MaterialInterface',PinType.PinSubCategoryMemberReference=(),PinType.bIsArray=False,PinType.bIsReference=False,PinType.bIsConst=False,PinType.bIsWeakPointer=False,LinkedTo=(K2Node_CallMaterialParameterCollectionFunction_42 89FB58F843256CED989D709D84629586,),PersistentGuid=00000000000000000000000000000000,bHidden=False,bNotConnectable=False,bDefaultValueIsReadOnly=False,bDefaultValueIsIgnored=False,bAdvancedView=False,)
+			   CustomProperties Pin (PinId=32DE39B6447A79161862EC858DDBB3C0,PinName="ScreenPosition",PinToolTip="Screen Position\nVector 2D Structure\n\nScreen space position to render the texture.",PinType.PinCategory="struct",PinType.PinSubCategoryObject=ScriptStruct'/Script/CoreUObject.Vector2D',PinType.PinSubCategoryMemberReference=(),PinType.bIsArray=False,PinType.bIsReference=False,PinType.bIsConst=False,PinType.bIsWeakPointer=False,PersistentGuid=00000000000000000000000000000000,bHidden=False,bNotConnectable=False,bDefaultValueIsReadOnly=False,bDefaultValueIsIgnored=False,bAdvancedView=False,)
+			   CustomProperties Pin (PinId=A27FA11345BBF4A4A5C25785C432EE8C,PinName="ScreenSize",PinToolTip="Screen Size\nVector 2D Structure\n\nScreen space size to render the texture.",PinType.PinCategory="struct",PinType.PinSubCategoryObject=ScriptStruct'/Script/CoreUObject.Vector2D',PinType.PinSubCategoryMemberReference=(),PinType.bIsArray=False,PinType.bIsReference=False,PinType.bIsConst=False,PinType.bIsWeakPointer=False,LinkedTo=(K2Node_CallFunction_247 2D8E4C1A462BB80E885F9D890A179F7D,),PersistentGuid=00000000000000000000000000000000,bHidden=False,bNotConnectable=False,bDefaultValueIsReadOnly=False,bDefaultValueIsIgnored=False,bAdvancedView=False,)
+			   CustomProperties Pin (PinId=21966F7C41F610934660EB967A63AD8C,PinName="CoordinatePosition",PinToolTip="Coordinate Position\nVector 2D Structure\n\nNormalized UV starting coordinate to use when rendering the texture.",PinType.PinCategory="struct",PinType.PinSubCategoryObject=ScriptStruct'/Script/CoreUObject.Vector2D',PinType.PinSubCategoryMemberReference=(),PinType.bIsArray=False,PinType.bIsReference=False,PinType.bIsConst=False,PinType.bIsWeakPointer=False,PersistentGuid=00000000000000000000000000000000,bHidden=False,bNotConnectable=False,bDefaultValueIsReadOnly=False,bDefaultValueIsIgnored=False,bAdvancedView=False,)
+			   CustomProperties Pin (PinId=424BD67E454F48AD7FA5F4B8CE0D0901,PinName="CoordinateSize",PinToolTip="Coordinate Size\nVector 2D Structure\n\nNormalized UV size coordinate to use when rendering the texture.",PinType.PinCategory="struct",PinType.PinSubCategoryObject=ScriptStruct'/Script/CoreUObject.Vector2D',PinType.PinSubCategoryMemberReference=(),PinType.bIsArray=False,PinType.bIsReference=False,PinType.bIsConst=False,PinType.bIsWeakPointer=False,DefaultValue="(X=1.000,Y=1.000)",AutogeneratedDefaultValue="(X=1.000,Y=1.000)",PersistentGuid=00000000000000000000000000000000,bHidden=False,bNotConnectable=False,bDefaultValueIsReadOnly=False,bDefaultValueIsIgnored=False,bAdvancedView=False,)
+			   CustomProperties Pin (PinId=6AE09A9A429F49827A92FCB922F88E56,PinName="Rotation",PinToolTip="Rotation\nFloat\n\nRotation, in degrees, to render the texture.",PinType.PinCategory="float",PinType.PinSubCategoryObject=None,PinType.PinSubCategoryMemberReference=(),PinType.bIsArray=False,PinType.bIsReference=False,PinType.bIsConst=False,PinType.bIsWeakPointer=False,DefaultValue="0.000000",AutogeneratedDefaultValue="0.000000",PersistentGuid=00000000000000000000000000000000,bHidden=False,bNotConnectable=False,bDefaultValueIsReadOnly=False,bDefaultValueIsIgnored=False,bAdvancedView=False,)
+			   CustomProperties Pin (PinId=E6FF7B0F43DBD9C9C4E0188E13DEEB95,PinName="PivotPoint",PinToolTip="Pivot Point\nVector 2D Structure\n\nNormalized pivot point to use when rotating the texture.",PinType.PinCategory="struct",PinType.PinSubCategoryObject=ScriptStruct'/Script/CoreUObject.Vector2D',PinType.PinSubCategoryMemberReference=(),PinType.bIsArray=False,PinType.bIsReference=False,PinType.bIsConst=False,PinType.bIsWeakPointer=False,DefaultValue="(X=0.500,Y=0.500)",AutogeneratedDefaultValue="(X=0.500,Y=0.500)",PersistentGuid=00000000000000000000000000000000,bHidden=False,bNotConnectable=False,bDefaultValueIsReadOnly=False,bDefaultValueIsIgnored=False,bAdvancedView=False,)
+			End Object
+			Begin Object Class=K2Node_CallFunction Name="K2Node_CallFunction_248"
+			   FunctionReference=(MemberParent=Class'/Script/Engine.KismetRenderingLibrary',MemberName="EndDrawCanvasToRenderTarget")
+			   NodePosX=4736
+			   NodePosY=128
+			   NodeGuid=505984BE4B39B458EE03A9B232FA6F33
+			   CustomProperties Pin (PinId=826FD10A429E7982FC6C50A368A2EF8B,PinName="execute",PinToolTip="\nExec",PinType.PinCategory="exec",PinType.PinSubCategoryObject=None,PinType.PinSubCategoryMemberReference=(),PinType.bIsArray=False,PinType.bIsReference=False,PinType.bIsConst=False,PinType.bIsWeakPointer=False,LinkedTo=(K2Node_CallFunction_249 CDC0C36240D9C55CD3C4E0989677D9CB,),PersistentGuid=00000000000000000000000000000000,bHidden=False,bNotConnectable=False,bDefaultValueIsReadOnly=False,bDefaultValueIsIgnored=False,bAdvancedView=False,)
+			   CustomProperties Pin (PinId=4090D05443C22104FBDB3C9F2A8CAF68,PinName="then",PinToolTip="\nExec",Direction="EGPD_Output",PinType.PinCategory="exec",PinType.PinSubCategoryObject=None,PinType.PinSubCategoryMemberReference=(),PinType.bIsArray=False,PinType.bIsReference=False,PinType.bIsConst=False,PinType.bIsWeakPointer=False,PersistentGuid=00000000000000000000000000000000,bHidden=False,bNotConnectable=False,bDefaultValueIsReadOnly=False,bDefaultValueIsIgnored=False,bAdvancedView=False,)
+			   CustomProperties Pin (PinId=1BFF624A4BDB6A248B0AEF93A218CA00,PinName="self",PinFriendlyName=NSLOCTEXT("K2Node", "Target", "Target"),PinToolTip="Target\nKismet Rendering Library Reference",PinType.PinCategory="object",PinType.PinSubCategoryObject=Class'/Script/Engine.KismetRenderingLibrary',PinType.PinSubCategoryMemberReference=(),PinType.bIsArray=False,PinType.bIsReference=False,PinType.bIsConst=False,PinType.bIsWeakPointer=False,DefaultObject="/Script/Engine.Default__KismetRenderingLibrary",PersistentGuid=00000000000000000000000000000000,bHidden=True,bNotConnectable=False,bDefaultValueIsReadOnly=False,bDefaultValueIsIgnored=False,bAdvancedView=False,)
+			   CustomProperties Pin (PinId=F928342C41C1E8AEC1338883F47BB007,PinName="WorldContextObject",PinToolTip="World Context Object\nObject Reference",PinType.PinCategory="object",PinType.PinSubCategoryObject=Class'/Script/CoreUObject.Object',PinType.PinSubCategoryMemberReference=(),PinType.bIsArray=False,PinType.bIsReference=False,PinType.bIsConst=False,PinType.bIsWeakPointer=False,PersistentGuid=00000000000000000000000000000000,bHidden=True,bNotConnectable=False,bDefaultValueIsReadOnly=False,bDefaultValueIsIgnored=False,bAdvancedView=False,)
+			   CustomProperties Pin (PinId=F45CDCA34036451D686B5396C1DA9D79,PinName="Context",PinToolTip="Context\nDraw To Render Target Context Structure (by ref)",PinType.PinCategory="struct",PinType.PinSubCategoryObject=ScriptStruct'/Script/Engine.DrawToRenderTargetContext',PinType.PinSubCategoryMemberReference=(),PinType.bIsArray=False,PinType.bIsReference=True,PinType.bIsConst=True,PinType.bIsWeakPointer=False,LinkedTo=(K2Node_CallFunction_247 FDEC399C4812BB5D95E3FBA8B0133244,),PersistentGuid=00000000000000000000000000000000,bHidden=False,bNotConnectable=False,bDefaultValueIsReadOnly=False,bDefaultValueIsIgnored=True,bAdvancedView=False,)
+			End Object
+			Begin Object Class=K2Node_CallMaterialParameterCollectionFunction Name="K2Node_CallMaterialParameterCollectionFunction_42"
+			   FunctionReference=(MemberParent=Class'/Script/Engine.KismetMaterialLibrary',MemberName="CreateDynamicMaterialInstance")
+			   NodePosX=2336
+			   NodePosY=128
+			   NodeGuid=A9147D4A4EB62A2C89522A8C3AA5A4AE
+			   CustomProperties Pin (PinId=FB7358C34A3C30CEF1AE7EA6D31A58D4,PinName="execute",PinToolTip="\nExec",PinType.PinCategory="exec",PinType.PinSubCategoryObject=None,PinType.PinSubCategoryMemberReference=(),PinType.bIsArray=False,PinType.bIsReference=False,PinType.bIsConst=False,PinType.bIsWeakPointer=False,LinkedTo=(K2Node_VariableSet_14 0B8AB4AA4E5048CF854A00B1212F57A3,),PersistentGuid=00000000000000000000000000000000,bHidden=False,bNotConnectable=False,bDefaultValueIsReadOnly=False,bDefaultValueIsIgnored=False,bAdvancedView=False,)
+			   CustomProperties Pin (PinId=7340368E45A0C843315B54AC2B1EDD86,PinName="then",PinToolTip="\nExec",Direction="EGPD_Output",PinType.PinCategory="exec",PinType.PinSubCategoryObject=None,PinType.PinSubCategoryMemberReference=(),PinType.bIsArray=False,PinType.bIsReference=False,PinType.bIsConst=False,PinType.bIsWeakPointer=False,LinkedTo=(K2Node_CallFunction_381 8D844E1D41E88AED47166F94960D29D2,),PersistentGuid=00000000000000000000000000000000,bHidden=False,bNotConnectable=False,bDefaultValueIsReadOnly=False,bDefaultValueIsIgnored=False,bAdvancedView=False,)
+			   CustomProperties Pin (PinId=13A08EDE4EA26D7429E84E95032DFDAB,PinName="self",PinFriendlyName=NSLOCTEXT("K2Node", "Target", "Target"),PinToolTip="Target\nKismet Material Library Reference",PinType.PinCategory="object",PinType.PinSubCategoryObject=Class'/Script/Engine.KismetMaterialLibrary',PinType.PinSubCategoryMemberReference=(),PinType.bIsArray=False,PinType.bIsReference=False,PinType.bIsConst=False,PinType.bIsWeakPointer=False,DefaultObject="/Script/Engine.Default__KismetMaterialLibrary",PersistentGuid=00000000000000000000000000000000,bHidden=True,bNotConnectable=False,bDefaultValueIsReadOnly=False,bDefaultValueIsIgnored=False,bAdvancedView=False,)
+			   CustomProperties Pin (PinId=0224F2244EC6D3B3CA52ED908E60B0A2,PinName="WorldContextObject",PinToolTip="World Context Object\nObject Reference",PinType.PinCategory="object",PinType.PinSubCategoryObject=Class'/Script/CoreUObject.Object',PinType.PinSubCategoryMemberReference=(),PinType.bIsArray=False,PinType.bIsReference=False,PinType.bIsConst=False,PinType.bIsWeakPointer=False,PersistentGuid=00000000000000000000000000000000,bHidden=True,bNotConnectable=False,bDefaultValueIsReadOnly=False,bDefaultValueIsIgnored=False,bAdvancedView=False,)
+			   CustomProperties Pin (PinId=4BC347BF40CB686B8AEA7EAC3B6D29A3,PinName="Parent",PinToolTip="Parent\nMaterial Interface Reference",PinType.PinCategory="object",PinType.PinSubCategoryObject=Class'/Script/Engine.MaterialInterface',PinType.PinSubCategoryMemberReference=(),PinType.bIsArray=False,PinType.bIsReference=False,PinType.bIsConst=False,PinType.bIsWeakPointer=False,DefaultObject="/Game/ExampleContent/BlueprintFluidSurface/HeightSim.HeightSim",PersistentGuid=00000000000000000000000000000000,bHidden=False,bNotConnectable=False,bDefaultValueIsReadOnly=False,bDefaultValueIsIgnored=False,bAdvancedView=False,)
+			   CustomProperties Pin (PinId=89FB58F843256CED989D709D84629586,PinName="ReturnValue",PinToolTip="Return Value\nMaterial Instance Dynamic Reference",Direction="EGPD_Output",PinType.PinCategory="object",PinType.PinSubCategoryObject=Class'/Script/Engine.MaterialInstanceDynamic',PinType.PinSubCategoryMemberReference=(),PinType.bIsArray=False,PinType.bIsReference=False,PinType.bIsConst=False,PinType.bIsWeakPointer=False,LinkedTo=(K2Node_CallFunction_382 3D5799644683985A609A4C863EE86A64,K2Node_CallFunction_378 ADF4A63143B4F65786F73E8DF9C1CBD4,K2Node_CallFunction_149 9E05610046120AA6FC4E36B6A1BA014F,K2Node_CallFunction_249 98FDDE614179E45081FD66B0BF94C79E,),PersistentGuid=00000000000000000000000000000000,bHidden=False,bNotConnectable=False,bDefaultValueIsReadOnly=False,bDefaultValueIsIgnored=False,bAdvancedView=False,)
+			End Object
+			Begin Object Class=K2Node_CallFunction Name="K2Node_CallFunction_382"
+			   FunctionReference=(MemberParent=Class'/Script/Engine.MaterialInstanceDynamic',MemberName="SetTextureParameterValue")
+			   NodePosX=2944
+			   NodePosY=112
+			   NodeGuid=742577924C4DCD90041B61A185598306
+			   CustomProperties Pin (PinId=2BBE582A424E4C1F87695C8DFC7121C6,PinName="execute",PinToolTip="\nExec",PinType.PinCategory="exec",PinType.PinSubCategoryObject=None,PinType.PinSubCategoryMemberReference=(),PinType.bIsArray=False,PinType.bIsReference=False,PinType.bIsConst=False,PinType.bIsWeakPointer=False,LinkedTo=(K2Node_CallFunction_381 D0249D394D4CDEAA09BB3DAE83C633DA,),PersistentGuid=00000000000000000000000000000000,bHidden=False,bNotConnectable=False,bDefaultValueIsReadOnly=False,bDefaultValueIsIgnored=False,bAdvancedView=False,)
+			   CustomProperties Pin (PinId=5C51FAB746DE3470AA92899F1843CEE8,PinName="then",PinToolTip="\nExec",Direction="EGPD_Output",PinType.PinCategory="exec",PinType.PinSubCategoryObject=None,PinType.PinSubCategoryMemberReference=(),PinType.bIsArray=False,PinType.bIsReference=False,PinType.bIsConst=False,PinType.bIsWeakPointer=False,LinkedTo=(K2Node_CallFunction_380 55D2EFEC4A45A7A5AA4FBFB8774A207C,),PersistentGuid=00000000000000000000000000000000,bHidden=False,bNotConnectable=False,bDefaultValueIsReadOnly=False,bDefaultValueIsIgnored=False,bAdvancedView=False,)
+			   CustomProperties Pin (PinId=3D5799644683985A609A4C863EE86A64,PinName="self",PinFriendlyName=NSLOCTEXT("K2Node", "Target", "Target"),PinToolTip="Target\nMaterial Instance Dynamic Reference",PinType.PinCategory="object",PinType.PinSubCategoryObject=Class'/Script/Engine.MaterialInstanceDynamic',PinType.PinSubCategoryMemberReference=(),PinType.bIsArray=False,PinType.bIsReference=False,PinType.bIsConst=False,PinType.bIsWeakPointer=False,LinkedTo=(K2Node_CallMaterialParameterCollectionFunction_42 89FB58F843256CED989D709D84629586,),PersistentGuid=00000000000000000000000000000000,bHidden=False,bNotConnectable=False,bDefaultValueIsReadOnly=False,bDefaultValueIsIgnored=False,bAdvancedView=False,)
+			   CustomProperties Pin (PinId=9969FBE341D2705B7931828A52CF8F86,PinName="ParameterName",PinToolTip="Parameter Name\nName",PinType.PinCategory="name",PinType.PinSubCategoryObject=None,PinType.PinSubCategoryMemberReference=(),PinType.bIsArray=False,PinType.bIsReference=False,PinType.bIsConst=False,PinType.bIsWeakPointer=False,DefaultValue="PreviousHeight1",AutogeneratedDefaultValue="None",PersistentGuid=00000000000000000000000000000000,bHidden=False,bNotConnectable=False,bDefaultValueIsReadOnly=False,bDefaultValueIsIgnored=False,bAdvancedView=False,)
+			   CustomProperties Pin (PinId=DDBA011D4F5D3CFA0DAEEEBDD3DA5B28,PinName="Value",PinToolTip="Value\nTexture Reference",PinType.PinCategory="object",PinType.PinSubCategoryObject=Class'/Script/Engine.Texture',PinType.PinSubCategoryMemberReference=(),PinType.bIsArray=False,PinType.bIsReference=False,PinType.bIsConst=False,PinType.bIsWeakPointer=False,LinkedTo=(K2Node_CallFunction_381 8B5E6BEC40B04E4589E04DBFF9431A99,),PersistentGuid=00000000000000000000000000000000,bHidden=False,bNotConnectable=False,bDefaultValueIsReadOnly=False,bDefaultValueIsIgnored=False,bAdvancedView=False,)
+			End Object
+			Begin Object Class=K2Node_CallFunction Name="K2Node_CallFunction_381"
+			   FunctionReference=(MemberName="GetLastHeight",MemberGuid=D948A6AF4D6A2872B1F0D3AD97E1A286,bSelfContext=True)
+			   NodePosX=2656
+			   NodePosY=112
+			   NodeGuid=D3378652497B2CFA5FE13E8E8649A02B
+			   CustomProperties Pin (PinId=8D844E1D41E88AED47166F94960D29D2,PinName="execute",PinToolTip="\nExec",PinType.PinCategory="exec",PinType.PinSubCategoryObject=None,PinType.PinSubCategoryMemberReference=(),PinType.bIsArray=False,PinType.bIsReference=False,PinType.bIsConst=False,PinType.bIsWeakPointer=False,LinkedTo=(K2Node_CallMaterialParameterCollectionFunction_42 7340368E45A0C843315B54AC2B1EDD86,),PersistentGuid=00000000000000000000000000000000,bHidden=False,bNotConnectable=False,bDefaultValueIsReadOnly=False,bDefaultValueIsIgnored=False,bAdvancedView=False,)
+			   CustomProperties Pin (PinId=D0249D394D4CDEAA09BB3DAE83C633DA,PinName="then",PinToolTip="\nExec",Direction="EGPD_Output",PinType.PinCategory="exec",PinType.PinSubCategoryObject=None,PinType.PinSubCategoryMemberReference=(),PinType.bIsArray=False,PinType.bIsReference=False,PinType.bIsConst=False,PinType.bIsWeakPointer=False,LinkedTo=(K2Node_CallFunction_382 2BBE582A424E4C1F87695C8DFC7121C6,),PersistentGuid=00000000000000000000000000000000,bHidden=False,bNotConnectable=False,bDefaultValueIsReadOnly=False,bDefaultValueIsIgnored=False,bAdvancedView=False,)
+			   CustomProperties Pin (PinId=B2163A8B4459B111A3007EBFCE1F8C64,PinName="self",PinFriendlyName=NSLOCTEXT("K2Node", "Target", "Target"),PinToolTip="Target\nSelf Reference",PinType.PinCategory="object",PinType.PinSubCategory="self",PinType.PinSubCategoryObject=None,PinType.PinSubCategoryMemberReference=(),PinType.bIsArray=False,PinType.bIsReference=False,PinType.bIsConst=False,PinType.bIsWeakPointer=False,PersistentGuid=00000000000000000000000000000000,bHidden=False,bNotConnectable=False,bDefaultValueIsReadOnly=False,bDefaultValueIsIgnored=False,bAdvancedView=False,)
+			   CustomProperties Pin (PinId=146F20394178EE7CF8C271BD3A36F211,PinName="CurrentHeightIndex",PinToolTip="Current Height Index\nInteger",PinType.PinCategory="int",PinType.PinSubCategoryObject=None,PinType.PinSubCategoryMemberReference=(),PinType.bIsArray=False,PinType.bIsReference=False,PinType.bIsConst=False,PinType.bIsWeakPointer=False,DefaultValue="0",AutogeneratedDefaultValue="0",LinkedTo=(K2Node_VariableGet_153 EEEDB4374B9823941033B880771F9C88,),PersistentGuid=00000000000000000000000000000000,bHidden=False,bNotConnectable=False,bDefaultValueIsReadOnly=False,bDefaultValueIsIgnored=False,bAdvancedView=False,)
+			   CustomProperties Pin (PinId=FB60E5104A64D99DBF5848BF3CF8F7FA,PinName="NumFramesOld",PinToolTip="Num Frames Old\nInteger",PinType.PinCategory="int",PinType.PinSubCategoryObject=None,PinType.PinSubCategoryMemberReference=(),PinType.bIsArray=False,PinType.bIsReference=False,PinType.bIsConst=False,PinType.bIsWeakPointer=False,DefaultValue="1",AutogeneratedDefaultValue="0",PersistentGuid=00000000000000000000000000000000,bHidden=False,bNotConnectable=False,bDefaultValueIsReadOnly=False,bDefaultValueIsIgnored=False,bAdvancedView=False,)
+			   CustomProperties Pin (PinId=8B5E6BEC40B04E4589E04DBFF9431A99,PinName="HeightRT",PinToolTip="Height RT\nTexture Render Target 2D Reference",Direction="EGPD_Output",PinType.PinCategory="object",PinType.PinSubCategoryObject=Class'/Script/Engine.TextureRenderTarget2D',PinType.PinSubCategoryMemberReference=(),PinType.bIsArray=False,PinType.bIsReference=False,PinType.bIsConst=False,PinType.bIsWeakPointer=False,LinkedTo=(K2Node_CallFunction_382 DDBA011D4F5D3CFA0DAEEEBDD3DA5B28,),PersistentGuid=00000000000000000000000000000000,bHidden=False,bNotConnectable=False,bDefaultValueIsReadOnly=False,bDefaultValueIsIgnored=False,bAdvancedView=False,)
+			End Object
+			Begin Object Class=K2Node_CallFunction Name="K2Node_CallFunction_380"
+			   FunctionReference=(MemberName="GetLastHeight",MemberGuid=D948A6AF4D6A2872B1F0D3AD97E1A286,bSelfContext=True)
+			   NodePosX=3232
+			   NodePosY=112
+			   NodeGuid=9570E928439D7F6CAA7C16A557AA4BAF
+			   CustomProperties Pin (PinId=55D2EFEC4A45A7A5AA4FBFB8774A207C,PinName="execute",PinToolTip="\nExec",PinType.PinCategory="exec",PinType.PinSubCategoryObject=None,PinType.PinSubCategoryMemberReference=(),PinType.bIsArray=False,PinType.bIsReference=False,PinType.bIsConst=False,PinType.bIsWeakPointer=False,LinkedTo=(K2Node_CallFunction_382 5C51FAB746DE3470AA92899F1843CEE8,),PersistentGuid=00000000000000000000000000000000,bHidden=False,bNotConnectable=False,bDefaultValueIsReadOnly=False,bDefaultValueIsIgnored=False,bAdvancedView=False,)
+			   CustomProperties Pin (PinId=8E7DC5644F8C2B43B136859B12ADD25E,PinName="then",PinToolTip="\nExec",Direction="EGPD_Output",PinType.PinCategory="exec",PinType.PinSubCategoryObject=None,PinType.PinSubCategoryMemberReference=(),PinType.bIsArray=False,PinType.bIsReference=False,PinType.bIsConst=False,PinType.bIsWeakPointer=False,LinkedTo=(K2Node_CallFunction_378 390DDED34C5FE3120EDD41AF6C4DABE8,),PersistentGuid=00000000000000000000000000000000,bHidden=False,bNotConnectable=False,bDefaultValueIsReadOnly=False,bDefaultValueIsIgnored=False,bAdvancedView=False,)
+			   CustomProperties Pin (PinId=C7F80E984740A1694DED6789CC243A81,PinName="self",PinFriendlyName=NSLOCTEXT("K2Node", "Target", "Target"),PinToolTip="Target\nSelf Reference",PinType.PinCategory="object",PinType.PinSubCategory="self",PinType.PinSubCategoryObject=None,PinType.PinSubCategoryMemberReference=(),PinType.bIsArray=False,PinType.bIsReference=False,PinType.bIsConst=False,PinType.bIsWeakPointer=False,PersistentGuid=00000000000000000000000000000000,bHidden=False,bNotConnectable=False,bDefaultValueIsReadOnly=False,bDefaultValueIsIgnored=False,bAdvancedView=False,)
+			   CustomProperties Pin (PinId=2F2CFA6B4B89D3081E4E1FB73E651534,PinName="CurrentHeightIndex",PinToolTip="Current Height Index\nInteger",PinType.PinCategory="int",PinType.PinSubCategoryObject=None,PinType.PinSubCategoryMemberReference=(),PinType.bIsArray=False,PinType.bIsReference=False,PinType.bIsConst=False,PinType.bIsWeakPointer=False,DefaultValue="0",AutogeneratedDefaultValue="0",LinkedTo=(K2Node_VariableGet_33 3FC6A51C457038BB48787BB6A2F0EB40,),PersistentGuid=00000000000000000000000000000000,bHidden=False,bNotConnectable=False,bDefaultValueIsReadOnly=False,bDefaultValueIsIgnored=False,bAdvancedView=False,)
+			   CustomProperties Pin (PinId=3067CAC04F4099EE8237BD9CB060E089,PinName="NumFramesOld",PinToolTip="Num Frames Old\nInteger",PinType.PinCategory="int",PinType.PinSubCategoryObject=None,PinType.PinSubCategoryMemberReference=(),PinType.bIsArray=False,PinType.bIsReference=False,PinType.bIsConst=False,PinType.bIsWeakPointer=False,DefaultValue="2",AutogeneratedDefaultValue="0",PersistentGuid=00000000000000000000000000000000,bHidden=False,bNotConnectable=False,bDefaultValueIsReadOnly=False,bDefaultValueIsIgnored=False,bAdvancedView=False,)
+			   CustomProperties Pin (PinId=BC8CC8E042D9442683296897B63DB119,PinName="HeightRT",PinToolTip="Height RT\nTexture Render Target 2D Reference",Direction="EGPD_Output",PinType.PinCategory="object",PinType.PinSubCategoryObject=Class'/Script/Engine.TextureRenderTarget2D',PinType.PinSubCategoryMemberReference=(),PinType.bIsArray=False,PinType.bIsReference=False,PinType.bIsConst=False,PinType.bIsWeakPointer=False,LinkedTo=(K2Node_CallFunction_378 F93870664C80B4361A8BD6B4C9B1DEB1,),PersistentGuid=00000000000000000000000000000000,bHidden=False,bNotConnectable=False,bDefaultValueIsReadOnly=False,bDefaultValueIsIgnored=False,bAdvancedView=False,)
+			End Object
+			Begin Object Class=K2Node_CallFunction Name="K2Node_CallFunction_378"
+			   FunctionReference=(MemberParent=Class'/Script/Engine.MaterialInstanceDynamic',MemberName="SetTextureParameterValue")
+			   NodePosX=3536
+			   NodePosY=112
+			   NodeGuid=AA17F7004D1E008BF12365A74D48BDA7
+			   CustomProperties Pin (PinId=390DDED34C5FE3120EDD41AF6C4DABE8,PinName="execute",PinToolTip="\nExec",PinType.PinCategory="exec",PinType.PinSubCategoryObject=None,PinType.PinSubCategoryMemberReference=(),PinType.bIsArray=False,PinType.bIsReference=False,PinType.bIsConst=False,PinType.bIsWeakPointer=False,LinkedTo=(K2Node_CallFunction_380 8E7DC5644F8C2B43B136859B12ADD25E,),PersistentGuid=00000000000000000000000000000000,bHidden=False,bNotConnectable=False,bDefaultValueIsReadOnly=False,bDefaultValueIsIgnored=False,bAdvancedView=False,)
+			   CustomProperties Pin (PinId=5F0D53E544CB8A0F095EBF9B0CE895E5,PinName="then",PinToolTip="\nExec",Direction="EGPD_Output",PinType.PinCategory="exec",PinType.PinSubCategoryObject=None,PinType.PinSubCategoryMemberReference=(),PinType.bIsArray=False,PinType.bIsReference=False,PinType.bIsConst=False,PinType.bIsWeakPointer=False,LinkedTo=(K2Node_CallFunction_1111 EE4DBD094867FEF67ED66A869D58EFB4,),PersistentGuid=00000000000000000000000000000000,bHidden=False,bNotConnectable=False,bDefaultValueIsReadOnly=False,bDefaultValueIsIgnored=False,bAdvancedView=False,)
+			   CustomProperties Pin (PinId=ADF4A63143B4F65786F73E8DF9C1CBD4,PinName="self",PinFriendlyName=NSLOCTEXT("K2Node", "Target", "Target"),PinToolTip="Target\nMaterial Instance Dynamic Reference",PinType.PinCategory="object",PinType.PinSubCategoryObject=Class'/Script/Engine.MaterialInstanceDynamic',PinType.PinSubCategoryMemberReference=(),PinType.bIsArray=False,PinType.bIsReference=False,PinType.bIsConst=False,PinType.bIsWeakPointer=False,LinkedTo=(K2Node_CallMaterialParameterCollectionFunction_42 89FB58F843256CED989D709D84629586,),PersistentGuid=00000000000000000000000000000000,bHidden=False,bNotConnectable=False,bDefaultValueIsReadOnly=False,bDefaultValueIsIgnored=False,bAdvancedView=False,)
+			   CustomProperties Pin (PinId=F537375441AF4FC3BA6DF6B7C78CAC37,PinName="ParameterName",PinToolTip="Parameter Name\nName",PinType.PinCategory="name",PinType.PinSubCategoryObject=None,PinType.PinSubCategoryMemberReference=(),PinType.bIsArray=False,PinType.bIsReference=False,PinType.bIsConst=False,PinType.bIsWeakPointer=False,DefaultValue="PreviousHeight2",AutogeneratedDefaultValue="None",PersistentGuid=00000000000000000000000000000000,bHidden=False,bNotConnectable=False,bDefaultValueIsReadOnly=False,bDefaultValueIsIgnored=False,bAdvancedView=False,)
+			   CustomProperties Pin (PinId=F93870664C80B4361A8BD6B4C9B1DEB1,PinName="Value",PinToolTip="Value\nTexture Reference",PinType.PinCategory="object",PinType.PinSubCategoryObject=Class'/Script/Engine.Texture',PinType.PinSubCategoryMemberReference=(),PinType.bIsArray=False,PinType.bIsReference=False,PinType.bIsConst=False,PinType.bIsWeakPointer=False,LinkedTo=(K2Node_CallFunction_380 BC8CC8E042D9442683296897B63DB119,),PersistentGuid=00000000000000000000000000000000,bHidden=False,bNotConnectable=False,bDefaultValueIsReadOnly=False,bDefaultValueIsIgnored=False,bAdvancedView=False,)
+			End Object
+			Begin Object Class=K2Node_VariableGet Name="K2Node_VariableGet_33"
+			   VariableReference=(MemberName="HeightState",MemberGuid=DFF4A74440AF3247B05F028FA72FDFF9,bSelfContext=True)
+			   NodePosX=3040
+			   NodePosY=336
+			   NodeGuid=368C714244FA166A22415BAB1807ACE8
+			   CustomProperties Pin (PinId=3FC6A51C457038BB48787BB6A2F0EB40,PinName="HeightState",Direction="EGPD_Output",PinType.PinCategory="int",PinType.PinSubCategoryObject=None,PinType.PinSubCategoryMemberReference=(),PinType.bIsArray=False,PinType.bIsReference=False,PinType.bIsConst=False,PinType.bIsWeakPointer=False,AutogeneratedDefaultValue="0",LinkedTo=(K2Node_CallFunction_380 2F2CFA6B4B89D3081E4E1FB73E651534,),PersistentGuid=00000000000000000000000000000000,bHidden=False,bNotConnectable=False,bDefaultValueIsReadOnly=False,bDefaultValueIsIgnored=False,bAdvancedView=False,)
+			   CustomProperties Pin (PinId=A17209204B913D7BAE68A4BA444D7B03,PinName="self",PinFriendlyName=NSLOCTEXT("K2Node", "Target", "Target"),PinType.PinCategory="object",PinType.PinSubCategoryObject=BlueprintGeneratedClass'/Game/ExampleContent/BlueprintFluidSurface/WaterSurface_BP.WaterSurface_BP_C',PinType.PinSubCategoryMemberReference=(),PinType.bIsArray=False,PinType.bIsReference=False,PinType.bIsConst=False,PinType.bIsWeakPointer=False,PersistentGuid=00000000000000000000000000000000,bHidden=True,bNotConnectable=False,bDefaultValueIsReadOnly=False,bDefaultValueIsIgnored=False,bAdvancedView=False,)
+			End Object
+			Begin Object Class=K2Node_CallFunction Name="K2Node_CallFunction_1111"
+			   FunctionReference=(MemberName="GetHeightRT",MemberGuid=2EEA338F4CA878E7228A10BB2DD63E00,bSelfContext=True)
+			   NodePosX=3824
+			   NodePosY=112
+			   NodeGuid=2443F3CF41CF209EC5EB2F895589A983
+			   CustomProperties Pin (PinId=EE4DBD094867FEF67ED66A869D58EFB4,PinName="execute",PinToolTip="\nExec",PinType.PinCategory="exec",PinType.PinSubCategoryObject=None,PinType.PinSubCategoryMemberReference=(),PinType.bIsArray=False,PinType.bIsReference=False,PinType.bIsConst=False,PinType.bIsWeakPointer=False,LinkedTo=(K2Node_CallFunction_378 5F0D53E544CB8A0F095EBF9B0CE895E5,),PersistentGuid=00000000000000000000000000000000,bHidden=False,bNotConnectable=False,bDefaultValueIsReadOnly=False,bDefaultValueIsIgnored=False,bAdvancedView=False,)
+			   CustomProperties Pin (PinId=518D003E46FCB89DB2EBFA9116A5F4F8,PinName="then",PinToolTip="\nExec",Direction="EGPD_Output",PinType.PinCategory="exec",PinType.PinSubCategoryObject=None,PinType.PinSubCategoryMemberReference=(),PinType.bIsArray=False,PinType.bIsReference=False,PinType.bIsConst=False,PinType.bIsWeakPointer=False,LinkedTo=(K2Node_CallFunction_247 DB2587924EB03B7375ECB4B67F84ACEC,),PersistentGuid=00000000000000000000000000000000,bHidden=False,bNotConnectable=False,bDefaultValueIsReadOnly=False,bDefaultValueIsIgnored=False,bAdvancedView=False,)
+			   CustomProperties Pin (PinId=428F55744E9265292CF54CA9B62BF4D1,PinName="self",PinFriendlyName=NSLOCTEXT("K2Node", "Target", "Target"),PinToolTip="Target\nSelf Reference",PinType.PinCategory="object",PinType.PinSubCategory="self",PinType.PinSubCategoryObject=None,PinType.PinSubCategoryMemberReference=(),PinType.bIsArray=False,PinType.bIsReference=False,PinType.bIsConst=False,PinType.bIsWeakPointer=False,PersistentGuid=00000000000000000000000000000000,bHidden=False,bNotConnectable=False,bDefaultValueIsReadOnly=False,bDefaultValueIsIgnored=False,bAdvancedView=False,)
+			   CustomProperties Pin (PinId=6CAABD09475E1A3DEA8C9DA7043D5D5D,PinName="Index",PinToolTip="Index\nInteger",PinType.PinCategory="int",PinType.PinSubCategoryObject=None,PinType.PinSubCategoryMemberReference=(),PinType.bIsArray=False,PinType.bIsReference=False,PinType.bIsConst=False,PinType.bIsWeakPointer=False,DefaultValue="0",AutogeneratedDefaultValue="0",LinkedTo=(K2Node_VariableGet_351 7C588D7E4937016CA7C60B9AD95D8524,),PersistentGuid=00000000000000000000000000000000,bHidden=False,bNotConnectable=False,bDefaultValueIsReadOnly=False,bDefaultValueIsIgnored=False,bAdvancedView=False,)
+			   CustomProperties Pin (PinId=DD0B4644476E458B894265AF5D1FB6BA,PinName="HeightRT",PinToolTip="Height RT\nTexture Render Target 2D Reference",Direction="EGPD_Output",PinType.PinCategory="object",PinType.PinSubCategoryObject=Class'/Script/Engine.TextureRenderTarget2D',PinType.PinSubCategoryMemberReference=(),PinType.bIsArray=False,PinType.bIsReference=False,PinType.bIsConst=False,PinType.bIsWeakPointer=False,LinkedTo=(K2Node_CallFunction_149 4071F025450DAD7A0AF9929767345F09,K2Node_CallFunction_247 52F872CE44DEAA359406B0B6F0A9F744,),PersistentGuid=00000000000000000000000000000000,bHidden=False,bNotConnectable=False,bDefaultValueIsReadOnly=False,bDefaultValueIsIgnored=False,bAdvancedView=False,)
+			End Object
+			Begin Object Class=K2Node_VariableGet Name="K2Node_VariableGet_351"
+			   VariableReference=(MemberName="HeightState",MemberGuid=8E10E24F4154A224C2065A89B367BD62,bSelfContext=True)
+			   NodePosX=3664
+			   NodePosY=304
+			   NodeGuid=808BAF994036BE2EA241BC963C05B011
+			   CustomProperties Pin (PinId=7C588D7E4937016CA7C60B9AD95D8524,PinName="HeightState",Direction="EGPD_Output",PinType.PinCategory="int",PinType.PinSubCategoryObject=None,PinType.PinSubCategoryMemberReference=(),PinType.bIsArray=False,PinType.bIsReference=False,PinType.bIsConst=False,PinType.bIsWeakPointer=False,AutogeneratedDefaultValue="0",LinkedTo=(K2Node_CallFunction_1111 6CAABD09475E1A3DEA8C9DA7043D5D5D,),PersistentGuid=00000000000000000000000000000000,bHidden=False,bNotConnectable=False,bDefaultValueIsReadOnly=False,bDefaultValueIsIgnored=False,bAdvancedView=False,)
+			   CustomProperties Pin (PinId=FA936FFF447FAB5ADAD46984445701FE,PinName="self",PinFriendlyName=NSLOCTEXT("K2Node", "Target", "Target"),PinType.PinCategory="object",PinType.PinSubCategoryObject=BlueprintGeneratedClass'/Game/ExampleContent/BlueprintFluidSurface/WaterSurface_BP.WaterSurface_BP_C',PinType.PinSubCategoryMemberReference=(),PinType.bIsArray=False,PinType.bIsReference=False,PinType.bIsConst=False,PinType.bIsWeakPointer=False,PersistentGuid=00000000000000000000000000000000,bHidden=True,bNotConnectable=False,bDefaultValueIsReadOnly=False,bDefaultValueIsIgnored=False,bAdvancedView=False,)
+			End Object
+			Begin Object Class=K2Node_VariableGet Name="K2Node_VariableGet_153"
+			   VariableReference=(MemberName="HeightState",MemberGuid=8E10E24F4154A224C2065A89B367BD62,bSelfContext=True)
+			   NodePosX=2416
+			   NodePosY=288
+			   NodeGuid=B57467624CE4D236E9BFDF88A35D8470
+			   CustomProperties Pin (PinId=EEEDB4374B9823941033B880771F9C88,PinName="HeightState",Direction="EGPD_Output",PinType.PinCategory="int",PinType.PinSubCategoryObject=None,PinType.PinSubCategoryMemberReference=(),PinType.bIsArray=False,PinType.bIsReference=False,PinType.bIsConst=False,PinType.bIsWeakPointer=False,AutogeneratedDefaultValue="0",LinkedTo=(K2Node_CallFunction_381 146F20394178EE7CF8C271BD3A36F211,),PersistentGuid=00000000000000000000000000000000,bHidden=False,bNotConnectable=False,bDefaultValueIsReadOnly=False,bDefaultValueIsIgnored=False,bAdvancedView=False,)
+			   CustomProperties Pin (PinId=B68E52FE436EF5470AC9EA9495942498,PinName="self",PinFriendlyName=NSLOCTEXT("K2Node", "Target", "Target"),PinType.PinCategory="object",PinType.PinSubCategoryObject=BlueprintGeneratedClass'/Game/ExampleContent/BlueprintFluidSurface/WaterSurface_BP.WaterSurface_BP_C',PinType.PinSubCategoryMemberReference=(),PinType.bIsArray=False,PinType.bIsReference=False,PinType.bIsConst=False,PinType.bIsWeakPointer=False,PersistentGuid=00000000000000000000000000000000,bHidden=True,bNotConnectable=False,bDefaultValueIsReadOnly=False,bDefaultValueIsIgnored=False,bAdvancedView=False,)
+			End Object
+			Begin Object Class=K2Node_CallFunction Name="K2Node_CallFunction_247"
+			   FunctionReference=(MemberParent=Class'/Script/Engine.KismetRenderingLibrary',MemberName="BeginDrawCanvasToRenderTarget")
+			   NodePosX=4096
+			   NodePosY=128
+			   NodeGuid=E0F44CEA4191F2CC621E0DB0E51718BF
+			   CustomProperties Pin (PinId=DB2587924EB03B7375ECB4B67F84ACEC,PinName="execute",PinToolTip="\nExec",PinType.PinCategory="exec",PinType.PinSubCategoryObject=None,PinType.PinSubCategoryMemberReference=(),PinType.bIsArray=False,PinType.bIsReference=False,PinType.bIsConst=False,PinType.bIsWeakPointer=False,LinkedTo=(K2Node_CallFunction_1111 518D003E46FCB89DB2EBFA9116A5F4F8,),PersistentGuid=00000000000000000000000000000000,bHidden=False,bNotConnectable=False,bDefaultValueIsReadOnly=False,bDefaultValueIsIgnored=False,bAdvancedView=False,)
+			   CustomProperties Pin (PinId=02B0A2C04D69D8337E659FB041A48B3D,PinName="then",PinToolTip="\nExec",Direction="EGPD_Output",PinType.PinCategory="exec",PinType.PinSubCategoryObject=None,PinType.PinSubCategoryMemberReference=(),PinType.bIsArray=False,PinType.bIsReference=False,PinType.bIsConst=False,PinType.bIsWeakPointer=False,LinkedTo=(K2Node_CallFunction_249 A92EB2EA4CEB745BB289AFA4E8D3E083,),PersistentGuid=00000000000000000000000000000000,bHidden=False,bNotConnectable=False,bDefaultValueIsReadOnly=False,bDefaultValueIsIgnored=False,bAdvancedView=False,)
+			   CustomProperties Pin (PinId=3622CA9D4D61E87B27A869A43F3DE87A,PinName="self",PinFriendlyName=NSLOCTEXT("K2Node", "Target", "Target"),PinToolTip="Target\nKismet Rendering Library Reference",PinType.PinCategory="object",PinType.PinSubCategoryObject=Class'/Script/Engine.KismetRenderingLibrary',PinType.PinSubCategoryMemberReference=(),PinType.bIsArray=False,PinType.bIsReference=False,PinType.bIsConst=False,PinType.bIsWeakPointer=False,DefaultObject="/Script/Engine.Default__KismetRenderingLibrary",PersistentGuid=00000000000000000000000000000000,bHidden=True,bNotConnectable=False,bDefaultValueIsReadOnly=False,bDefaultValueIsIgnored=False,bAdvancedView=False,)
+			   CustomProperties Pin (PinId=CE65EABA4B8D9DA8E21485B75C734C0E,PinName="WorldContextObject",PinToolTip="World Context Object\nObject Reference",PinType.PinCategory="object",PinType.PinSubCategoryObject=Class'/Script/CoreUObject.Object',PinType.PinSubCategoryMemberReference=(),PinType.bIsArray=False,PinType.bIsReference=False,PinType.bIsConst=False,PinType.bIsWeakPointer=False,PersistentGuid=00000000000000000000000000000000,bHidden=True,bNotConnectable=False,bDefaultValueIsReadOnly=False,bDefaultValueIsIgnored=False,bAdvancedView=False,)
+			   CustomProperties Pin (PinId=52F872CE44DEAA359406B0B6F0A9F744,PinName="TextureRenderTarget",PinToolTip="Texture Render Target\nTexture Render Target 2D Reference",PinType.PinCategory="object",PinType.PinSubCategoryObject=Class'/Script/Engine.TextureRenderTarget2D',PinType.PinSubCategoryMemberReference=(),PinType.bIsArray=False,PinType.bIsReference=False,PinType.bIsConst=False,PinType.bIsWeakPointer=False,LinkedTo=(K2Node_CallFunction_1111 DD0B4644476E458B894265AF5D1FB6BA,),PersistentGuid=00000000000000000000000000000000,bHidden=False,bNotConnectable=False,bDefaultValueIsReadOnly=False,bDefaultValueIsIgnored=False,bAdvancedView=False,)
+			   CustomProperties Pin (PinId=DADFC7E349838F790BC569A556130F8B,PinName="Canvas",PinToolTip="Canvas\nCanvas Reference",Direction="EGPD_Output",PinType.PinCategory="object",PinType.PinSubCategoryObject=Class'/Script/Engine.Canvas',PinType.PinSubCategoryMemberReference=(),PinType.bIsArray=False,PinType.bIsReference=False,PinType.bIsConst=False,PinType.bIsWeakPointer=False,LinkedTo=(K2Node_CallFunction_249 8DDBD5C941067A40EC75EBBB4E6BCD74,),PersistentGuid=00000000000000000000000000000000,bHidden=False,bNotConnectable=False,bDefaultValueIsReadOnly=False,bDefaultValueIsIgnored=False,bAdvancedView=False,)
+			   CustomProperties Pin (PinId=2D8E4C1A462BB80E885F9D890A179F7D,PinName="Size",PinToolTip="Size\nVector 2D Structure",Direction="EGPD_Output",PinType.PinCategory="struct",PinType.PinSubCategoryObject=ScriptStruct'/Script/CoreUObject.Vector2D',PinType.PinSubCategoryMemberReference=(),PinType.bIsArray=False,PinType.bIsReference=False,PinType.bIsConst=False,PinType.bIsWeakPointer=False,LinkedTo=(K2Node_CallFunction_249 A27FA11345BBF4A4A5C25785C432EE8C,),PersistentGuid=00000000000000000000000000000000,bHidden=False,bNotConnectable=False,bDefaultValueIsReadOnly=False,bDefaultValueIsIgnored=False,bAdvancedView=False,)
+			   CustomProperties Pin (PinId=FDEC399C4812BB5D95E3FBA8B0133244,PinName="Context",PinToolTip="Context\nDraw To Render Target Context Structure",Direction="EGPD_Output",PinType.PinCategory="struct",PinType.PinSubCategoryObject=ScriptStruct'/Script/Engine.DrawToRenderTargetContext',PinType.PinSubCategoryMemberReference=(),PinType.bIsArray=False,PinType.bIsReference=False,PinType.bIsConst=False,PinType.bIsWeakPointer=False,LinkedTo=(K2Node_CallFunction_248 F45CDCA34036451D686B5396C1DA9D79,),PersistentGuid=00000000000000000000000000000000,bHidden=False,bNotConnectable=False,bDefaultValueIsReadOnly=False,bDefaultValueIsIgnored=False,bAdvancedView=False,)
+			End Object
+			Begin Object Class=K2Node_CallFunction Name="K2Node_CallFunction_149"
+			   FunctionReference=(MemberParent=Class'/Script/Engine.KismetRenderingLibrary',MemberName="DrawMaterialToRenderTarget")
+			   NodePosX=4112
+			   NodePosY=320
+			   NodeGuid=2C8CEDF14E8C77632F69DC9EC749DA2A
+			   CustomProperties Pin (PinId=27DAF40C4752D1A6FABA19B428E4B44D,PinName="execute",PinType.PinCategory="exec",PinType.PinSubCategoryObject=None,PinType.PinSubCategoryMemberReference=(),PinType.bIsArray=False,PinType.bIsReference=False,PinType.bIsConst=False,PinType.bIsWeakPointer=False,PersistentGuid=00000000000000000000000000000000,bHidden=False,bNotConnectable=False,bDefaultValueIsReadOnly=False,bDefaultValueIsIgnored=False,bAdvancedView=False,)
+			   CustomProperties Pin (PinId=75DE731B4E5CC0BDEBEAFC96E209BD72,PinName="then",Direction="EGPD_Output",PinType.PinCategory="exec",PinType.PinSubCategoryObject=None,PinType.PinSubCategoryMemberReference=(),PinType.bIsArray=False,PinType.bIsReference=False,PinType.bIsConst=False,PinType.bIsWeakPointer=False,PersistentGuid=00000000000000000000000000000000,bHidden=False,bNotConnectable=False,bDefaultValueIsReadOnly=False,bDefaultValueIsIgnored=False,bAdvancedView=False,)
+			   CustomProperties Pin (PinId=09D4564848B0149EA700A18181F1544B,PinName="self",PinFriendlyName=NSLOCTEXT("K2Node", "Target", "Target"),PinType.PinCategory="object",PinType.PinSubCategoryObject=Class'/Script/Engine.KismetRenderingLibrary',PinType.PinSubCategoryMemberReference=(),PinType.bIsArray=False,PinType.bIsReference=False,PinType.bIsConst=False,PinType.bIsWeakPointer=False,DefaultObject="/Script/Engine.Default__KismetRenderingLibrary",PersistentGuid=00000000000000000000000000000000,bHidden=True,bNotConnectable=False,bDefaultValueIsReadOnly=False,bDefaultValueIsIgnored=False,bAdvancedView=False,)
+			   CustomProperties Pin (PinId=D4512E6C4408F62942521CA73338B052,PinName="WorldContextObject",PinType.PinCategory="object",PinType.PinSubCategoryObject=Class'/Script/CoreUObject.Object',PinType.PinSubCategoryMemberReference=(),PinType.bIsArray=False,PinType.bIsReference=False,PinType.bIsConst=False,PinType.bIsWeakPointer=False,PersistentGuid=00000000000000000000000000000000,bHidden=True,bNotConnectable=False,bDefaultValueIsReadOnly=False,bDefaultValueIsIgnored=False,bAdvancedView=False,)
+			   CustomProperties Pin (PinId=4071F025450DAD7A0AF9929767345F09,PinName="TextureRenderTarget",PinType.PinCategory="object",PinType.PinSubCategoryObject=Class'/Script/Engine.TextureRenderTarget2D',PinType.PinSubCategoryMemberReference=(),PinType.bIsArray=False,PinType.bIsReference=False,PinType.bIsConst=False,PinType.bIsWeakPointer=False,LinkedTo=(K2Node_CallFunction_1111 DD0B4644476E458B894265AF5D1FB6BA,),PersistentGuid=00000000000000000000000000000000,bHidden=False,bNotConnectable=False,bDefaultValueIsReadOnly=False,bDefaultValueIsIgnored=False,bAdvancedView=False,)
+			   CustomProperties Pin (PinId=9E05610046120AA6FC4E36B6A1BA014F,PinName="Material",PinType.PinCategory="object",PinType.PinSubCategoryObject=Class'/Script/Engine.MaterialInterface',PinType.PinSubCategoryMemberReference=(),PinType.bIsArray=False,PinType.bIsReference=False,PinType.bIsConst=False,PinType.bIsWeakPointer=False,LinkedTo=(K2Node_CallMaterialParameterCollectionFunction_42 89FB58F843256CED989D709D84629586,),PersistentGuid=00000000000000000000000000000000,bHidden=False,bNotConnectable=False,bDefaultValueIsReadOnly=False,bDefaultValueIsIgnored=False,bAdvancedView=False,)
+			End Object
+			Begin Object Class=EdGraphNode_Comment Name="EdGraphNode_Comment_9"
+			   NodePosX=2320
+			   NodePosY=16
+			   NodeWidth=2704
+			   NodeHeight=512
+			   NodeComment="Apply fluid surface simulation kernel, which propagates waves"
+			   NodeGuid=47D40DCF403ED4951ED75E94C7AF3C5A
+			End Object
+			END OBJECT
+	[/PARAM]
+[/OBJECT]
+[/EXCERPT:ApplySim]
+-->
\ No newline at end of file
