Copying //UE4/Dev-Documentation to Samples-Main (//UE4/Samples-Main)

diff --git a/Engine/Documentation/Source/Engine/Rendering/RenderTargets/BlueprintRenderTargets/HowTo/FluidSurface/FluidSurface.INT.udn b/Engine/Documentation/Source/Engine/Rendering/RenderTargets/BlueprintRenderTargets/HowTo/FluidSurface/FluidSurface.INT.udn
new file mode 100644
--- /dev/null
+++ b/Engine/Documentation/Source/Engine/Rendering/RenderTargets/BlueprintRenderTargets/HowTo/FluidSurface/FluidSurface.INT.udn
@@ -0,0 +1,32 @@
+availability:docs
+title:Creating a Fluid surface with Blueprints and Render Targets
+description:How-To use Blueprints and Render Targets to create realtime fluid surface Actor.
+crumbs:%ROOT%, Engine, Engine\Rendering\RenderTargets
+version:4.13
+Type:how-to
+SkillLevel:Advanced
+parent:Engine/Rendering/RenderTargets/BlueprintRenderTargets/HowTo
+order:1
+Navigation:landing
+Related:Gameplay/Framework/GameMode/
+Related:Engine/Blueprints/
+Related:Engine/Rendering/Materials/
+Related:Engine/Blueprints/UserGuide/UserConstructionScript/
+tags:Render Targets
+tags:Blueprint
+checkpoint:FSB
+
+![](FS_Header_Image.png)
+
+
+Using Blueprints and Render Targets, you can create a fluid surface that simulates what happens to water when it is hit by a projectile or player. 
+In the following guide we will learn how you can use Blueprints & Render Targets to create a fluid surface actor by creating a new set of render targets and Materials that are specifically setup to work with render targets. We will then create and setup a new Blueprint that will contain and power the fluid surface interaction when the fluid surface is touched or shot at.
+
+[DIR(output:"checkpoint" parent:"Engine/Rendering/RenderTargets/BlueprintRenderTargets/HowTo/FluidSurface")]
+
+[region:note]
+Part of this guide uses the same assets created in the [](Engine\Rendering\RenderTargets\BlueprintRenderTargets\HowTo\HeightFieldPainter) How To guide. 
+If you have not completed the [](Engine\Rendering\RenderTargets\BlueprintRenderTargets\HowTo\HeightFieldPainter) How To guide, please do so before starting this one. You can also find this example map in the Content Example project that can be downloaded from the Learn tab of the Epic Games launcher.
+[/region]
+
+
diff --git a/Engine/Documentation/Source/Engine/Rendering/RenderTargets/BlueprintRenderTargets/HowTo/HeightFieldPainter/1/HFP1.INT.udn b/Engine/Documentation/Source/Engine/Rendering/RenderTargets/BlueprintRenderTargets/HowTo/HeightFieldPainter/1/HFP1.INT.udn
new file mode 100644
--- /dev/null
+++ b/Engine/Documentation/Source/Engine/Rendering/RenderTargets/BlueprintRenderTargets/HowTo/HeightFieldPainter/1/HFP1.INT.udn
@@ -0,0 +1,234 @@
+availability:docs
+title:1. Project Setup
+description:Going over the Project Setup for creating a Height Map painter.
+crumbs:%ROOT%, Engine, Engine\Rendering\RenderTargets
+version:4.13
+type:multi-step
+SkillLevel:Intermediate
+parent:Engine/Rendering/RenderTargets/BlueprintRenderTargets/HowTo/HeightFieldPainter
+Navigation:multi-step
+tags: Render Targets
+order:1
+checkpoint:HFP
+
+[NAV]
+
+In this section of the Height Map Painter How - To, we will go over creating a new level for our project to use.
+
+## Project Setup
+Before we can start to create and setup the required Actors for the Height Field Painter, we will first need to create and setup a **Gamemode Override**. In the following section, 
+we will take a look at how this can be accomplished in your project.
+
+[region:note]
+If you are working out of the **Content Examples** map, you can skip this step and use the **MyGame** game mode. This game mode can be found under **Content** > **Blueprints** > **MyGame**.
+
+ ![](Engine\Rendering\RenderTargets\BlueprintRenderTargets\HowTo\HeightFieldPainter\HFP_PS_00.png)
+[/region]
+
+## Creating a New Game Mode & HUD
+In order to see the crosshair that will let the user know where they are shooting, we will need to create and setup new Game Mode and HUD Blueprints. In the following section, we will take a 
+look at how to setup all of this so that it will work with our Height Field Painter.
+
+1. Right-click in the **Content Browser** and from the menu in the **Create Basic Asset** section, click on the **Blueprint Class** option.
+	
+	![](Engine\Rendering\RenderTargets\BlueprintRenderTargets\HowTo\HeightFieldPainter\HFP_BP_Setup_00.png)
+
+1. From the **Pick Parent Class** Window, click on the **Game Mode** option, naming the new game mode, `HighFieldPainterGameMode`.
+
+	![](Engine\Rendering\RenderTargets\BlueprintRenderTargets\HowTo\HeightFieldPainter\HFP_PS_01.png)
+
+1. Right-click in the Content Browser again to bring up the Create Basic Asset menu and then click on the **Blueprint Class** option.
+
+	![](Engine\Rendering\RenderTargets\BlueprintRenderTargets\HowTo\HeightFieldPainter\HFP_BP_Setup_09.png)
+
+1. When the Pick Parent Class window is displayed, expand the **All Classes** option and input **HUD** into the search box. Look for the **HUD** option and when located, click on it and then press the **Select** button to create it.
+
+	![](Engine\Rendering\RenderTargets\BlueprintRenderTargets\HowTo\HeightFieldPainter\HFP_PS_02.png)
+
+1. Name the newly created HUD Blueprint, `HeightFieldPainterHUD`, and then double-click on the **HeightFieldPainterGameMode** Blueprint to open it up.
+
+	![](Engine\Rendering\RenderTargets\BlueprintRenderTargets\HowTo\HeightFieldPainter\HFP_PS_04.png)
+
+1. In the **Details** panel of the HeightFieldPainterGameMode under the **Classes** section, set the **HUD Class** to the  `HeightFieldPainterHUD` Blueprint that was just created. 
+
+	![](Engine\Rendering\RenderTargets\BlueprintRenderTargets\HowTo\HeightFieldPainter\HFP_PS_03.png)
+
+1. Now, **Compile** and **Save** the HeightFieldPainterGameMode Blueprint and then open the HeightFieldPainterHUD Blueprint by double-clicking on it in the Content Browser.
+
+	![](Engine\Rendering\RenderTargets\BlueprintRenderTargets\HowTo\HeightFieldPainter\HFP_PS_05.png)
+
+1. Inside of the HeightFieldPainterHUD Blueprint, go to the **Event Graph** and copy-paste the following Blueprint code into the graph. 
+
+	[REGION:fullwidth]
+	[INCLUDE:Engine\Rendering\RenderTargets\BlueprintRenderTargets\HowTo\HeightFieldPainter\1\#HUD]
+	[/REGION]
+
+	[region:note]
+	Do not forget to Compile and Save your Blueprint so that it can be used.
+	[/region]
+
+	This Blueprint code is telling the HUD to draw the specified Material so that it will be centered on the player's screen while the game is running. To see if everything has been setup correctly, press the **Play** button to run the game and you should now have a crosshair positioned in the center of the screen (as shown in the following image).
+
+	![](Engine\Rendering\RenderTargets\BlueprintRenderTargets\HowTo\HeightFieldPainter\HFP_PS_07.png)
+
+	[region:note]
+	If you don't see a cross hair on the screen, check to make sure that the **Material** input on the **Draw Material** node has something inputted into the Material input.
+
+	![](Engine\Rendering\RenderTargets\BlueprintRenderTargets\HowTo\HeightFieldPainter\HFP_PS_08.png)
+	[/region]
+
+Now that the level and HUD have been created, in the next section we will take a look at what assets need to be created and how to set them up.
+
+[NAV]
+
+<!--
+[EXCERPT:HUD]	
+[OBJECT:Snippet]
+	[PARAMLITERAL:id]
+		EventGraph
+	[/PARAMLITERAL]
+	[PARAMLITERAL:type]
+		blueprint
+	[/PARAMLITERAL]
+	[PARAM:image]
+		![](Engine\Rendering\RenderTargets\BlueprintRenderTargets\HowTo\HeightFieldPainter\HFP_PS_06.png)
+	[/PARAM]
+	[PARAM:snippet]
+			BEGIN OBJECT
+							Begin Object Class=K2Node_Event Name="K2Node_Event_2653"
+			   EventSignatureName="ReceiveDrawHUD"
+			   EventSignatureClass=Class'/Script/Engine.HUD'
+			   EventReference=(MemberParent=Class'/Script/Engine.HUD',MemberName="ReceiveDrawHUD")
+			   bOverrideFunction=True
+			   NodePosX=-352
+			   NodePosY=64
+			   NodeGuid=2E96B8304D2BDF7906CA1A9640C22F3A
+			   CustomProperties Pin (PinId=29B867B649B6460F8DE08EAC4E51C8DB,PinName="OutputDelegate",Direction="EGPD_Output",PinType.PinCategory="delegate",PinType.PinSubCategoryObject=None,PinType.PinSubCategoryMemberReference=(MemberParent=Class'/Script/Engine.HUD',MemberName="ReceiveDrawHUD"),PinType.bIsArray=False,PinType.bIsReference=False,PinType.bIsConst=False,PinType.bIsWeakPointer=False,PersistentGuid=00000000000000000000000000000000,bHidden=False,bNotConnectable=False,bDefaultValueIsReadOnly=False,bDefaultValueIsIgnored=False,bAdvancedView=False,)
+			   CustomProperties Pin (PinId=5BE63FAB45228CD0D9EFCD9314C7819C,PinName="then",Direction="EGPD_Output",PinType.PinCategory="exec",PinType.PinSubCategoryObject=None,PinType.PinSubCategoryMemberReference=(),PinType.bIsArray=False,PinType.bIsReference=False,PinType.bIsConst=False,PinType.bIsWeakPointer=False,LinkedTo=(K2Node_CallFunction_8894 7E6038ED4BD8AF14885479BF2ED38699,),PersistentGuid=00000000000000000000000000000000,bHidden=False,bNotConnectable=False,bDefaultValueIsReadOnly=False,bDefaultValueIsIgnored=False,bAdvancedView=False,)
+			   CustomProperties Pin (PinId=C425E55F4C9DF22C6D8F08BD42F3BBEA,PinName="SizeX",PinToolTip="Size X\nInteger",Direction="EGPD_Output",PinType.PinCategory="int",PinType.PinSubCategoryObject=None,PinType.PinSubCategoryMemberReference=(),PinType.bIsArray=False,PinType.bIsReference=False,PinType.bIsConst=False,PinType.bIsWeakPointer=False,DefaultValue="0",AutogeneratedDefaultValue="0",LinkedTo=(K2Node_CallFunction_111158 627E2E4E4B1CC4EE5CD349A830442642,),PersistentGuid=00000000000000000000000000000000,bHidden=False,bNotConnectable=False,bDefaultValueIsReadOnly=False,bDefaultValueIsIgnored=False,bAdvancedView=False,)
+			   CustomProperties Pin (PinId=1848957E403CC954D80F3CA003B54942,PinName="SizeY",PinToolTip="Size Y\nInteger",Direction="EGPD_Output",PinType.PinCategory="int",PinType.PinSubCategoryObject=None,PinType.PinSubCategoryMemberReference=(),PinType.bIsArray=False,PinType.bIsReference=False,PinType.bIsConst=False,PinType.bIsWeakPointer=False,DefaultValue="0",AutogeneratedDefaultValue="0",LinkedTo=(K2Node_CallFunction_111157 B2B58B0B436A1FFB1E4F409A3C5D17AC,),PersistentGuid=00000000000000000000000000000000,bHidden=False,bNotConnectable=False,bDefaultValueIsReadOnly=False,bDefaultValueIsIgnored=False,bAdvancedView=False,)
+			End Object
+			Begin Object Class=K2Node_CallFunction Name="K2Node_CallFunction_111158"
+			   bIsPureFunc=True
+			   FunctionReference=(MemberParent=Class'/Script/Engine.KismetMathLibrary',MemberName="Conv_IntToFloat")
+			   NodePosX=80
+			   NodePosY=112
+			   NodeGuid=79FDF440407059C77A2CDCB0F8D56ED6
+			   CustomProperties Pin (PinId=9B84CC3A490B84C4CFA4E99A4DAC237E,PinName="self",PinFriendlyName=NSLOCTEXT("K2Node", "Target", "Target"),PinToolTip="Target\nKismet Math Library Reference",PinType.PinCategory="object",PinType.PinSubCategoryObject=Class'/Script/Engine.KismetMathLibrary',PinType.PinSubCategoryMemberReference=(),PinType.bIsArray=False,PinType.bIsReference=False,PinType.bIsConst=False,PinType.bIsWeakPointer=False,DefaultObject="/Script/Engine.Default__KismetMathLibrary",PersistentGuid=00000000000000000000000000000000,bHidden=True,bNotConnectable=False,bDefaultValueIsReadOnly=False,bDefaultValueIsIgnored=False,bAdvancedView=False,)
+			   CustomProperties Pin (PinId=627E2E4E4B1CC4EE5CD349A830442642,PinName="InInt",PinToolTip="In Int\nInteger",PinType.PinCategory="int",PinType.PinSubCategoryObject=None,PinType.PinSubCategoryMemberReference=(),PinType.bIsArray=False,PinType.bIsReference=False,PinType.bIsConst=False,PinType.bIsWeakPointer=False,DefaultValue="0",AutogeneratedDefaultValue="0",LinkedTo=(K2Node_Event_2653 C425E55F4C9DF22C6D8F08BD42F3BBEA,),PersistentGuid=00000000000000000000000000000000,bHidden=False,bNotConnectable=False,bDefaultValueIsReadOnly=False,bDefaultValueIsIgnored=False,bAdvancedView=False,)
+			   CustomProperties Pin (PinId=B73659A547602AD4811E6CAF0531666A,PinName="ReturnValue",PinToolTip="Return Value\nFloat",Direction="EGPD_Output",PinType.PinCategory="float",PinType.PinSubCategoryObject=None,PinType.PinSubCategoryMemberReference=(),PinType.bIsArray=False,PinType.bIsReference=False,PinType.bIsConst=False,PinType.bIsWeakPointer=False,DefaultValue="0.0",AutogeneratedDefaultValue="0.0",LinkedTo=(K2Node_CallFunction_111407 355102244077B448306118B37C6E116B,),PersistentGuid=00000000000000000000000000000000,bHidden=False,bNotConnectable=False,bDefaultValueIsReadOnly=False,bDefaultValueIsIgnored=False,bAdvancedView=False,)
+			End Object
+			Begin Object Class=K2Node_CallFunction Name="K2Node_CallFunction_111157"
+			   bIsPureFunc=True
+			   FunctionReference=(MemberParent=Class'/Script/Engine.KismetMathLibrary',MemberName="Conv_IntToFloat")
+			   NodePosX=64
+			   NodePosY=224
+			   NodeGuid=E11ED92D4DB1634A541E379DC15571FB
+			   CustomProperties Pin (PinId=FF48802B423B9360DAB437ADC0847AA0,PinName="self",PinFriendlyName=NSLOCTEXT("K2Node", "Target", "Target"),PinType.PinCategory="object",PinType.PinSubCategoryObject=Class'/Script/Engine.KismetMathLibrary',PinType.PinSubCategoryMemberReference=(),PinType.bIsArray=False,PinType.bIsReference=False,PinType.bIsConst=False,PinType.bIsWeakPointer=False,DefaultObject="/Script/Engine.Default__KismetMathLibrary",PersistentGuid=00000000000000000000000000000000,bHidden=True,bNotConnectable=False,bDefaultValueIsReadOnly=False,bDefaultValueIsIgnored=False,bAdvancedView=False,)
+			   CustomProperties Pin (PinId=B2B58B0B436A1FFB1E4F409A3C5D17AC,PinName="InInt",PinType.PinCategory="int",PinType.PinSubCategoryObject=None,PinType.PinSubCategoryMemberReference=(),PinType.bIsArray=False,PinType.bIsReference=False,PinType.bIsConst=False,PinType.bIsWeakPointer=False,DefaultValue="0",AutogeneratedDefaultValue="0",LinkedTo=(K2Node_Event_2653 1848957E403CC954D80F3CA003B54942,),PersistentGuid=00000000000000000000000000000000,bHidden=False,bNotConnectable=False,bDefaultValueIsReadOnly=False,bDefaultValueIsIgnored=False,bAdvancedView=False,)
+			   CustomProperties Pin (PinId=110976EA4E0B8B1D740DB2A51BF58D43,PinName="ReturnValue",Direction="EGPD_Output",PinType.PinCategory="float",PinType.PinSubCategoryObject=None,PinType.PinSubCategoryMemberReference=(),PinType.bIsArray=False,PinType.bIsReference=False,PinType.bIsConst=False,PinType.bIsWeakPointer=False,DefaultValue="0.0",AutogeneratedDefaultValue="0.0",LinkedTo=(K2Node_CallFunction_111406 0DCE1DA84C900E6C3A5FD6BDAA87A65B,),PersistentGuid=00000000000000000000000000000000,bHidden=False,bNotConnectable=False,bDefaultValueIsReadOnly=False,bDefaultValueIsIgnored=False,bAdvancedView=False,)
+			End Object
+			Begin Object Class=K2Node_CallFunction Name="K2Node_CallFunction_111407"
+			   bIsPureFunc=True
+			   FunctionReference=(MemberParent=Class'/Script/Engine.KismetMathLibrary',MemberName="Divide_FloatFloat")
+			   NodePosX=224
+			   NodePosY=128
+			   NodeGuid=EE86D04D427547DD7E9109A4832A3A33
+			   CustomProperties Pin (PinId=330D22EE415FABDCB4E20F8529BBCDA0,PinName="self",PinFriendlyName=NSLOCTEXT("K2Node", "Target", "Target"),PinToolTip="Target\nKismet Math Library Reference",PinType.PinCategory="object",PinType.PinSubCategoryObject=Class'/Script/Engine.KismetMathLibrary',PinType.PinSubCategoryMemberReference=(),PinType.bIsArray=False,PinType.bIsReference=False,PinType.bIsConst=False,PinType.bIsWeakPointer=False,DefaultObject="/Script/Engine.Default__KismetMathLibrary",PersistentGuid=00000000000000000000000000000000,bHidden=True,bNotConnectable=False,bDefaultValueIsReadOnly=False,bDefaultValueIsIgnored=False,bAdvancedView=False,)
+			   CustomProperties Pin (PinId=355102244077B448306118B37C6E116B,PinName="A",PinToolTip="A\nFloat",PinType.PinCategory="float",PinType.PinSubCategoryObject=None,PinType.PinSubCategoryMemberReference=(),PinType.bIsArray=False,PinType.bIsReference=False,PinType.bIsConst=False,PinType.bIsWeakPointer=False,DefaultValue="0.0",AutogeneratedDefaultValue="0.0",LinkedTo=(K2Node_CallFunction_111158 B73659A547602AD4811E6CAF0531666A,),PersistentGuid=00000000000000000000000000000000,bHidden=False,bNotConnectable=False,bDefaultValueIsReadOnly=False,bDefaultValueIsIgnored=False,bAdvancedView=False,)
+			   CustomProperties Pin (PinId=DA4662DC439C9F8BDD142E9752016CC3,PinName="B",PinToolTip="B\nFloat",PinType.PinCategory="float",PinType.PinSubCategoryObject=None,PinType.PinSubCategoryMemberReference=(),PinType.bIsArray=False,PinType.bIsReference=False,PinType.bIsConst=False,PinType.bIsWeakPointer=False,DefaultValue="2",AutogeneratedDefaultValue="1.000000",PersistentGuid=00000000000000000000000000000000,bHidden=False,bNotConnectable=False,bDefaultValueIsReadOnly=False,bDefaultValueIsIgnored=False,bAdvancedView=False,)
+			   CustomProperties Pin (PinId=83AC9C434C10BDD9B9E095974665DF30,PinName="ReturnValue",PinToolTip="Return Value\nFloat",Direction="EGPD_Output",PinType.PinCategory="float",PinType.PinSubCategoryObject=None,PinType.PinSubCategoryMemberReference=(),PinType.bIsArray=False,PinType.bIsReference=False,PinType.bIsConst=False,PinType.bIsWeakPointer=False,DefaultValue="0.0",AutogeneratedDefaultValue="0.0",LinkedTo=(K2Node_CallFunction_111971 D985192E4047CD01E7CF9BBEB73729F2,),PersistentGuid=00000000000000000000000000000000,bHidden=False,bNotConnectable=False,bDefaultValueIsReadOnly=False,bDefaultValueIsIgnored=False,bAdvancedView=False,)
+			End Object
+			Begin Object Class=K2Node_CallFunction Name="K2Node_CallFunction_111406"
+			   bIsPureFunc=True
+			   FunctionReference=(MemberParent=Class'/Script/Engine.KismetMathLibrary',MemberName="Divide_FloatFloat")
+			   NodePosX=224
+			   NodePosY=224
+			   NodeGuid=05C35D2B42F22BF65547F6B8C8E3BECB
+			   CustomProperties Pin (PinId=13469FB94E7B24779760418AE5A22481,PinName="self",PinFriendlyName=NSLOCTEXT("K2Node", "Target", "Target"),PinToolTip="Target\nKismet Math Library Reference",PinType.PinCategory="object",PinType.PinSubCategoryObject=Class'/Script/Engine.KismetMathLibrary',PinType.PinSubCategoryMemberReference=(),PinType.bIsArray=False,PinType.bIsReference=False,PinType.bIsConst=False,PinType.bIsWeakPointer=False,DefaultObject="/Script/Engine.Default__KismetMathLibrary",PersistentGuid=00000000000000000000000000000000,bHidden=True,bNotConnectable=False,bDefaultValueIsReadOnly=False,bDefaultValueIsIgnored=False,bAdvancedView=False,)
+			   CustomProperties Pin (PinId=0DCE1DA84C900E6C3A5FD6BDAA87A65B,PinName="A",PinToolTip="A\nFloat",PinType.PinCategory="float",PinType.PinSubCategoryObject=None,PinType.PinSubCategoryMemberReference=(),PinType.bIsArray=False,PinType.bIsReference=False,PinType.bIsConst=False,PinType.bIsWeakPointer=False,DefaultValue="0.0",AutogeneratedDefaultValue="0.0",LinkedTo=(K2Node_CallFunction_111157 110976EA4E0B8B1D740DB2A51BF58D43,),PersistentGuid=00000000000000000000000000000000,bHidden=False,bNotConnectable=False,bDefaultValueIsReadOnly=False,bDefaultValueIsIgnored=False,bAdvancedView=False,)
+			   CustomProperties Pin (PinId=AA256AC64DDA7D92F926EE8CD36423C9,PinName="B",PinToolTip="B\nFloat",PinType.PinCategory="float",PinType.PinSubCategoryObject=None,PinType.PinSubCategoryMemberReference=(),PinType.bIsArray=False,PinType.bIsReference=False,PinType.bIsConst=False,PinType.bIsWeakPointer=False,DefaultValue="2",AutogeneratedDefaultValue="1.000000",PersistentGuid=00000000000000000000000000000000,bHidden=False,bNotConnectable=False,bDefaultValueIsReadOnly=False,bDefaultValueIsIgnored=False,bAdvancedView=False,)
+			   CustomProperties Pin (PinId=A7A941F34A6B8078A6C799AA18A61AFE,PinName="ReturnValue",PinToolTip="Return Value\nFloat",Direction="EGPD_Output",PinType.PinCategory="float",PinType.PinSubCategoryObject=None,PinType.PinSubCategoryMemberReference=(),PinType.bIsArray=False,PinType.bIsReference=False,PinType.bIsConst=False,PinType.bIsWeakPointer=False,DefaultValue="0.0",AutogeneratedDefaultValue="0.0",LinkedTo=(K2Node_CallFunction_111970 16F22B7B46746B48238ECC9E4CDDF565,),PersistentGuid=00000000000000000000000000000000,bHidden=False,bNotConnectable=False,bDefaultValueIsReadOnly=False,bDefaultValueIsIgnored=False,bAdvancedView=False,)
+			End Object
+			Begin Object Class=K2Node_CallFunction Name="K2Node_CallFunction_111475"
+			   bIsPureFunc=True
+			   FunctionReference=(MemberParent=Class'/Script/Engine.KismetSystemLibrary',MemberName="MakeLiteralFloat")
+			   NodePosX=-16
+			   NodePosY=416
+			   NodeGuid=E08EC2CB4433A83555A952BF162CFDD0
+			   CustomProperties Pin (PinId=ACED869B4C66021D2E475C993E750705,PinName="self",PinFriendlyName=NSLOCTEXT("K2Node", "Target", "Target"),PinType.PinCategory="object",PinType.PinSubCategoryObject=Class'/Script/Engine.KismetSystemLibrary',PinType.PinSubCategoryMemberReference=(),PinType.bIsArray=False,PinType.bIsReference=False,PinType.bIsConst=False,PinType.bIsWeakPointer=False,DefaultObject="/Script/Engine.Default__KismetSystemLibrary",PersistentGuid=00000000000000000000000000000000,bHidden=True,bNotConnectable=False,bDefaultValueIsReadOnly=False,bDefaultValueIsIgnored=False,bAdvancedView=False,)
+			   CustomProperties Pin (PinId=98B586C245B17C9069C66DB4561ACBB9,PinName="Value",PinType.PinCategory="float",PinType.PinSubCategoryObject=None,PinType.PinSubCategoryMemberReference=(),PinType.bIsArray=False,PinType.bIsReference=False,PinType.bIsConst=False,PinType.bIsWeakPointer=False,DefaultValue="40",AutogeneratedDefaultValue="0.0",PersistentGuid=00000000000000000000000000000000,bHidden=False,bNotConnectable=False,bDefaultValueIsReadOnly=False,bDefaultValueIsIgnored=False,bAdvancedView=False,)
+			   CustomProperties Pin (PinId=C1F2DEB54EFA6E0B933D2AA78836CFE7,PinName="ReturnValue",Direction="EGPD_Output",PinType.PinCategory="float",PinType.PinSubCategoryObject=None,PinType.PinSubCategoryMemberReference=(),PinType.bIsArray=False,PinType.bIsReference=False,PinType.bIsConst=False,PinType.bIsWeakPointer=False,DefaultValue="0.0",AutogeneratedDefaultValue="0.0",LinkedTo=(K2Node_CallFunction_111722 7C748362490BFEA846B7639E34BDE195,K2Node_CallFunction_8894 BAD37ED4468872440CA7528D6A91FA3D,K2Node_CallFunction_8894 D65908D74BBEF192C5A5D8B59CBF9732,),PersistentGuid=00000000000000000000000000000000,bHidden=False,bNotConnectable=False,bDefaultValueIsReadOnly=False,bDefaultValueIsIgnored=False,bAdvancedView=False,)
+			End Object
+			Begin Object Class=K2Node_CallFunction Name="K2Node_CallFunction_111722"
+			   bIsPureFunc=True
+			   FunctionReference=(MemberParent=Class'/Script/Engine.KismetMathLibrary',MemberName="Divide_FloatFloat")
+			   NodePosX=272
+			   NodePosY=304
+			   NodeGuid=30A0A89F45D4FC5F7DB5F694F31E302A
+			   CustomProperties Pin (PinId=CB09CD944A39C1F2C1A37ABB31206BC2,PinName="self",PinFriendlyName=NSLOCTEXT("K2Node", "Target", "Target"),PinToolTip="Target\nKismet Math Library Reference",PinType.PinCategory="object",PinType.PinSubCategoryObject=Class'/Script/Engine.KismetMathLibrary',PinType.PinSubCategoryMemberReference=(),PinType.bIsArray=False,PinType.bIsReference=False,PinType.bIsConst=False,PinType.bIsWeakPointer=False,DefaultObject="/Script/Engine.Default__KismetMathLibrary",PersistentGuid=00000000000000000000000000000000,bHidden=True,bNotConnectable=False,bDefaultValueIsReadOnly=False,bDefaultValueIsIgnored=False,bAdvancedView=False,)
+			   CustomProperties Pin (PinId=7C748362490BFEA846B7639E34BDE195,PinName="A",PinToolTip="A\nFloat",PinType.PinCategory="float",PinType.PinSubCategoryObject=None,PinType.PinSubCategoryMemberReference=(),PinType.bIsArray=False,PinType.bIsReference=False,PinType.bIsConst=False,PinType.bIsWeakPointer=False,DefaultValue="0.0",AutogeneratedDefaultValue="0.0",LinkedTo=(K2Node_CallFunction_111475 C1F2DEB54EFA6E0B933D2AA78836CFE7,),PersistentGuid=00000000000000000000000000000000,bHidden=False,bNotConnectable=False,bDefaultValueIsReadOnly=False,bDefaultValueIsIgnored=False,bAdvancedView=False,)
+			   CustomProperties Pin (PinId=4BF94B924C25552B3BF4039BC14774D5,PinName="B",PinToolTip="B\nFloat",PinType.PinCategory="float",PinType.PinSubCategoryObject=None,PinType.PinSubCategoryMemberReference=(),PinType.bIsArray=False,PinType.bIsReference=False,PinType.bIsConst=False,PinType.bIsWeakPointer=False,DefaultValue="2",AutogeneratedDefaultValue="1.000000",PersistentGuid=00000000000000000000000000000000,bHidden=False,bNotConnectable=False,bDefaultValueIsReadOnly=False,bDefaultValueIsIgnored=False,bAdvancedView=False,)
+			   CustomProperties Pin (PinId=ECDE39F64E7E3FC8A1C48EA6AF9400E1,PinName="ReturnValue",PinToolTip="Return Value\nFloat",Direction="EGPD_Output",PinType.PinCategory="float",PinType.PinSubCategoryObject=None,PinType.PinSubCategoryMemberReference=(),PinType.bIsArray=False,PinType.bIsReference=False,PinType.bIsConst=False,PinType.bIsWeakPointer=False,DefaultValue="0.0",AutogeneratedDefaultValue="0.0",LinkedTo=(K2Node_CallFunction_111971 035AEABF4A21C8873545BF8BF45FF9B4,K2Node_CallFunction_111970 0F6387B44A0A26FA803463A6D154ABEF,),PersistentGuid=00000000000000000000000000000000,bHidden=False,bNotConnectable=False,bDefaultValueIsReadOnly=False,bDefaultValueIsIgnored=False,bAdvancedView=False,)
+			End Object
+			Begin Object Class=K2Node_CallFunction Name="K2Node_CallFunction_111971"
+			   bIsPureFunc=True
+			   FunctionReference=(MemberParent=Class'/Script/Engine.KismetMathLibrary',MemberName="Subtract_FloatFloat")
+			   NodePosX=448
+			   NodePosY=128
+			   NodeGuid=C46DBD5541FD31FE49DCB190A50403E2
+			   CustomProperties Pin (PinId=2ADC1B0143A84E3B145BDDAD0F356D13,PinName="self",PinFriendlyName=NSLOCTEXT("K2Node", "Target", "Target"),PinType.PinCategory="object",PinType.PinSubCategoryObject=Class'/Script/Engine.KismetMathLibrary',PinType.PinSubCategoryMemberReference=(),PinType.bIsArray=False,PinType.bIsReference=False,PinType.bIsConst=False,PinType.bIsWeakPointer=False,DefaultObject="/Script/Engine.Default__KismetMathLibrary",PersistentGuid=00000000000000000000000000000000,bHidden=True,bNotConnectable=False,bDefaultValueIsReadOnly=False,bDefaultValueIsIgnored=False,bAdvancedView=False,)
+			   CustomProperties Pin (PinId=D985192E4047CD01E7CF9BBEB73729F2,PinName="A",PinType.PinCategory="float",PinType.PinSubCategoryObject=None,PinType.PinSubCategoryMemberReference=(),PinType.bIsArray=False,PinType.bIsReference=False,PinType.bIsConst=False,PinType.bIsWeakPointer=False,DefaultValue="0.0",AutogeneratedDefaultValue="0.0",LinkedTo=(K2Node_CallFunction_111407 83AC9C434C10BDD9B9E095974665DF30,),PersistentGuid=00000000000000000000000000000000,bHidden=False,bNotConnectable=False,bDefaultValueIsReadOnly=False,bDefaultValueIsIgnored=False,bAdvancedView=False,)
+			   CustomProperties Pin (PinId=035AEABF4A21C8873545BF8BF45FF9B4,PinName="B",PinType.PinCategory="float",PinType.PinSubCategoryObject=None,PinType.PinSubCategoryMemberReference=(),PinType.bIsArray=False,PinType.bIsReference=False,PinType.bIsConst=False,PinType.bIsWeakPointer=False,DefaultValue="0.0",AutogeneratedDefaultValue="1.000000",LinkedTo=(K2Node_CallFunction_111722 ECDE39F64E7E3FC8A1C48EA6AF9400E1,),PersistentGuid=00000000000000000000000000000000,bHidden=False,bNotConnectable=False,bDefaultValueIsReadOnly=False,bDefaultValueIsIgnored=False,bAdvancedView=False,)
+			   CustomProperties Pin (PinId=858D219A4D61743D711214956DCCD88F,PinName="ReturnValue",Direction="EGPD_Output",PinType.PinCategory="float",PinType.PinSubCategoryObject=None,PinType.PinSubCategoryMemberReference=(),PinType.bIsArray=False,PinType.bIsReference=False,PinType.bIsConst=False,PinType.bIsWeakPointer=False,DefaultValue="0.0",AutogeneratedDefaultValue="0.0",LinkedTo=(K2Node_CallFunction_8894 16CE787846EE203915CF0AA7CA317397,),PersistentGuid=00000000000000000000000000000000,bHidden=False,bNotConnectable=False,bDefaultValueIsReadOnly=False,bDefaultValueIsIgnored=False,bAdvancedView=False,)
+			End Object
+			Begin Object Class=K2Node_CallFunction Name="K2Node_CallFunction_111970"
+			   bIsPureFunc=True
+			   FunctionReference=(MemberParent=Class'/Script/Engine.KismetMathLibrary',MemberName="Subtract_FloatFloat")
+			   NodePosX=448
+			   NodePosY=224
+			   NodeGuid=69EE10B947680EEB45E613B74B5DF667
+			   CustomProperties Pin (PinId=DA4749E945263660E49F1E8ECC8982A0,PinName="self",PinFriendlyName=NSLOCTEXT("K2Node", "Target", "Target"),PinType.PinCategory="object",PinType.PinSubCategoryObject=Class'/Script/Engine.KismetMathLibrary',PinType.PinSubCategoryMemberReference=(),PinType.bIsArray=False,PinType.bIsReference=False,PinType.bIsConst=False,PinType.bIsWeakPointer=False,DefaultObject="/Script/Engine.Default__KismetMathLibrary",PersistentGuid=00000000000000000000000000000000,bHidden=True,bNotConnectable=False,bDefaultValueIsReadOnly=False,bDefaultValueIsIgnored=False,bAdvancedView=False,)
+			   CustomProperties Pin (PinId=16F22B7B46746B48238ECC9E4CDDF565,PinName="A",PinType.PinCategory="float",PinType.PinSubCategoryObject=None,PinType.PinSubCategoryMemberReference=(),PinType.bIsArray=False,PinType.bIsReference=False,PinType.bIsConst=False,PinType.bIsWeakPointer=False,DefaultValue="0.0",AutogeneratedDefaultValue="0.0",LinkedTo=(K2Node_CallFunction_111406 A7A941F34A6B8078A6C799AA18A61AFE,),PersistentGuid=00000000000000000000000000000000,bHidden=False,bNotConnectable=False,bDefaultValueIsReadOnly=False,bDefaultValueIsIgnored=False,bAdvancedView=False,)
+			   CustomProperties Pin (PinId=0F6387B44A0A26FA803463A6D154ABEF,PinName="B",PinType.PinCategory="float",PinType.PinSubCategoryObject=None,PinType.PinSubCategoryMemberReference=(),PinType.bIsArray=False,PinType.bIsReference=False,PinType.bIsConst=False,PinType.bIsWeakPointer=False,DefaultValue="0.0",AutogeneratedDefaultValue="1.000000",LinkedTo=(K2Node_CallFunction_111722 ECDE39F64E7E3FC8A1C48EA6AF9400E1,),PersistentGuid=00000000000000000000000000000000,bHidden=False,bNotConnectable=False,bDefaultValueIsReadOnly=False,bDefaultValueIsIgnored=False,bAdvancedView=False,)
+			   CustomProperties Pin (PinId=CCA31F8E46CCCA7AA0A1A2B4C22EDC96,PinName="ReturnValue",Direction="EGPD_Output",PinType.PinCategory="float",PinType.PinSubCategoryObject=None,PinType.PinSubCategoryMemberReference=(),PinType.bIsArray=False,PinType.bIsReference=False,PinType.bIsConst=False,PinType.bIsWeakPointer=False,DefaultValue="0.0",AutogeneratedDefaultValue="0.0",LinkedTo=(K2Node_CallFunction_8894 54AFCEB54A488D37A6A81C826877C41A,),PersistentGuid=00000000000000000000000000000000,bHidden=False,bNotConnectable=False,bDefaultValueIsReadOnly=False,bDefaultValueIsIgnored=False,bAdvancedView=False,)
+			End Object
+			Begin Object Class=K2Node_CallFunction Name="K2Node_CallFunction_8894"
+			   FunctionReference=(MemberName="DrawMaterial",bSelfContext=True)
+			   NodePosX=784
+			   NodePosY=64
+			   NodeGuid=24539BC14B99E6616938DAA5B996311C
+			   AdvancedPinDisplay=Shown
+			   CustomProperties Pin (PinId=7E6038ED4BD8AF14885479BF2ED38699,PinName="execute",PinToolTip="\nExec",PinType.PinCategory="exec",PinType.PinSubCategoryObject=None,PinType.PinSubCategoryMemberReference=(),PinType.bIsArray=False,PinType.bIsReference=False,PinType.bIsConst=False,PinType.bIsWeakPointer=False,LinkedTo=(K2Node_Event_2653 5BE63FAB45228CD0D9EFCD9314C7819C,),PersistentGuid=00000000000000000000000000000000,bHidden=False,bNotConnectable=False,bDefaultValueIsReadOnly=False,bDefaultValueIsIgnored=False,bAdvancedView=False,)
+			   CustomProperties Pin (PinId=B780D4EC48E6762018EBC79F7B0BACA8,PinName="then",PinToolTip="\nExec",Direction="EGPD_Output",PinType.PinCategory="exec",PinType.PinSubCategoryObject=None,PinType.PinSubCategoryMemberReference=(),PinType.bIsArray=False,PinType.bIsReference=False,PinType.bIsConst=False,PinType.bIsWeakPointer=False,PersistentGuid=00000000000000000000000000000000,bHidden=False,bNotConnectable=False,bDefaultValueIsReadOnly=False,bDefaultValueIsIgnored=False,bAdvancedView=False,)
+			   CustomProperties Pin (PinId=EE6AB59E499723B293DD1FBCD6C81DFE,PinName="self",PinFriendlyName=NSLOCTEXT("K2Node", "Target", "Target"),PinToolTip="Target\nHUD Reference",PinType.PinCategory="object",PinType.PinSubCategoryObject=Class'/Script/Engine.HUD',PinType.PinSubCategoryMemberReference=(),PinType.bIsArray=False,PinType.bIsReference=False,PinType.bIsConst=False,PinType.bIsWeakPointer=False,PersistentGuid=00000000000000000000000000000000,bHidden=False,bNotConnectable=False,bDefaultValueIsReadOnly=False,bDefaultValueIsIgnored=False,bAdvancedView=False,)
+			   CustomProperties Pin (PinId=24DC1F6C47AA5539A6D6A99BFF1E3CF2,PinName="Material",PinToolTip="Material\nMaterial Interface Reference\n\nMaterial to use",PinType.PinCategory="object",PinType.PinSubCategoryObject=Class'/Script/Engine.MaterialInterface',PinType.PinSubCategoryMemberReference=(),PinType.bIsArray=False,PinType.bIsReference=False,PinType.bIsConst=False,PinType.bIsWeakPointer=False,DefaultObject="/Game/ExampleContent/Blueprint_Communication/Materials/M_Crosshair.M_Crosshair",PersistentGuid=00000000000000000000000000000000,bHidden=False,bNotConnectable=False,bDefaultValueIsReadOnly=False,bDefaultValueIsIgnored=False,bAdvancedView=False,)
+			   CustomProperties Pin (PinId=16CE787846EE203915CF0AA7CA317397,PinName="ScreenX",PinToolTip="Screen X\nFloat\n\nScreen-space X coordinate of upper left corner of the quad.",PinType.PinCategory="float",PinType.PinSubCategoryObject=None,PinType.PinSubCategoryMemberReference=(),PinType.bIsArray=False,PinType.bIsReference=False,PinType.bIsConst=False,PinType.bIsWeakPointer=False,DefaultValue="0.0",AutogeneratedDefaultValue="0.0",LinkedTo=(K2Node_CallFunction_111971 858D219A4D61743D711214956DCCD88F,),PersistentGuid=00000000000000000000000000000000,bHidden=False,bNotConnectable=False,bDefaultValueIsReadOnly=False,bDefaultValueIsIgnored=False,bAdvancedView=False,)
+			   CustomProperties Pin (PinId=54AFCEB54A488D37A6A81C826877C41A,PinName="ScreenY",PinToolTip="Screen Y\nFloat\n\nScreen-space Y coordinate of upper left corner of the quad.",PinType.PinCategory="float",PinType.PinSubCategoryObject=None,PinType.PinSubCategoryMemberReference=(),PinType.bIsArray=False,PinType.bIsReference=False,PinType.bIsConst=False,PinType.bIsWeakPointer=False,DefaultValue="0.0",AutogeneratedDefaultValue="0.0",LinkedTo=(K2Node_CallFunction_111970 CCA31F8E46CCCA7AA0A1A2B4C22EDC96,),PersistentGuid=00000000000000000000000000000000,bHidden=False,bNotConnectable=False,bDefaultValueIsReadOnly=False,bDefaultValueIsIgnored=False,bAdvancedView=False,)
+			   CustomProperties Pin (PinId=BAD37ED4468872440CA7528D6A91FA3D,PinName="ScreenW",PinToolTip="Screen W\nFloat\n\nScreen-space width of the quad (in pixels).",PinType.PinCategory="float",PinType.PinSubCategoryObject=None,PinType.PinSubCategoryMemberReference=(),PinType.bIsArray=False,PinType.bIsReference=False,PinType.bIsConst=False,PinType.bIsWeakPointer=False,DefaultValue="0.0",AutogeneratedDefaultValue="0.0",LinkedTo=(K2Node_CallFunction_111475 C1F2DEB54EFA6E0B933D2AA78836CFE7,),PersistentGuid=00000000000000000000000000000000,bHidden=False,bNotConnectable=False,bDefaultValueIsReadOnly=False,bDefaultValueIsIgnored=False,bAdvancedView=False,)
+			   CustomProperties Pin (PinId=D65908D74BBEF192C5A5D8B59CBF9732,PinName="ScreenH",PinToolTip="Screen H\nFloat\n\nScreen-space height of the quad (in pixels).",PinType.PinCategory="float",PinType.PinSubCategoryObject=None,PinType.PinSubCategoryMemberReference=(),PinType.bIsArray=False,PinType.bIsReference=False,PinType.bIsConst=False,PinType.bIsWeakPointer=False,DefaultValue="0.0",AutogeneratedDefaultValue="0.0",LinkedTo=(K2Node_CallFunction_111475 C1F2DEB54EFA6E0B933D2AA78836CFE7,),PersistentGuid=00000000000000000000000000000000,bHidden=False,bNotConnectable=False,bDefaultValueIsReadOnly=False,bDefaultValueIsIgnored=False,bAdvancedView=False,)
+			   CustomProperties Pin (PinId=B24FDC42460DE9223DF28F870D98DE6F,PinName="MaterialU",PinToolTip="Material U\nFloat\n\nTexture-space U coordinate of upper left corner of the quad",PinType.PinCategory="float",PinType.PinSubCategoryObject=None,PinType.PinSubCategoryMemberReference=(),PinType.bIsArray=False,PinType.bIsReference=False,PinType.bIsConst=False,PinType.bIsWeakPointer=False,DefaultValue="0.0",AutogeneratedDefaultValue="0.0",PersistentGuid=00000000000000000000000000000000,bHidden=False,bNotConnectable=False,bDefaultValueIsReadOnly=False,bDefaultValueIsIgnored=False,bAdvancedView=False,)
+			   CustomProperties Pin (PinId=3670D2524767A060171206AD0028FBD1,PinName="MaterialV",PinToolTip="Material V\nFloat\n\nTexture-space V coordinate of upper left corner of the quad.",PinType.PinCategory="float",PinType.PinSubCategoryObject=None,PinType.PinSubCategoryMemberReference=(),PinType.bIsArray=False,PinType.bIsReference=False,PinType.bIsConst=False,PinType.bIsWeakPointer=False,DefaultValue="0.0",AutogeneratedDefaultValue="0.0",PersistentGuid=00000000000000000000000000000000,bHidden=False,bNotConnectable=False,bDefaultValueIsReadOnly=False,bDefaultValueIsIgnored=False,bAdvancedView=False,)
+			   CustomProperties Pin (PinId=4AB87A914F9D4B1DCAFFE9B5324A94AA,PinName="MaterialUWidth",PinToolTip="Material UWidth\nFloat\n\nTexture-space width of the quad (in normalized UV distance).",PinType.PinCategory="float",PinType.PinSubCategoryObject=None,PinType.PinSubCategoryMemberReference=(),PinType.bIsArray=False,PinType.bIsReference=False,PinType.bIsConst=False,PinType.bIsWeakPointer=False,DefaultValue="1",AutogeneratedDefaultValue="0.0",PersistentGuid=00000000000000000000000000000000,bHidden=False,bNotConnectable=False,bDefaultValueIsReadOnly=False,bDefaultValueIsIgnored=False,bAdvancedView=False,)
+			   CustomProperties Pin (PinId=BBB3D8DE4EF0233940C9A89D7AA8E4E5,PinName="MaterialVHeight",PinToolTip="Material VHeight\nFloat\n\nTexture-space height of the quad (in normalized UV distance).",PinType.PinCategory="float",PinType.PinSubCategoryObject=None,PinType.PinSubCategoryMemberReference=(),PinType.bIsArray=False,PinType.bIsReference=False,PinType.bIsConst=False,PinType.bIsWeakPointer=False,DefaultValue="1",AutogeneratedDefaultValue="0.0",PersistentGuid=00000000000000000000000000000000,bHidden=False,bNotConnectable=False,bDefaultValueIsReadOnly=False,bDefaultValueIsIgnored=False,bAdvancedView=False,)
+			   CustomProperties Pin (PinId=52A0D91C4F8A74D2F4D412ABA962CBE5,PinName="Scale",PinToolTip="Scale\nFloat\n\nAmount to scale the entire texture (horizontally and vertically)",PinType.PinCategory="float",PinType.PinSubCategoryObject=None,PinType.PinSubCategoryMemberReference=(),PinType.bIsArray=False,PinType.bIsReference=False,PinType.bIsConst=False,PinType.bIsWeakPointer=False,DefaultValue="1.000000",AutogeneratedDefaultValue="1.000000",PersistentGuid=00000000000000000000000000000000,bHidden=False,bNotConnectable=False,bDefaultValueIsReadOnly=False,bDefaultValueIsIgnored=False,bAdvancedView=True,)
+			   CustomProperties Pin (PinId=7FC80E864659A1C000DC478DD0609EBD,PinName="bScalePosition",PinToolTip="Scale Position\nBoolean\n\nWhether the \"Scale\" parameter should also scale the position of this draw call.",PinType.PinCategory="bool",PinType.PinSubCategoryObject=None,PinType.PinSubCategoryMemberReference=(),PinType.bIsArray=False,PinType.bIsReference=False,PinType.bIsConst=False,PinType.bIsWeakPointer=False,DefaultValue="false",AutogeneratedDefaultValue="false",PersistentGuid=00000000000000000000000000000000,bHidden=False,bNotConnectable=False,bDefaultValueIsReadOnly=False,bDefaultValueIsIgnored=False,bAdvancedView=True,)
+			   CustomProperties Pin (PinId=4CC5D85949E3CCAC770F2D8F6B17C169,PinName="Rotation",PinToolTip="Rotation\nFloat\n\nAmount to rotate this quad",PinType.PinCategory="float",PinType.PinSubCategoryObject=None,PinType.PinSubCategoryMemberReference=(),PinType.bIsArray=False,PinType.bIsReference=False,PinType.bIsConst=False,PinType.bIsWeakPointer=False,DefaultValue="0.000000",AutogeneratedDefaultValue="0.000000",PersistentGuid=00000000000000000000000000000000,bHidden=False,bNotConnectable=False,bDefaultValueIsReadOnly=False,bDefaultValueIsIgnored=False,bAdvancedView=True,)
+			   CustomProperties Pin (PinId=56486C6B491D5ECDBBD5B9B4B1F17EE5,PinName="RotPivot",PinToolTip="Rot Pivot\nVector 2D Structure\n\nLocation (as proportion of quad, 0-1) to rotate about",PinType.PinCategory="struct",PinType.PinSubCategoryObject=ScriptStruct'/Script/CoreUObject.Vector2D',PinType.PinSubCategoryMemberReference=(),PinType.bIsArray=False,PinType.bIsReference=False,PinType.bIsConst=False,PinType.bIsWeakPointer=False,PersistentGuid=00000000000000000000000000000000,bHidden=False,bNotConnectable=False,bDefaultValueIsReadOnly=False,bDefaultValueIsIgnored=False,bAdvancedView=True,)
+			End Object
+			Begin Object Class=EdGraphNode_Comment Name="EdGraphNode_Comment_215"
+			   NodePosX=-400
+			   NodePosY=16
+			   NodeWidth=1440
+			   NodeHeight=640
+			   NodeComment="Draw Crosshair in the center of the screen"
+			   NodeGuid=00F19DD941DB2809BF0B8794E9DB37F7
+			End Object
+
+			END OBJECT
+	[/PARAM]
+[/OBJECT]
+[/EXCERPT:HUD]
+
+
+-->
\ No newline at end of file
diff --git a/Engine/Documentation/Source/Engine/Rendering/RenderTargets/BlueprintRenderTargets/HowTo/HeightFieldPainter/2/HFP2.INT.udn b/Engine/Documentation/Source/Engine/Rendering/RenderTargets/BlueprintRenderTargets/HowTo/HeightFieldPainter/2/HFP2.INT.udn
new file mode 100644
--- /dev/null
+++ b/Engine/Documentation/Source/Engine/Rendering/RenderTargets/BlueprintRenderTargets/HowTo/HeightFieldPainter/2/HFP2.INT.udn
@@ -0,0 +1,552 @@
+availability:docs
+title:2. Required Asset Creation & Setup
+description:Going over creating and setting up the required assets.
+crumbs:%ROOT%, Engine, Engine\Rendering\RenderTargets
+version:4.13
+type:multi-step
+SkillLevel:Intermediate
+parent:Engine/Rendering/RenderTargets/BlueprintRenderTargets/HowTo/HeightFieldPainter
+Navigation:multi-step
+tags: Render Targets
+order:2
+checkpoint:HFP
+
+[NAV]
+
+
+Before we can create the Height Field Painter Blueprint, we will need to create and setup the following asset types inside of the UE4 Content Browser. In the following sections, we will go over how each of the assets needs to be setup as well as any special properties that you will need to be aware of.
+
+[region:note]
+Please note that the following is not specific to any UE4 project type. The following can be accomplished in any project that uses UE4 version 4.13 or later. 
+[/region]
+
+* **Render Target**
+* **Material**
+* **Blueprint** (Based off Actor)
+
+### Render Target Creation & Setup
+
+1. Right-click inside of the **Content Browser** and from the displayed menu, go to **Create Advanced Asset** > **Materials & Textures** and then select the **Render Target** option.
+
+	![](Engine\Rendering\RenderTargets\BlueprintRenderTargets\HowTo\HeightFieldPainter\HFP_RT_Creation_00.png)
+
+1. Name the newly created Render Target, `RT_Height_Filed_Painter`, and then press the **Save All** button to save your work.
+
+	![](Engine\Rendering\RenderTargets\BlueprintRenderTargets\HowTo\HeightFieldPainter\HFP_RT_Creation_01.png)
+
+### Material Creation & Setup
+
+1. Right-click inside of the Content Browser and from the **Create Basic Asset** list, select the **Material** option to create a new Material, naming it `MAT_HeighfieldPainter`.
+
+	![](Engine\Rendering\RenderTargets\BlueprintRenderTargets\HowTo\HeightFieldPainter\HFP_Material_Creation_00.png)
+
+1. Select the **MAT_HeightfieldPainter** Material in the Content Browser and then right-click on it. From the displayed menu, under the **Common** section, select the **Duplicate** option to create a copy of the Material, naming it `MAT_ForceSplat`.
+
+	![](Engine\Rendering\RenderTargets\BlueprintRenderTargets\HowTo\HeightFieldPainter\HFP_Material_Creation_01.png)
+
+1. Double-click on the MAT_HeighfieldPainter Material to open it up and then select the **Main Material Node**. In the **Details** panel under the **Material** section, setup the following parameters.
+
+	* **Shading Model:** Unlit
+
+	![](Engine\Rendering\RenderTargets\BlueprintRenderTargets\HowTo\HeightFieldPainter\HFP_Material_Creation_04.png)
+
+1. The MAT_HeighfieldPainter Material is going to use the **World Position Offset** input to increase the Z-height of the Static Mesh vertices it is applied to. You can copy it into your Material by hovering your mouse over the image below and then clicking on the **Copy Expression Graph** option, copying-and-pasting the provided Material code into `MAT_HeighfieldPainter`.
+
+	[REGION:fullwidth]
+	[INCLUDE:Engine\Rendering\RenderTargets\BlueprintRenderTargets\HowTo\HeightFieldPainter\2\#HFP]
+	[/REGION]
+
+	[region:note]
+	Do not forget to Apply and Save your Material so that it can be used.
+	[/region]
+
+1. Open up the MAT_ForceSplat Material and then select the **Main Material Node**. In the **Details** panel under the **Material** section set up the following parameters.
+
+	* **Blend Mode:** Additive
+	* **Shading Model:** Unlit
+	* **Allow Negative Emissive Color:** Checked
+
+	![](Engine\Rendering\RenderTargets\BlueprintRenderTargets\HowTo\HeightFieldPainter\HFP_Material_Creation_02.png)
+
+1. Since the MAT_ForceSplat Material is going to act as our paint brush for raising the surface of our height field, it will need to have a few options that will allow you to control its size and intensity. You can copy the completed Material into your Material by hovering your mouse over the image below and then clicking on the  **Copy Expression Graph** option, copying-and-pasting the provided Material code into MAT_ForceSplat.
+
+	[REGION:fullwidth]
+	[INCLUDE:Engine\Rendering\RenderTargets\BlueprintRenderTargets\HowTo\HeightFieldPainter\2\#ForceSplat]
+	[/REGION]
+
+	[region:note]
+	Do not forget to Apply and Save your Material when done so that it can be used.
+	[/region]
+
+
+With the required assets now setup, in the next section, we will take a look at creating a Blueprint that will link all of the assets together to create the height map painter.
+
+<!--
+[EXCERPT:HFP]
+[OBJECT:Snippet]
+	[PARAMLITERAL:id]
+		MAT_HeighfieldPainter
+	[/PARAMLITERAL]
+	[PARAMLITERAL:type]
+		material
+	[/PARAMLITERAL]
+	[PARAM:image]
+		![](Engine\Rendering\RenderTargets\BlueprintRenderTargets\HowTo\HeightFieldPainter\HFP_Material_Creation_05.png)
+	[/PARAM]
+	[PARAM:snippet]
+			BEGIN OBJECT
+			Begin Object Class=MaterialGraphNode Name="MaterialGraphNode_14"
+			   Begin Object Class=MaterialExpressionMultiply Name="MaterialExpressionMultiply_22"
+			   End Object
+			   Begin Object Name="MaterialExpressionMultiply_22"
+				  A=(Expression=MaterialExpressionTextureSampleParameter2D'MaterialGraphNode_16.MaterialExpressionTextureSampleParameter2D_9',OutputIndex=1,Mask=1,MaskR=1)
+				  ConstB=0.040000
+				  MaterialExpressionEditorX=-624
+				  MaterialExpressionEditorY=192
+				  MaterialExpressionGuid=E2AF283D404201DC7789608E2A53324A
+				  Material=PreviewMaterial'/Engine/Transient.PreviewMaterial_0'
+			   End Object
+			   MaterialExpression=MaterialExpressionMultiply'MaterialExpressionMultiply_22'
+			   NodePosX=-624
+			   NodePosY=192
+			   NodeGuid=CAA8D05E49144ED277A3BA95B1962A4E
+			   CustomProperties Pin (PinId=97B93A4840B0D7CD3D26D392080091F0,PinName="A",PinType.PinCategory="optional",PinType.PinSubCategoryObject=None,PinType.PinSubCategoryMemberReference=(),PinType.bIsArray=False,PinType.bIsReference=False,PinType.bIsConst=False,PinType.bIsWeakPointer=False,LinkedTo=(MaterialGraphNode_16 CC3EA05C484B6BCFB58BFEB82ED90030,),PersistentGuid=00000000000000000000000000000000,bHidden=False,bNotConnectable=False,bDefaultValueIsReadOnly=False,bDefaultValueIsIgnored=False,bAdvancedView=False,)
+			   CustomProperties Pin (PinId=7EC6FC1940D7CD8131A976BBEE809C5A,PinName="B",PinType.PinCategory="optional",PinType.PinSubCategoryObject=None,PinType.PinSubCategoryMemberReference=(),PinType.bIsArray=False,PinType.bIsReference=False,PinType.bIsConst=False,PinType.bIsWeakPointer=False,PersistentGuid=00000000000000000000000000000000,bHidden=False,bNotConnectable=False,bDefaultValueIsReadOnly=False,bDefaultValueIsIgnored=False,bAdvancedView=False,)
+			   CustomProperties Pin (PinId=79D732ED45ED66B357286288CE1D4169,PinName="Output",PinFriendlyName=" ",Direction="EGPD_Output",PinType.PinSubCategoryObject=None,PinType.PinSubCategoryMemberReference=(),PinType.bIsArray=False,PinType.bIsReference=False,PinType.bIsConst=False,PinType.bIsWeakPointer=False,PersistentGuid=00000000000000000000000000000000,bHidden=False,bNotConnectable=False,bDefaultValueIsReadOnly=False,bDefaultValueIsIgnored=False,bAdvancedView=False,)
+			End Object
+			Begin Object Class=MaterialGraphNode Name="MaterialGraphNode_15"
+			   Begin Object Class=MaterialExpressionAppendVector Name="MaterialExpressionAppendVector_18"
+			   End Object
+			   Begin Object Name="MaterialExpressionAppendVector_18"
+				  A=(Expression=MaterialExpressionAppendVector'MaterialGraphNode_18.MaterialExpressionAppendVector_19')
+				  B=(Expression=MaterialExpressionMultiply'MaterialGraphNode_20.MaterialExpressionMultiply_23')
+				  MaterialExpressionEditorX=-368
+				  MaterialExpressionEditorY=304
+				  MaterialExpressionGuid=3E433813436AADBE38552592F019CFA9
+				  Material=PreviewMaterial'/Engine/Transient.PreviewMaterial_0'
+			   End Object
+			   MaterialExpression=MaterialExpressionAppendVector'MaterialExpressionAppendVector_18'
+			   NodePosX=-368
+			   NodePosY=304
+			   NodeGuid=19532E454AF723C60B444A96831729C6
+			   CustomProperties Pin (PinId=D61F70A04A8AAB86D09BD3A1BBA5086B,PinName="A",PinType.PinCategory="required",PinType.PinSubCategoryObject=None,PinType.PinSubCategoryMemberReference=(),PinType.bIsArray=False,PinType.bIsReference=False,PinType.bIsConst=False,PinType.bIsWeakPointer=False,LinkedTo=(MaterialGraphNode_18 DDF992834FB4205C052BACB5CC1A8B19,),PersistentGuid=00000000000000000000000000000000,bHidden=False,bNotConnectable=False,bDefaultValueIsReadOnly=False,bDefaultValueIsIgnored=False,bAdvancedView=False,)
+			   CustomProperties Pin (PinId=6AF4F8AD4D399C5CD0248A9DED8E5B6C,PinName="B",PinType.PinCategory="required",PinType.PinSubCategoryObject=None,PinType.PinSubCategoryMemberReference=(),PinType.bIsArray=False,PinType.bIsReference=False,PinType.bIsConst=False,PinType.bIsWeakPointer=False,LinkedTo=(MaterialGraphNode_20 09BEE4B14B6758FF738E8CA0356164B3,),PersistentGuid=00000000000000000000000000000000,bHidden=False,bNotConnectable=False,bDefaultValueIsReadOnly=False,bDefaultValueIsIgnored=False,bAdvancedView=False,)
+			   CustomProperties Pin (PinId=2047D8DD4E4B6BCE0A4A95BF0A71EBE9,PinName="Output",PinFriendlyName=" ",Direction="EGPD_Output",PinType.PinSubCategoryObject=None,PinType.PinSubCategoryMemberReference=(),PinType.bIsArray=False,PinType.bIsReference=False,PinType.bIsConst=False,PinType.bIsWeakPointer=False,PersistentGuid=00000000000000000000000000000000,bHidden=False,bNotConnectable=False,bDefaultValueIsReadOnly=False,bDefaultValueIsIgnored=False,bAdvancedView=False,)
+			End Object
+			Begin Object Class=MaterialGraphNode Name="MaterialGraphNode_16"
+			   Begin Object Class=MaterialExpressionTextureSampleParameter2D Name="MaterialExpressionTextureSampleParameter2D_9"
+			   End Object
+			   Begin Object Name="MaterialExpressionTextureSampleParameter2D_9"
+				  ParameterName="Heightfield"
+				  ExpressionGUID=3591EB284AC9E677EE103BA9C02FDE85
+				  Texture=Texture2D'/Game/ExampleContent/BlueprintFluidSurface/DefaultHeightfield.DefaultHeightfield'
+				  SamplerType=SAMPLERTYPE_LinearColor
+				  MaterialExpressionEditorX=-944
+				  MaterialExpressionEditorY=144
+				  MaterialExpressionGuid=81CC5EB24D517AFD346496AC611464AF
+				  Material=PreviewMaterial'/Engine/Transient.PreviewMaterial_0'
+			   End Object
+			   MaterialExpression=MaterialExpressionTextureSampleParameter2D'MaterialExpressionTextureSampleParameter2D_9'
+			   NodePosX=-944
+			   NodePosY=144
+			   bCanRenameNode=True
+			   NodeGuid=F9A9F78B480596161D21B0BF16B1604C
+			   CustomProperties Pin (PinId=98EBA32F426B3DE25B86E38B708A7BB7,PinName="UVs",PinType.PinCategory="optional",PinType.PinSubCategoryObject=None,PinType.PinSubCategoryMemberReference=(),PinType.bIsArray=False,PinType.bIsReference=False,PinType.bIsConst=False,PinType.bIsWeakPointer=False,PersistentGuid=00000000000000000000000000000000,bHidden=False,bNotConnectable=False,bDefaultValueIsReadOnly=False,bDefaultValueIsIgnored=False,bAdvancedView=False,)
+			   CustomProperties Pin (PinId=B624B46141F0A5C25E0D939D12CF55BB,PinName="Tex",PinType.PinCategory="optional",PinType.PinSubCategoryObject=None,PinType.PinSubCategoryMemberReference=(),PinType.bIsArray=False,PinType.bIsReference=False,PinType.bIsConst=False,PinType.bIsWeakPointer=False,PersistentGuid=00000000000000000000000000000000,bHidden=False,bNotConnectable=False,bDefaultValueIsReadOnly=False,bDefaultValueIsIgnored=False,bAdvancedView=False,)
+			   CustomProperties Pin (PinId=C0CCDCFC4AF6C065BD9421866C14E1B6,PinName="Output",PinFriendlyName=" ",Direction="EGPD_Output",PinType.PinCategory="mask",PinType.PinSubCategoryObject=None,PinType.PinSubCategoryMemberReference=(),PinType.bIsArray=False,PinType.bIsReference=False,PinType.bIsConst=False,PinType.bIsWeakPointer=False,PersistentGuid=00000000000000000000000000000000,bHidden=False,bNotConnectable=False,bDefaultValueIsReadOnly=False,bDefaultValueIsIgnored=False,bAdvancedView=False,)
+			   CustomProperties Pin (PinId=CC3EA05C484B6BCFB58BFEB82ED90030,PinName="Output2",PinFriendlyName=" ",Direction="EGPD_Output",PinType.PinCategory="mask",PinType.PinSubCategory="red",PinType.PinSubCategoryObject=None,PinType.PinSubCategoryMemberReference=(),PinType.bIsArray=False,PinType.bIsReference=False,PinType.bIsConst=False,PinType.bIsWeakPointer=False,LinkedTo=(MaterialGraphNode_14 97B93A4840B0D7CD3D26D392080091F0,MaterialGraphNode_20 FA5D4D794B54A67776ED72862BD5CAF4,),PersistentGuid=00000000000000000000000000000000,bHidden=False,bNotConnectable=False,bDefaultValueIsReadOnly=False,bDefaultValueIsIgnored=False,bAdvancedView=False,)
+			   CustomProperties Pin (PinId=F2F541134532F3877048FA8419FF5FB5,PinName="Output3",PinFriendlyName=" ",Direction="EGPD_Output",PinType.PinCategory="mask",PinType.PinSubCategory="green",PinType.PinSubCategoryObject=None,PinType.PinSubCategoryMemberReference=(),PinType.bIsArray=False,PinType.bIsReference=False,PinType.bIsConst=False,PinType.bIsWeakPointer=False,PersistentGuid=00000000000000000000000000000000,bHidden=False,bNotConnectable=False,bDefaultValueIsReadOnly=False,bDefaultValueIsIgnored=False,bAdvancedView=False,)
+			   CustomProperties Pin (PinId=740C095F463A6417CFF8D49296DC15E9,PinName="Output4",PinFriendlyName=" ",Direction="EGPD_Output",PinType.PinCategory="mask",PinType.PinSubCategory="blue",PinType.PinSubCategoryObject=None,PinType.PinSubCategoryMemberReference=(),PinType.bIsArray=False,PinType.bIsReference=False,PinType.bIsConst=False,PinType.bIsWeakPointer=False,PersistentGuid=00000000000000000000000000000000,bHidden=False,bNotConnectable=False,bDefaultValueIsReadOnly=False,bDefaultValueIsIgnored=False,bAdvancedView=False,)
+			   CustomProperties Pin (PinId=276B603D4C841B3549B3BD94B6EB9740,PinName="Output5",PinFriendlyName=" ",Direction="EGPD_Output",PinType.PinCategory="mask",PinType.PinSubCategory="alpha",PinType.PinSubCategoryObject=None,PinType.PinSubCategoryMemberReference=(),PinType.bIsArray=False,PinType.bIsReference=False,PinType.bIsConst=False,PinType.bIsWeakPointer=False,PersistentGuid=00000000000000000000000000000000,bHidden=False,bNotConnectable=False,bDefaultValueIsReadOnly=False,bDefaultValueIsIgnored=False,bAdvancedView=False,)
+			End Object
+			Begin Object Class=MaterialGraphNode Name="MaterialGraphNode_17"
+			   Begin Object Class=MaterialExpressionScalarParameter Name="MaterialExpressionScalarParameter_19"
+			   End Object
+			   Begin Object Name="MaterialExpressionScalarParameter_19"
+				  DefaultValue=5.000000
+				  ParameterName="HeightScale"
+				  ExpressionGUID=C066E7854D940718A6354E9305297A64
+				  MaterialExpressionEditorX=-768
+				  MaterialExpressionEditorY=448
+				  MaterialExpressionGuid=A848058E4A2A157D5EAEE0AA28CCC2AE
+				  Material=PreviewMaterial'/Engine/Transient.PreviewMaterial_0'
+			   End Object
+			   MaterialExpression=MaterialExpressionScalarParameter'MaterialExpressionScalarParameter_19'
+			   NodePosX=-768
+			   NodePosY=448
+			   bCanRenameNode=True
+			   NodeGuid=C3946F664AA7008F564DBA9FADBF8963
+			   CustomProperties Pin (PinId=2F4B7FC149C8A90335DABAAA2D1E12ED,PinName="Output",PinFriendlyName=" ",Direction="EGPD_Output",PinType.PinSubCategoryObject=None,PinType.PinSubCategoryMemberReference=(),PinType.bIsArray=False,PinType.bIsReference=False,PinType.bIsConst=False,PinType.bIsWeakPointer=False,LinkedTo=(MaterialGraphNode_20 4BEE0E524D40A738BE6A83ABC395888D,),PersistentGuid=00000000000000000000000000000000,bHidden=False,bNotConnectable=False,bDefaultValueIsReadOnly=False,bDefaultValueIsIgnored=False,bAdvancedView=False,)
+			End Object
+			Begin Object Class=MaterialGraphNode Name="MaterialGraphNode_18"
+			   Begin Object Class=MaterialExpressionAppendVector Name="MaterialExpressionAppendVector_19"
+			   End Object
+			   Begin Object Name="MaterialExpressionAppendVector_19"
+				  A=(Expression=MaterialExpressionConstant'MaterialGraphNode_19.MaterialExpressionConstant_14')
+				  B=(Expression=MaterialExpressionConstant'MaterialGraphNode_19.MaterialExpressionConstant_14')
+				  MaterialExpressionEditorX=-512
+				  MaterialExpressionEditorY=304
+				  MaterialExpressionGuid=3E433813436AADBE38552592F019CFA9
+				  Material=PreviewMaterial'/Engine/Transient.PreviewMaterial_0'
+			   End Object
+			   MaterialExpression=MaterialExpressionAppendVector'MaterialExpressionAppendVector_19'
+			   NodePosX=-512
+			   NodePosY=304
+			   NodeGuid=27720EBA402EEF65B41295ADC8E3E617
+			   CustomProperties Pin (PinId=3DA7F1CD49894BCCA1CCE3855C861D9C,PinName="A",PinType.PinCategory="required",PinType.PinSubCategoryObject=None,PinType.PinSubCategoryMemberReference=(),PinType.bIsArray=False,PinType.bIsReference=False,PinType.bIsConst=False,PinType.bIsWeakPointer=False,LinkedTo=(MaterialGraphNode_19 205C66154595D71B256F1D97701BB371,),PersistentGuid=00000000000000000000000000000000,bHidden=False,bNotConnectable=False,bDefaultValueIsReadOnly=False,bDefaultValueIsIgnored=False,bAdvancedView=False,)
+			   CustomProperties Pin (PinId=73D98CBC451ECC9028D454B447B5C151,PinName="B",PinType.PinCategory="required",PinType.PinSubCategoryObject=None,PinType.PinSubCategoryMemberReference=(),PinType.bIsArray=False,PinType.bIsReference=False,PinType.bIsConst=False,PinType.bIsWeakPointer=False,LinkedTo=(MaterialGraphNode_19 205C66154595D71B256F1D97701BB371,),PersistentGuid=00000000000000000000000000000000,bHidden=False,bNotConnectable=False,bDefaultValueIsReadOnly=False,bDefaultValueIsIgnored=False,bAdvancedView=False,)
+			   CustomProperties Pin (PinId=DDF992834FB4205C052BACB5CC1A8B19,PinName="Output",PinFriendlyName=" ",Direction="EGPD_Output",PinType.PinSubCategoryObject=None,PinType.PinSubCategoryMemberReference=(),PinType.bIsArray=False,PinType.bIsReference=False,PinType.bIsConst=False,PinType.bIsWeakPointer=False,LinkedTo=(MaterialGraphNode_15 D61F70A04A8AAB86D09BD3A1BBA5086B,),PersistentGuid=00000000000000000000000000000000,bHidden=False,bNotConnectable=False,bDefaultValueIsReadOnly=False,bDefaultValueIsIgnored=False,bAdvancedView=False,)
+			End Object
+			Begin Object Class=MaterialGraphNode Name="MaterialGraphNode_19"
+			   Begin Object Class=MaterialExpressionConstant Name="MaterialExpressionConstant_14"
+			   End Object
+			   Begin Object Name="MaterialExpressionConstant_14"
+				  MaterialExpressionEditorX=-608
+				  MaterialExpressionEditorY=320
+				  MaterialExpressionGuid=4A04576C44F3196185EC04801D05ED96
+				  Material=PreviewMaterial'/Engine/Transient.PreviewMaterial_0'
+			   End Object
+			   MaterialExpression=MaterialExpressionConstant'MaterialExpressionConstant_14'
+			   NodePosX=-608
+			   NodePosY=320
+			   NodeGuid=CBEA7DF24FFC399A31E95C9D9CF660EE
+			   CustomProperties Pin (PinId=205C66154595D71B256F1D97701BB371,PinName="Output",PinFriendlyName=" ",Direction="EGPD_Output",PinType.PinSubCategoryObject=None,PinType.PinSubCategoryMemberReference=(),PinType.bIsArray=False,PinType.bIsReference=False,PinType.bIsConst=False,PinType.bIsWeakPointer=False,LinkedTo=(MaterialGraphNode_18 3DA7F1CD49894BCCA1CCE3855C861D9C,MaterialGraphNode_18 73D98CBC451ECC9028D454B447B5C151,),PersistentGuid=00000000000000000000000000000000,bHidden=False,bNotConnectable=False,bDefaultValueIsReadOnly=False,bDefaultValueIsIgnored=False,bAdvancedView=False,)
+			End Object
+			Begin Object Class=MaterialGraphNode Name="MaterialGraphNode_20"
+			   Begin Object Class=MaterialExpressionMultiply Name="MaterialExpressionMultiply_23"
+			   End Object
+			   Begin Object Name="MaterialExpressionMultiply_23"
+				  A=(Expression=MaterialExpressionTextureSampleParameter2D'MaterialGraphNode_16.MaterialExpressionTextureSampleParameter2D_9',OutputIndex=1,Mask=1,MaskR=1)
+				  B=(Expression=MaterialExpressionScalarParameter'MaterialGraphNode_17.MaterialExpressionScalarParameter_19')
+				  MaterialExpressionEditorX=-576
+				  MaterialExpressionEditorY=400
+				  MaterialExpressionGuid=912109BE43BDDA48B7874CA2FBBA178E
+				  Material=PreviewMaterial'/Engine/Transient.PreviewMaterial_0'
+			   End Object
+			   MaterialExpression=MaterialExpressionMultiply'MaterialExpressionMultiply_23'
+			   NodePosX=-576
+			   NodePosY=400
+			   NodeGuid=E8B782254B2103A33DF97397EC53ECE5
+			   CustomProperties Pin (PinId=FA5D4D794B54A67776ED72862BD5CAF4,PinName="A",PinType.PinCategory="optional",PinType.PinSubCategoryObject=None,PinType.PinSubCategoryMemberReference=(),PinType.bIsArray=False,PinType.bIsReference=False,PinType.bIsConst=False,PinType.bIsWeakPointer=False,LinkedTo=(MaterialGraphNode_16 CC3EA05C484B6BCFB58BFEB82ED90030,),PersistentGuid=00000000000000000000000000000000,bHidden=False,bNotConnectable=False,bDefaultValueIsReadOnly=False,bDefaultValueIsIgnored=False,bAdvancedView=False,)
+			   CustomProperties Pin (PinId=4BEE0E524D40A738BE6A83ABC395888D,PinName="B",PinType.PinCategory="optional",PinType.PinSubCategoryObject=None,PinType.PinSubCategoryMemberReference=(),PinType.bIsArray=False,PinType.bIsReference=False,PinType.bIsConst=False,PinType.bIsWeakPointer=False,LinkedTo=(MaterialGraphNode_17 2F4B7FC149C8A90335DABAAA2D1E12ED,),PersistentGuid=00000000000000000000000000000000,bHidden=False,bNotConnectable=False,bDefaultValueIsReadOnly=False,bDefaultValueIsIgnored=False,bAdvancedView=False,)
+			   CustomProperties Pin (PinId=09BEE4B14B6758FF738E8CA0356164B3,PinName="Output",PinFriendlyName=" ",Direction="EGPD_Output",PinType.PinSubCategoryObject=None,PinType.PinSubCategoryMemberReference=(),PinType.bIsArray=False,PinType.bIsReference=False,PinType.bIsConst=False,PinType.bIsWeakPointer=False,LinkedTo=(MaterialGraphNode_15 6AF4F8AD4D399C5CD0248A9DED8E5B6C,),PersistentGuid=00000000000000000000000000000000,bHidden=False,bNotConnectable=False,bDefaultValueIsReadOnly=False,bDefaultValueIsIgnored=False,bAdvancedView=False,)
+			End Object
+
+			END OBJECT
+	[/PARAM]
+[/OBJECT]
+[/EXCERPT:HFP]
+-->
+
+<!--
+[EXCERPT:ForceSplat]	
+[OBJECT:Snippet]
+	[PARAMLITERAL:id]
+		MAT_ForceSplat
+	[/PARAMLITERAL]
+	[PARAMLITERAL:type]
+		material
+	[/PARAMLITERAL]
+	[PARAM:image]
+		![](Engine\Rendering\RenderTargets\BlueprintRenderTargets\HowTo\HeightFieldPainter\HFP_Material_Creation_03.png)
+	[/PARAM]
+	[PARAM:snippet]
+			Begin Object Class=MaterialGraphNode Name="MaterialGraphNode_0"
+			   Begin Object Class=MaterialExpressionMultiply Name="MaterialExpressionMultiply_2"
+			   End Object
+			   Begin Object Name="MaterialExpressionMultiply_2"
+				  A=(Expression=MaterialExpressionDivide'MaterialGraphNode_13.MaterialExpressionDivide_2')
+				  B=(Expression=MaterialExpressionScalarParameter'MaterialGraphNode_1.MaterialExpressionScalarParameter_4')
+				  MaterialExpressionEditorX=112
+				  MaterialExpressionEditorY=80
+				  MaterialExpressionGuid=3025F98C46D3F1D0C59CD387F8C33F29
+				  Material=PreviewMaterial'/Engine/Transient.ForceSplat'
+			   End Object
+			   MaterialExpression=MaterialExpressionMultiply'MaterialExpressionMultiply_2'
+			   NodePosX=112
+			   NodePosY=80
+			   NodeGuid=5AD1F5CD44FA09FD0FDAA1AFA22A4AA9
+			   CustomProperties Pin (PinId=B7D0CF28498836478F9C8987312501BF,PinName="A",PinType.PinCategory="optional",PinType.PinSubCategoryObject=None,PinType.PinSubCategoryMemberReference=(),PinType.bIsArray=False,PinType.bIsReference=False,PinType.bIsConst=False,PinType.bIsWeakPointer=False,LinkedTo=(MaterialGraphNode_13 D0CE6B2B4A6886F59532E4B83367D7BD,),PersistentGuid=00000000000000000000000000000000,bHidden=False,bNotConnectable=False,bDefaultValueIsReadOnly=False,bDefaultValueIsIgnored=False,bAdvancedView=False,)
+			   CustomProperties Pin (PinId=362F295247D0A2B0935DB9889A9995EB,PinName="B",PinType.PinCategory="optional",PinType.PinSubCategoryObject=None,PinType.PinSubCategoryMemberReference=(),PinType.bIsArray=False,PinType.bIsReference=False,PinType.bIsConst=False,PinType.bIsWeakPointer=False,LinkedTo=(MaterialGraphNode_1 751E250A4F77873070E2D29BD2471826,),PersistentGuid=00000000000000000000000000000000,bHidden=False,bNotConnectable=False,bDefaultValueIsReadOnly=False,bDefaultValueIsIgnored=False,bAdvancedView=False,)
+			   CustomProperties Pin (PinId=1F55ED9148C9F02DB35DA78AA20C458D,PinName="Output",PinFriendlyName=" ",Direction="EGPD_Output",PinType.PinSubCategoryObject=None,PinType.PinSubCategoryMemberReference=(),PinType.bIsArray=False,PinType.bIsReference=False,PinType.bIsConst=False,PinType.bIsWeakPointer=False,LinkedTo=(MaterialGraphNode_Root_0 3890D74A48618B9B44F8AA8F876CD0C8,),PersistentGuid=00000000000000000000000000000000,bHidden=False,bNotConnectable=False,bDefaultValueIsReadOnly=False,bDefaultValueIsIgnored=False,bAdvancedView=False,)
+			End Object
+			Begin Object Class=MaterialGraphNode Name="MaterialGraphNode_1"
+			   Begin Object Class=MaterialExpressionScalarParameter Name="MaterialExpressionScalarParameter_4"
+			   End Object
+			   Begin Object Name="MaterialExpressionScalarParameter_4"
+				  DefaultValue=1.000000
+				  ParameterName="ForceStrength"
+				  ExpressionGUID=EF5F282F41F66D3274BC6FA7B987DCC9
+				  MaterialExpressionEditorX=-64
+				  MaterialExpressionEditorY=192
+				  MaterialExpressionGuid=9EF14C914FADADA4DBDE7B9C45E23576
+				  Material=PreviewMaterial'/Engine/Transient.ForceSplat'
+			   End Object
+			   MaterialExpression=MaterialExpressionScalarParameter'MaterialExpressionScalarParameter_4'
+			   NodePosX=-64
+			   NodePosY=192
+			   bCanRenameNode=True
+			   NodeGuid=E2EF5596438E8DA9CBD9509FBD3BF3F5
+			   CustomProperties Pin (PinId=751E250A4F77873070E2D29BD2471826,PinName="Output",PinFriendlyName=" ",Direction="EGPD_Output",PinType.PinSubCategoryObject=None,PinType.PinSubCategoryMemberReference=(),PinType.bIsArray=False,PinType.bIsReference=False,PinType.bIsConst=False,PinType.bIsWeakPointer=False,LinkedTo=(MaterialGraphNode_0 362F295247D0A2B0935DB9889A9995EB,),PersistentGuid=00000000000000000000000000000000,bHidden=False,bNotConnectable=False,bDefaultValueIsReadOnly=False,bDefaultValueIsIgnored=False,bAdvancedView=False,)
+			End Object
+			Begin Object Class=MaterialGraphNode Name="MaterialGraphNode_2"
+			   Begin Object Class=MaterialExpressionSubtract Name="MaterialExpressionSubtract_6"
+			   End Object
+			   Begin Object Name="MaterialExpressionSubtract_6"
+				  A=(Expression=MaterialExpressionSubtract'MaterialGraphNode_8.MaterialExpressionSubtract_7')
+				  B=(Expression=MaterialExpressionOneMinus'MaterialGraphNode_3.MaterialExpressionOneMinus_2')
+				  MaterialExpressionEditorX=-496
+				  MaterialExpressionEditorY=32
+				  MaterialExpressionGuid=0F3C46254F00CCEB93571093119ECB9D
+				  Material=PreviewMaterial'/Engine/Transient.ForceSplat'
+			   End Object
+			   MaterialExpression=MaterialExpressionSubtract'MaterialExpressionSubtract_6'
+			   NodePosX=-496
+			   NodePosY=32
+			   NodeGuid=478BC67945A3DC7AAD1D518FEE303739
+			   CustomProperties Pin (PinId=C5F733EF466287E9E4EB35ACE6B65E84,PinName="A",PinType.PinCategory="optional",PinType.PinSubCategoryObject=None,PinType.PinSubCategoryMemberReference=(),PinType.bIsArray=False,PinType.bIsReference=False,PinType.bIsConst=False,PinType.bIsWeakPointer=False,LinkedTo=(MaterialGraphNode_8 38FFC1874298113EDA7D9DBC71D7C22F,),PersistentGuid=00000000000000000000000000000000,bHidden=False,bNotConnectable=False,bDefaultValueIsReadOnly=False,bDefaultValueIsIgnored=False,bAdvancedView=False,)
+			   CustomProperties Pin (PinId=6AADA53C49FE3593A2AB3D9626E59461,PinName="B",PinType.PinCategory="optional",PinType.PinSubCategoryObject=None,PinType.PinSubCategoryMemberReference=(),PinType.bIsArray=False,PinType.bIsReference=False,PinType.bIsConst=False,PinType.bIsWeakPointer=False,LinkedTo=(MaterialGraphNode_3 368EE4A347A49750642B86858743DC92,),PersistentGuid=00000000000000000000000000000000,bHidden=False,bNotConnectable=False,bDefaultValueIsReadOnly=False,bDefaultValueIsIgnored=False,bAdvancedView=False,)
+			   CustomProperties Pin (PinId=25D2D0404B46764F28521E9B20B14B64,PinName="Output",PinFriendlyName=" ",Direction="EGPD_Output",PinType.PinSubCategoryObject=None,PinType.PinSubCategoryMemberReference=(),PinType.bIsArray=False,PinType.bIsReference=False,PinType.bIsConst=False,PinType.bIsWeakPointer=False,LinkedTo=(MaterialGraphNode_7 42C937124B3EBFB7D0BFD98373FA7A0B,),PersistentGuid=00000000000000000000000000000000,bHidden=False,bNotConnectable=False,bDefaultValueIsReadOnly=False,bDefaultValueIsIgnored=False,bAdvancedView=False,)
+			End Object
+			Begin Object Class=MaterialGraphNode Name="MaterialGraphNode_3"
+			   Begin Object Class=MaterialExpressionOneMinus Name="MaterialExpressionOneMinus_2"
+			   End Object
+			   Begin Object Name="MaterialExpressionOneMinus_2"
+				  Input=(Expression=MaterialExpressionScalarParameter'MaterialGraphNode_4.MaterialExpressionScalarParameter_5')
+				  MaterialExpressionEditorX=-656
+				  MaterialExpressionEditorY=64
+				  MaterialExpressionGuid=0FC613604260BF8A67765EB5C5556AE0
+				  Material=PreviewMaterial'/Engine/Transient.ForceSplat'
+			   End Object
+			   MaterialExpression=MaterialExpressionOneMinus'MaterialExpressionOneMinus_2'
+			   NodePosX=-656
+			   NodePosY=64
+			   NodeGuid=1530583D4196B0C70F53138D51ED4AF4
+			   CustomProperties Pin (PinId=98DDAD8444773CE90CF776BF47876227,PinName="Input",PinFriendlyName=" ",PinType.PinCategory="required",PinType.PinSubCategoryObject=None,PinType.PinSubCategoryMemberReference=(),PinType.bIsArray=False,PinType.bIsReference=False,PinType.bIsConst=False,PinType.bIsWeakPointer=False,LinkedTo=(MaterialGraphNode_4 243D5A9F4A177E21AD5B0C9A9EE13B6C,),PersistentGuid=00000000000000000000000000000000,bHidden=False,bNotConnectable=False,bDefaultValueIsReadOnly=False,bDefaultValueIsIgnored=False,bAdvancedView=False,)
+			   CustomProperties Pin (PinId=368EE4A347A49750642B86858743DC92,PinName="Output",PinFriendlyName=" ",Direction="EGPD_Output",PinType.PinSubCategoryObject=None,PinType.PinSubCategoryMemberReference=(),PinType.bIsArray=False,PinType.bIsReference=False,PinType.bIsConst=False,PinType.bIsWeakPointer=False,LinkedTo=(MaterialGraphNode_2 6AADA53C49FE3593A2AB3D9626E59461,),PersistentGuid=00000000000000000000000000000000,bHidden=False,bNotConnectable=False,bDefaultValueIsReadOnly=False,bDefaultValueIsIgnored=False,bAdvancedView=False,)
+			End Object
+			Begin Object Class=MaterialGraphNode Name="MaterialGraphNode_4"
+			   Begin Object Class=MaterialExpressionScalarParameter Name="MaterialExpressionScalarParameter_5"
+			   End Object
+			   Begin Object Name="MaterialExpressionScalarParameter_5"
+				  DefaultValue=0.020000
+				  ParameterName="ForceSize"
+				  ExpressionGUID=22FEC1D646BC1B274C6F6B8AF07EC523
+				  MaterialExpressionEditorX=-880
+				  MaterialExpressionEditorY=224
+				  MaterialExpressionGuid=08209E994F6A7C27F8AF1B987C8BA87A
+				  Material=PreviewMaterial'/Engine/Transient.ForceSplat'
+			   End Object
+			   MaterialExpression=MaterialExpressionScalarParameter'MaterialExpressionScalarParameter_5'
+			   NodePosX=-880
+			   NodePosY=224
+			   bCanRenameNode=True
+			   NodeGuid=3C5CD52A4E4EACE650CA47871ED15868
+			   CustomProperties Pin (PinId=243D5A9F4A177E21AD5B0C9A9EE13B6C,PinName="Output",PinFriendlyName=" ",Direction="EGPD_Output",PinType.PinSubCategoryObject=None,PinType.PinSubCategoryMemberReference=(),PinType.bIsArray=False,PinType.bIsReference=False,PinType.bIsConst=False,PinType.bIsWeakPointer=False,LinkedTo=(MaterialGraphNode_3 98DDAD8444773CE90CF776BF47876227,MaterialGraphNode_13 3501ECFD4967CB5C6FA143BB453E5636,),PersistentGuid=00000000000000000000000000000000,bHidden=False,bNotConnectable=False,bDefaultValueIsReadOnly=False,bDefaultValueIsIgnored=False,bAdvancedView=False,)
+			End Object
+			Begin Object Class=MaterialGraphNode Name="MaterialGraphNode_5"
+			   Begin Object Class=MaterialExpressionComponentMask Name="MaterialExpressionComponentMask_2"
+			   End Object
+			   Begin Object Name="MaterialExpressionComponentMask_2"
+				  Input=(Expression=MaterialExpressionVectorParameter'MaterialGraphNode_6.MaterialExpressionVectorParameter_2',Mask=1,MaskR=1,MaskG=1,MaskB=1)
+				  R=True
+				  G=True
+				  MaterialExpressionEditorX=-1280
+				  MaterialExpressionEditorY=48
+				  MaterialExpressionGuid=FF2B6BA24182D92569D49491035FFA18
+				  Material=PreviewMaterial'/Engine/Transient.ForceSplat'
+			   End Object
+			   MaterialExpression=MaterialExpressionComponentMask'MaterialExpressionComponentMask_2'
+			   NodePosX=-1280
+			   NodePosY=48
+			   NodeGuid=FC6C573F4C2431384FF35987120D12B5
+			   CustomProperties Pin (PinId=EC552D08411CF9140F7BCDA9D904E890,PinName="Input",PinFriendlyName=" ",PinType.PinCategory="required",PinType.PinSubCategoryObject=None,PinType.PinSubCategoryMemberReference=(),PinType.bIsArray=False,PinType.bIsReference=False,PinType.bIsConst=False,PinType.bIsWeakPointer=False,LinkedTo=(MaterialGraphNode_6 7B0A20374241C6FD036BDFADE36D1DD7,),PersistentGuid=00000000000000000000000000000000,bHidden=False,bNotConnectable=False,bDefaultValueIsReadOnly=False,bDefaultValueIsIgnored=False,bAdvancedView=False,)
+			   CustomProperties Pin (PinId=2A7AE85A415E54152002F9BA674B7E41,PinName="Output",PinFriendlyName=" ",Direction="EGPD_Output",PinType.PinSubCategoryObject=None,PinType.PinSubCategoryMemberReference=(),PinType.bIsArray=False,PinType.bIsReference=False,PinType.bIsConst=False,PinType.bIsWeakPointer=False,LinkedTo=(MaterialGraphNode_10 26815BDA44BF1141EF0D0084C2E3D993,),PersistentGuid=00000000000000000000000000000000,bHidden=False,bNotConnectable=False,bDefaultValueIsReadOnly=False,bDefaultValueIsIgnored=False,bAdvancedView=False,)
+			End Object
+			Begin Object Class=MaterialGraphNode Name="MaterialGraphNode_6"
+			   Begin Object Class=MaterialExpressionVectorParameter Name="MaterialExpressionVectorParameter_2"
+			   End Object
+			   Begin Object Name="MaterialExpressionVectorParameter_2"
+				  DefaultValue=(R=0.500000,G=0.500000,B=0.000000,A=0.000000)
+				  ParameterName="ForcePosition"
+				  ExpressionGUID=EB3126774C539A8048339782BDD5212C
+				  MaterialExpressionEditorX=-1472
+				  MaterialExpressionEditorY=32
+				  MaterialExpressionGuid=75C595794CE24CE6008419B632EE3760
+				  Material=PreviewMaterial'/Engine/Transient.ForceSplat'
+			   End Object
+			   MaterialExpression=MaterialExpressionVectorParameter'MaterialExpressionVectorParameter_2'
+			   NodePosX=-1472
+			   NodePosY=32
+			   bCanRenameNode=True
+			   NodeGuid=3869F5FA462D6E66ABABC89D92F2BF20
+			   CustomProperties Pin (PinId=7B0A20374241C6FD036BDFADE36D1DD7,PinName="Output",PinFriendlyName=" ",Direction="EGPD_Output",PinType.PinCategory="mask",PinType.PinSubCategoryObject=None,PinType.PinSubCategoryMemberReference=(),PinType.bIsArray=False,PinType.bIsReference=False,PinType.bIsConst=False,PinType.bIsWeakPointer=False,LinkedTo=(MaterialGraphNode_5 EC552D08411CF9140F7BCDA9D904E890,),PersistentGuid=00000000000000000000000000000000,bHidden=False,bNotConnectable=False,bDefaultValueIsReadOnly=False,bDefaultValueIsIgnored=False,bAdvancedView=False,)
+			   CustomProperties Pin (PinId=DA45DDA34307453D2D4A169D672CADA4,PinName="Output2",PinFriendlyName=" ",Direction="EGPD_Output",PinType.PinCategory="mask",PinType.PinSubCategory="red",PinType.PinSubCategoryObject=None,PinType.PinSubCategoryMemberReference=(),PinType.bIsArray=False,PinType.bIsReference=False,PinType.bIsConst=False,PinType.bIsWeakPointer=False,PersistentGuid=00000000000000000000000000000000,bHidden=False,bNotConnectable=False,bDefaultValueIsReadOnly=False,bDefaultValueIsIgnored=False,bAdvancedView=False,)
+			   CustomProperties Pin (PinId=D93CF20D46F990734F44869983F64E57,PinName="Output3",PinFriendlyName=" ",Direction="EGPD_Output",PinType.PinCategory="mask",PinType.PinSubCategory="green",PinType.PinSubCategoryObject=None,PinType.PinSubCategoryMemberReference=(),PinType.bIsArray=False,PinType.bIsReference=False,PinType.bIsConst=False,PinType.bIsWeakPointer=False,PersistentGuid=00000000000000000000000000000000,bHidden=False,bNotConnectable=False,bDefaultValueIsReadOnly=False,bDefaultValueIsIgnored=False,bAdvancedView=False,)
+			   CustomProperties Pin (PinId=7930B3334341172972107B8C418D144C,PinName="Output4",PinFriendlyName=" ",Direction="EGPD_Output",PinType.PinCategory="mask",PinType.PinSubCategory="blue",PinType.PinSubCategoryObject=None,PinType.PinSubCategoryMemberReference=(),PinType.bIsArray=False,PinType.bIsReference=False,PinType.bIsConst=False,PinType.bIsWeakPointer=False,PersistentGuid=00000000000000000000000000000000,bHidden=False,bNotConnectable=False,bDefaultValueIsReadOnly=False,bDefaultValueIsIgnored=False,bAdvancedView=False,)
+			   CustomProperties Pin (PinId=76733A1049495B6147E21FB88B151EF3,PinName="Output5",PinFriendlyName=" ",Direction="EGPD_Output",PinType.PinCategory="mask",PinType.PinSubCategory="alpha",PinType.PinSubCategoryObject=None,PinType.PinSubCategoryMemberReference=(),PinType.bIsArray=False,PinType.bIsReference=False,PinType.bIsConst=False,PinType.bIsWeakPointer=False,PersistentGuid=00000000000000000000000000000000,bHidden=False,bNotConnectable=False,bDefaultValueIsReadOnly=False,bDefaultValueIsIgnored=False,bAdvancedView=False,)
+			End Object
+			Begin Object Class=MaterialGraphNode Name="MaterialGraphNode_7"
+			   Begin Object Class=MaterialExpressionClamp Name="MaterialExpressionClamp_2"
+			   End Object
+			   Begin Object Name="MaterialExpressionClamp_2"
+				  Input=(Expression=MaterialExpressionSubtract'MaterialGraphNode_2.MaterialExpressionSubtract_6')
+				  MaterialExpressionEditorX=-368
+				  MaterialExpressionEditorY=48
+				  MaterialExpressionGuid=A23A3F7746E7BD37E9FD7180212F059B
+				  Material=PreviewMaterial'/Engine/Transient.ForceSplat'
+			   End Object
+			   MaterialExpression=MaterialExpressionClamp'MaterialExpressionClamp_2'
+			   NodePosX=-368
+			   NodePosY=48
+			   NodeGuid=A0E540AE43E3B6A3F381579913CF988D
+			   CustomProperties Pin (PinId=42C937124B3EBFB7D0BFD98373FA7A0B,PinName="Input",PinFriendlyName=" ",PinType.PinCategory="required",PinType.PinSubCategoryObject=None,PinType.PinSubCategoryMemberReference=(),PinType.bIsArray=False,PinType.bIsReference=False,PinType.bIsConst=False,PinType.bIsWeakPointer=False,LinkedTo=(MaterialGraphNode_2 25D2D0404B46764F28521E9B20B14B64,),PersistentGuid=00000000000000000000000000000000,bHidden=False,bNotConnectable=False,bDefaultValueIsReadOnly=False,bDefaultValueIsIgnored=False,bAdvancedView=False,)
+			   CustomProperties Pin (PinId=566760C74D635A14E498668CDEBA3396,PinName="Min",PinType.PinCategory="optional",PinType.PinSubCategoryObject=None,PinType.PinSubCategoryMemberReference=(),PinType.bIsArray=False,PinType.bIsReference=False,PinType.bIsConst=False,PinType.bIsWeakPointer=False,PersistentGuid=00000000000000000000000000000000,bHidden=False,bNotConnectable=False,bDefaultValueIsReadOnly=False,bDefaultValueIsIgnored=False,bAdvancedView=False,)
+			   CustomProperties Pin (PinId=5B6B6C3B48C819D0607B028ED1143585,PinName="Max",PinType.PinCategory="optional",PinType.PinSubCategoryObject=None,PinType.PinSubCategoryMemberReference=(),PinType.bIsArray=False,PinType.bIsReference=False,PinType.bIsConst=False,PinType.bIsWeakPointer=False,PersistentGuid=00000000000000000000000000000000,bHidden=False,bNotConnectable=False,bDefaultValueIsReadOnly=False,bDefaultValueIsIgnored=False,bAdvancedView=False,)
+			   CustomProperties Pin (PinId=8D4598BA460B77BFB2B4F0BA844C1932,PinName="Output",PinFriendlyName=" ",Direction="EGPD_Output",PinType.PinSubCategoryObject=None,PinType.PinSubCategoryMemberReference=(),PinType.bIsArray=False,PinType.bIsReference=False,PinType.bIsConst=False,PinType.bIsWeakPointer=False,LinkedTo=(MaterialGraphNode_13 13CBDB574A6785183151A39E8853B26E,),PersistentGuid=00000000000000000000000000000000,bHidden=False,bNotConnectable=False,bDefaultValueIsReadOnly=False,bDefaultValueIsIgnored=False,bAdvancedView=False,)
+			End Object
+			Begin Object Class=MaterialGraphNode Name="MaterialGraphNode_8"
+			   Begin Object Class=MaterialExpressionSubtract Name="MaterialExpressionSubtract_7"
+			   End Object
+			   Begin Object Name="MaterialExpressionSubtract_7"
+				  A=(Expression=MaterialExpressionConstant'MaterialGraphNode_9.MaterialExpressionConstant_2')
+				  B=(Expression=MaterialExpressionMaterialFunctionCall'MaterialGraphNode_12.MaterialExpressionMaterialFunctionCall_2')
+				  MaterialExpressionEditorX=-672
+				  MaterialExpressionEditorY=-32
+				  MaterialExpressionGuid=263301C645665622D4AF209DC4DFC1D4
+				  Material=PreviewMaterial'/Engine/Transient.ForceSplat'
+			   End Object
+			   MaterialExpression=MaterialExpressionSubtract'MaterialExpressionSubtract_7'
+			   NodePosX=-672
+			   NodePosY=-32
+			   NodeGuid=FCC50A2447D0671839BBFDAD0038A57D
+			   CustomProperties Pin (PinId=F06C53144F1B362FBAEC0AB0C16081E1,PinName="A",PinType.PinCategory="optional",PinType.PinSubCategoryObject=None,PinType.PinSubCategoryMemberReference=(),PinType.bIsArray=False,PinType.bIsReference=False,PinType.bIsConst=False,PinType.bIsWeakPointer=False,LinkedTo=(MaterialGraphNode_9 C21918BC40A135B5C229D1845162D312,),PersistentGuid=00000000000000000000000000000000,bHidden=False,bNotConnectable=False,bDefaultValueIsReadOnly=False,bDefaultValueIsIgnored=False,bAdvancedView=False,)
+			   CustomProperties Pin (PinId=CCDF13A940A6E80D85B59299C96B714F,PinName="B",PinType.PinCategory="optional",PinType.PinSubCategoryObject=None,PinType.PinSubCategoryMemberReference=(),PinType.bIsArray=False,PinType.bIsReference=False,PinType.bIsConst=False,PinType.bIsWeakPointer=False,LinkedTo=(MaterialGraphNode_12 DBCC0DEA4DCEE420BA5862AAC582B0C0,),PersistentGuid=00000000000000000000000000000000,bHidden=False,bNotConnectable=False,bDefaultValueIsReadOnly=False,bDefaultValueIsIgnored=False,bAdvancedView=False,)
+			   CustomProperties Pin (PinId=38FFC1874298113EDA7D9DBC71D7C22F,PinName="Output",PinFriendlyName=" ",Direction="EGPD_Output",PinType.PinSubCategoryObject=None,PinType.PinSubCategoryMemberReference=(),PinType.bIsArray=False,PinType.bIsReference=False,PinType.bIsConst=False,PinType.bIsWeakPointer=False,LinkedTo=(MaterialGraphNode_2 C5F733EF466287E9E4EB35ACE6B65E84,),PersistentGuid=00000000000000000000000000000000,bHidden=False,bNotConnectable=False,bDefaultValueIsReadOnly=False,bDefaultValueIsIgnored=False,bAdvancedView=False,)
+			End Object
+			Begin Object Class=MaterialGraphNode Name="MaterialGraphNode_9"
+			   Begin Object Class=MaterialExpressionConstant Name="MaterialExpressionConstant_2"
+			   End Object
+			   Begin Object Name="MaterialExpressionConstant_2"
+				  R=1.000000
+				  MaterialExpressionEditorX=-768
+				  MaterialExpressionEditorY=-64
+				  MaterialExpressionGuid=4D20573E49F05F155BF9099DE5AB86DB
+				  Material=PreviewMaterial'/Engine/Transient.ForceSplat'
+			   End Object
+			   MaterialExpression=MaterialExpressionConstant'MaterialExpressionConstant_2'
+			   NodePosX=-768
+			   NodePosY=-64
+			   NodeGuid=417D7D11440191D4AA9C28B3871ACF2E
+			   CustomProperties Pin (PinId=C21918BC40A135B5C229D1845162D312,PinName="Output",PinFriendlyName=" ",Direction="EGPD_Output",PinType.PinSubCategoryObject=None,PinType.PinSubCategoryMemberReference=(),PinType.bIsArray=False,PinType.bIsReference=False,PinType.bIsConst=False,PinType.bIsWeakPointer=False,LinkedTo=(MaterialGraphNode_8 F06C53144F1B362FBAEC0AB0C16081E1,),PersistentGuid=00000000000000000000000000000000,bHidden=False,bNotConnectable=False,bDefaultValueIsReadOnly=False,bDefaultValueIsIgnored=False,bAdvancedView=False,)
+			End Object
+			Begin Object Class=MaterialGraphNode Name="MaterialGraphNode_10"
+			   Begin Object Class=MaterialExpressionSubtract Name="MaterialExpressionSubtract_8"
+			   End Object
+			   Begin Object Name="MaterialExpressionSubtract_8"
+				  A=(Expression=MaterialExpressionTextureCoordinate'MaterialGraphNode_11.MaterialExpressionTextureCoordinate_2')
+				  B=(Expression=MaterialExpressionComponentMask'MaterialGraphNode_5.MaterialExpressionComponentMask_2')
+				  MaterialExpressionEditorX=-1088
+				  MaterialExpressionGuid=68FF50AA4D025CB5B12184BDE48D55D5
+				  Material=PreviewMaterial'/Engine/Transient.ForceSplat'
+			   End Object
+			   MaterialExpression=MaterialExpressionSubtract'MaterialExpressionSubtract_8'
+			   NodePosX=-1088
+			   NodeGuid=990564B140370AE959137EABE556BEAC
+			   CustomProperties Pin (PinId=1D0B51A44832EA486A38FBABDE408C34,PinName="A",PinType.PinCategory="optional",PinType.PinSubCategoryObject=None,PinType.PinSubCategoryMemberReference=(),PinType.bIsArray=False,PinType.bIsReference=False,PinType.bIsConst=False,PinType.bIsWeakPointer=False,LinkedTo=(MaterialGraphNode_11 8D89F6834B42C8AADDD283818E4EDC94,),PersistentGuid=00000000000000000000000000000000,bHidden=False,bNotConnectable=False,bDefaultValueIsReadOnly=False,bDefaultValueIsIgnored=False,bAdvancedView=False,)
+			   CustomProperties Pin (PinId=26815BDA44BF1141EF0D0084C2E3D993,PinName="B",PinType.PinCategory="optional",PinType.PinSubCategoryObject=None,PinType.PinSubCategoryMemberReference=(),PinType.bIsArray=False,PinType.bIsReference=False,PinType.bIsConst=False,PinType.bIsWeakPointer=False,LinkedTo=(MaterialGraphNode_5 2A7AE85A415E54152002F9BA674B7E41,),PersistentGuid=00000000000000000000000000000000,bHidden=False,bNotConnectable=False,bDefaultValueIsReadOnly=False,bDefaultValueIsIgnored=False,bAdvancedView=False,)
+			   CustomProperties Pin (PinId=CE7F469F4980CA4BA970EDA2C41EAB70,PinName="Output",PinFriendlyName=" ",Direction="EGPD_Output",PinType.PinSubCategoryObject=None,PinType.PinSubCategoryMemberReference=(),PinType.bIsArray=False,PinType.bIsReference=False,PinType.bIsConst=False,PinType.bIsWeakPointer=False,LinkedTo=(MaterialGraphNode_12 4D386898434293CC7337C5B2EADB4BF2,),PersistentGuid=00000000000000000000000000000000,bHidden=False,bNotConnectable=False,bDefaultValueIsReadOnly=False,bDefaultValueIsIgnored=False,bAdvancedView=False,)
+			End Object
+			Begin Object Class=MaterialGraphNode Name="MaterialGraphNode_11"
+			   Begin Object Class=MaterialExpressionTextureCoordinate Name="MaterialExpressionTextureCoordinate_2"
+			   End Object
+			   Begin Object Name="MaterialExpressionTextureCoordinate_2"
+				  MaterialExpressionEditorX=-1280
+				  MaterialExpressionEditorY=-48
+				  MaterialExpressionGuid=28E452A8400D1753908CE0BE4D090D8A
+				  Material=PreviewMaterial'/Engine/Transient.ForceSplat'
+			   End Object
+			   MaterialExpression=MaterialExpressionTextureCoordinate'MaterialExpressionTextureCoordinate_2'
+			   NodePosX=-1280
+			   NodePosY=-48
+			   NodeGuid=0F745D7B4DD6F4646535A8A53DE84B49
+			   CustomProperties Pin (PinId=8D89F6834B42C8AADDD283818E4EDC94,PinName="Output",PinFriendlyName=" ",Direction="EGPD_Output",PinType.PinSubCategoryObject=None,PinType.PinSubCategoryMemberReference=(),PinType.bIsArray=False,PinType.bIsReference=False,PinType.bIsConst=False,PinType.bIsWeakPointer=False,LinkedTo=(MaterialGraphNode_10 1D0B51A44832EA486A38FBABDE408C34,),PersistentGuid=00000000000000000000000000000000,bHidden=False,bNotConnectable=False,bDefaultValueIsReadOnly=False,bDefaultValueIsIgnored=False,bAdvancedView=False,)
+			End Object
+			Begin Object Class=MaterialGraphNode Name="MaterialGraphNode_12"
+			   Begin Object Class=MaterialExpressionMaterialFunctionCall Name="MaterialExpressionMaterialFunctionCall_2"
+			   End Object
+			   Begin Object Name="MaterialExpressionMaterialFunctionCall_2"
+				  MaterialFunction=MaterialFunction'/Engine/Functions/Engine_MaterialFunctions02/Utility/VectorLength.VectorLength'
+				  FunctionInputs(0)=(ExpressionInputId=E06081AD4255F67C6D398889E3CE73BA,Input=(Expression=MaterialExpressionSubtract'MaterialGraphNode_10.MaterialExpressionSubtract_8',InputName="Vector 2"))
+				  FunctionInputs(1)=(ExpressionInputId=1F411C544AE4A9ACC5F5F8A2A099B309,Input=(OutputIndex=-1,InputName="Vector 3"))
+				  FunctionOutputs(0)=(ExpressionOutputId=3C53F2D54FE1029372151E82C85F7755,Output=(OutputName="V2 Length"))
+				  FunctionOutputs(1)=(ExpressionOutputId=F186F43445EF47104CD650A242DEAECE,Output=(OutputName="V3 Length"))
+				  MaterialExpressionEditorX=-912
+				  MaterialExpressionEditorY=16
+				  MaterialExpressionGuid=CA6CCA5B4A71415750F226A88DE0B584
+				  Material=PreviewMaterial'/Engine/Transient.ForceSplat'
+				  Outputs(0)=(OutputName="V2 Length")
+				  Outputs(1)=(OutputName="V3 Length")
+			   End Object
+			   MaterialExpression=MaterialExpressionMaterialFunctionCall'MaterialExpressionMaterialFunctionCall_2'
+			   NodePosX=-912
+			   NodePosY=16
+			   NodeGuid=FFAF3FEE489F8E401BF0D3BA11C26C1D
+			   CustomProperties Pin (PinId=4D386898434293CC7337C5B2EADB4BF2,PinName="Vector 2 (V2)",PinType.PinCategory="optional",PinType.PinSubCategoryObject=None,PinType.PinSubCategoryMemberReference=(),PinType.bIsArray=False,PinType.bIsReference=False,PinType.bIsConst=False,PinType.bIsWeakPointer=False,LinkedTo=(MaterialGraphNode_10 CE7F469F4980CA4BA970EDA2C41EAB70,),PersistentGuid=00000000000000000000000000000000,bHidden=False,bNotConnectable=False,bDefaultValueIsReadOnly=False,bDefaultValueIsIgnored=False,bAdvancedView=False,)
+			   CustomProperties Pin (PinId=51F10BD04E8FF243B7963F95D51E7A0E,PinName="Vector 3 (V3)",PinType.PinCategory="optional",PinType.PinSubCategoryObject=None,PinType.PinSubCategoryMemberReference=(),PinType.bIsArray=False,PinType.bIsReference=False,PinType.bIsConst=False,PinType.bIsWeakPointer=False,PersistentGuid=00000000000000000000000000000000,bHidden=False,bNotConnectable=False,bDefaultValueIsReadOnly=False,bDefaultValueIsIgnored=False,bAdvancedView=False,)
+			   CustomProperties Pin (PinId=DBCC0DEA4DCEE420BA5862AAC582B0C0,PinName="V2 Length",Direction="EGPD_Output",PinType.PinSubCategoryObject=None,PinType.PinSubCategoryMemberReference=(),PinType.bIsArray=False,PinType.bIsReference=False,PinType.bIsConst=False,PinType.bIsWeakPointer=False,LinkedTo=(MaterialGraphNode_8 CCDF13A940A6E80D85B59299C96B714F,),PersistentGuid=00000000000000000000000000000000,bHidden=False,bNotConnectable=False,bDefaultValueIsReadOnly=False,bDefaultValueIsIgnored=False,bAdvancedView=False,)
+			   CustomProperties Pin (PinId=1EDAC5CB45C35155B74AC9AC7C968F52,PinName="V3 Length",Direction="EGPD_Output",PinType.PinSubCategoryObject=None,PinType.PinSubCategoryMemberReference=(),PinType.bIsArray=False,PinType.bIsReference=False,PinType.bIsConst=False,PinType.bIsWeakPointer=False,PersistentGuid=00000000000000000000000000000000,bHidden=False,bNotConnectable=False,bDefaultValueIsReadOnly=False,bDefaultValueIsIgnored=False,bAdvancedView=False,)
+			End Object
+			Begin Object Class=MaterialGraphNode Name="MaterialGraphNode_13"
+			   Begin Object Class=MaterialExpressionDivide Name="MaterialExpressionDivide_2"
+			   End Object
+			   Begin Object Name="MaterialExpressionDivide_2"
+				  A=(Expression=MaterialExpressionClamp'MaterialGraphNode_7.MaterialExpressionClamp_2')
+				  B=(Expression=MaterialExpressionScalarParameter'MaterialGraphNode_4.MaterialExpressionScalarParameter_5')
+				  MaterialExpressionEditorX=-80
+				  MaterialExpressionEditorY=80
+				  MaterialExpressionGuid=8730CE28486A7A975F06AAAF76360BDE
+				  Material=PreviewMaterial'/Engine/Transient.ForceSplat'
+			   End Object
+			   MaterialExpression=MaterialExpressionDivide'MaterialExpressionDivide_2'
+			   NodePosX=-80
+			   NodePosY=80
+			   NodeGuid=77DB1B5D4C0BCEC4D991518C010F05E1
+			   CustomProperties Pin (PinId=13CBDB574A6785183151A39E8853B26E,PinName="A",PinType.PinCategory="optional",PinType.PinSubCategoryObject=None,PinType.PinSubCategoryMemberReference=(),PinType.bIsArray=False,PinType.bIsReference=False,PinType.bIsConst=False,PinType.bIsWeakPointer=False,LinkedTo=(MaterialGraphNode_7 8D4598BA460B77BFB2B4F0BA844C1932,),PersistentGuid=00000000000000000000000000000000,bHidden=False,bNotConnectable=False,bDefaultValueIsReadOnly=False,bDefaultValueIsIgnored=False,bAdvancedView=False,)
+			   CustomProperties Pin (PinId=3501ECFD4967CB5C6FA143BB453E5636,PinName="B",PinType.PinCategory="optional",PinType.PinSubCategoryObject=None,PinType.PinSubCategoryMemberReference=(),PinType.bIsArray=False,PinType.bIsReference=False,PinType.bIsConst=False,PinType.bIsWeakPointer=False,LinkedTo=(MaterialGraphNode_4 243D5A9F4A177E21AD5B0C9A9EE13B6C,),PersistentGuid=00000000000000000000000000000000,bHidden=False,bNotConnectable=False,bDefaultValueIsReadOnly=False,bDefaultValueIsIgnored=False,bAdvancedView=False,)
+			   CustomProperties Pin (PinId=D0CE6B2B4A6886F59532E4B83367D7BD,PinName="Output",PinFriendlyName=" ",Direction="EGPD_Output",PinType.PinSubCategoryObject=None,PinType.PinSubCategoryMemberReference=(),PinType.bIsArray=False,PinType.bIsReference=False,PinType.bIsConst=False,PinType.bIsWeakPointer=False,LinkedTo=(MaterialGraphNode_0 B7D0CF28498836478F9C8987312501BF,),PersistentGuid=00000000000000000000000000000000,bHidden=False,bNotConnectable=False,bDefaultValueIsReadOnly=False,bDefaultValueIsIgnored=False,bAdvancedView=False,)
+			End Object
+
+	[/PARAM]
+[/OBJECT]
+[/EXCERPT:ForceSplat]
+-->
+
+
+
+
+
+
+
+
+
+
+
+
+
+
+
+
+
+
+
+
+
+
+
+
+
+
+
+
+[NAV]
\ No newline at end of file
