Copying //UE4/Dev-Documentation to Samples-Main (//UE4/Samples-Main)

diff --git a/Engine/Documentation/Source/Engine/Rendering/RenderTargets/BlueprintRenderTargets/HowTo/HeightFieldPainter/3/HFP3.INT.udn b/Engine/Documentation/Source/Engine/Rendering/RenderTargets/BlueprintRenderTargets/HowTo/HeightFieldPainter/3/HFP3.INT.udn
new file mode 100644
--- /dev/null
+++ b/Engine/Documentation/Source/Engine/Rendering/RenderTargets/BlueprintRenderTargets/HowTo/HeightFieldPainter/3/HFP3.INT.udn
@@ -0,0 +1,848 @@
+availability:docs
+title:3. Blueprint Setup
+description:Going over creating and setting up the height map painter Blueprint.
+crumbs:%ROOT%, Engine, Engine\Rendering\RenderTargets
+version:4.13
+type:multi-step
+SkillLevel:Intermediate
+parent:Engine/Rendering/RenderTargets/BlueprintRenderTargets/HowTo/HeightFieldPainter
+Navigation:multi-step
+tags: Render Targets
+order:3
+checkpoint:HFP
+
+[NAV]
+
+With the render targets and needed Materials setup and ready to use, it is now time to bring everything together using a Blueprint. In the following section, we will go over how to setup all of the various parts of the Height Field Painter Blueprint.
+
+## Blueprint Creation & Variable Setup
+
+1. Right-click in the Content Browser and from the menu in the **Create Basic Asset** section, click on the **Blueprint Class** option.
+	
+	![](Engine\Rendering\RenderTargets\BlueprintRenderTargets\HowTo\HeightFieldPainter\HFP_BP_Setup_00.png)
+
+1. From the **Pick Parent Class** window, select **Actor** and name it `HeightFieldPainter`.
+
+	![](Engine\Rendering\RenderTargets\BlueprintRenderTargets\HowTo\HeightFieldPainter\HFP_BP_Setup_01.png)
+
+1. Double-click on the HeightFieldPainter Blueprint to open it up and then click on the **Event Graph** tab. When the Event Graph is open, add the following variables.
+	
+	![](Engine\Rendering\RenderTargets\BlueprintRenderTargets\HowTo\HeightFieldPainter\HFP_BP_Setup_02.png)
+
+	|Variable Name| Variable Type|Default Value|
+	|--------|------------|---|
+	|HeightfieldRT|Texture Render Target 2D|N/A|
+	|PainterMaterialInstance|Material Instance Dynamic|N/A|
+	|MouseDown|Bool|N/A|
+	|ForceSplatMID|Material Instance Dynamic|N/A|
+	|InteractionDistance|Float|1000.0|
+	|BrushSize|Float|0.04|
+	|BrushStrength|Float|1.0|
+
+## Components Setup
+Inside of the **Components** tab, press the **Add Component** button and add a **Static Mesh** Component to the Blueprint. Once the Static Mesh Component has been added, select the Static Mesh Component and setup the following options. 
+
+![](Engine\Rendering\RenderTargets\BlueprintRenderTargets\HowTo\HeightFieldPainter\HFP_BP_Setup_06.png)
+
+|Section|Property| Required Asset|Value|
+|--------|--------|------------|---|
+|Transform|Mobility|N/A|Movable|
+|Static Mesh|Static Mesh|SM_Plane1000_512|N/A|
+|Materials|Material|HeightfieldPainterMaterial|N/A|
+|Collision|Generate Overlap Events|N/A|Checked|
+|Collision|Collision Presets|N/A|OverlapOnlyPawn|
+
+## Construction Script Setup	
+The Construction script is used to create and assign a 2D render target in addition to creating and assigning the required Dynamic Material Instances. If you arean't already there, click on the **Construction Script** tab and then copy the following code into your Blueprints Construction Script. 
+
+[REGION:fullwidth]
+[INCLUDE:Engine\Rendering\RenderTargets\BlueprintRenderTargets\HowTo\HeightFieldPainter\3\#ConScript]
+[/REGION]
+
+## Trace from Camera Function Setup
+The Trace from Camera function will be used to determine the location on the Static Mesh that should have its height increased. It determines this location by casting a ray from the center of the player's camera into the world. When that ray intersects with a Static Mesh, the location that it intersects will be the location that is raised. To set up the `Trace from Camera` function, you will need to do the following.
+
+1. Click the small (white) plus sign on the **Functions** tab to create a new function and name it **TraceFromCamera**.
+
+	![](Engine\Rendering\RenderTargets\BlueprintRenderTargets\HowTo\HeightFieldPainter\HFP_BP_Setup_04.png)
+
+1. Open the TraceFromCamera function and copy the following Blueprint code into your Blueprint by hovering your mouse over the image below and then clicking on the **Copy Expression Graph** option, selecting the code, copying it, and then pasting it into the TraceFromCamera function.
+
+	[REGION:fullwidth]
+	[INCLUDE:Engine\Rendering\RenderTargets\BlueprintRenderTargets\HowTo\HeightFieldPainter\3\#Function]
+	[/REGION]
+
+	[region:note]
+	Do not forget to Compile and Save your Blueprint so that it can be used.
+	[/region]
+
+## Event Graph Setup
+Now that all of the need parts have been created and setup, it is now time to connect all of the various parts together in the Event Graph. You can copy the completed Blueprint code into your Blueprint by hovering your mouse over the image below and then clicking on the  **Copy Node Graph** option, copying-and-pasting the provided Blueprint code into your Height Field Painter Blueprint.
+
+
+[REGION:fullwidth]
+[INCLUDE:Engine\Rendering\RenderTargets\BlueprintRenderTargets\HowTo\HeightFieldPainter\3\#EG]
+[/REGION]
+
+[region:note]
+When you copy-paste the Blueprint code, the **Event Begin Play** node will not be added to the graph. Make sure that you add the Event Begin Play node and connect it to the **Enable Input** input. 
+
+![](Engine\Rendering\RenderTargets\BlueprintRenderTargets\HowTo\HeightFieldPainter\HFP_BP_Setup_08.png)
+[/region]
+
+
+Now that the required Blueprints have been setup, in the next section we will go over how to use the height map painter inside UE4.
+
+[NAV]
+
+
+<!--
+[EXCERPT:Function]	
+[OBJECT:Snippet]
+	[PARAMLITERAL:id]
+		TraceFromCamera
+	[/PARAMLITERAL]
+	[PARAMLITERAL:type]
+		blueprint
+	[/PARAMLITERAL]
+	[PARAM:image]
+		![](Engine\Rendering\RenderTargets\BlueprintRenderTargets\HowTo\HeightFieldPainter\HFP_BP_Setup_05.png)
+	[/PARAM]
+	[PARAM:snippet]
+				Begin Object Class=K2Node_FunctionEntry Name="K2Node_FunctionEntry_9"
+			   ExtraFlags=201457664
+			   SignatureName="TraceFromCamera"
+			   bIsEditable=True
+			   NodePosX=-96
+			   NodePosY=-384
+			   NodeGuid=A6585C2E45244BF6E4B02AAADF163432
+			   CustomProperties Pin (PinId=598C3D6A46B70A76463B448AA52DE23A,PinName="then",Direction="EGPD_Output",PinType.PinCategory="exec",PinType.PinSubCategoryObject=None,PinType.PinSubCategoryMemberReference=(),PinType.bIsArray=False,PinType.bIsReference=False,PinType.bIsConst=False,PinType.bIsWeakPointer=False,LinkedTo=(K2Node_CallFunction_65598 B3AD098C4E9F315318BE2C9A781CFD16,),PersistentGuid=00000000000000000000000000000000,bHidden=False,bNotConnectable=False,bDefaultValueIsReadOnly=False,bDefaultValueIsIgnored=False,bAdvancedView=False,)
+			End Object
+			Begin Object Class=K2Node_MacroInstance Name="K2Node_MacroInstance_65"
+			   MacroGraphReference=(MacroGraph=EdGraph'/Engine/EditorBlueprintResources/StandardMacros.StandardMacros:IsValid',GraphBlueprint=Blueprint'/Engine/EditorBlueprintResources/StandardMacros.StandardMacros',GraphGuid=64422BCD430703FF5CAEA8B79A32AA65)
+			   NodePosX=1984
+			   NodePosY=-432
+			   ErrorType=4
+			   ErrorMsg="Info The execution path doesn\'t end with a return node.  Is Not Valid "
+			   NodeGuid=269E6BC649B1DF4406573EB2C13D874A
+			   CustomProperties Pin (PinId=54DCDF6A4C18DFFBF71F54AE24DA40C0,PinName="exec",PinType.PinCategory="exec",PinType.PinSubCategoryObject=None,PinType.PinSubCategoryMemberReference=(),PinType.bIsArray=False,PinType.bIsReference=False,PinType.bIsConst=False,PinType.bIsWeakPointer=False,LinkedTo=(K2Node_IfThenElse_161 8DD9A5BC40D019A1AFD6F284F8C15FFD,),PersistentGuid=00000000000000000000000000000000,bHidden=False,bNotConnectable=False,bDefaultValueIsReadOnly=False,bDefaultValueIsIgnored=False,bAdvancedView=False,)
+			   CustomProperties Pin (PinId=C3FB28AC47AAD2B5998C2AA9D4862C2E,PinName="InputObject",PinType.PinCategory="object",PinType.PinSubCategoryObject=Class'/Script/CoreUObject.Object',PinType.PinSubCategoryMemberReference=(),PinType.bIsArray=False,PinType.bIsReference=False,PinType.bIsConst=False,PinType.bIsWeakPointer=False,LinkedTo=(K2Node_CallFunction_65759 F1D6985542953219C7EF09A3FBFA8AF0,),PersistentGuid=00000000000000000000000000000000,bHidden=False,bNotConnectable=False,bDefaultValueIsReadOnly=False,bDefaultValueIsIgnored=False,bAdvancedView=False,)
+			   CustomProperties Pin (PinId=E75DA26F4EBE18C63E3CEDA6B1D88DC6,PinName="Is Valid",Direction="EGPD_Output",PinType.PinCategory="exec",PinType.PinSubCategoryObject=None,PinType.PinSubCategoryMemberReference=(),PinType.bIsArray=False,PinType.bIsReference=False,PinType.bIsConst=False,PinType.bIsWeakPointer=False,LinkedTo=(K2Node_CallFunction_1861 6F410C0B46D7B07315AD748B03EB86C8,),PersistentGuid=00000000000000000000000000000000,bHidden=False,bNotConnectable=False,bDefaultValueIsReadOnly=False,bDefaultValueIsIgnored=False,bAdvancedView=False,)
+			   CustomProperties Pin (PinId=7426FEC54F3C31FE985AE088ACC9F56E,PinName="Is Not Valid",Direction="EGPD_Output",PinType.PinCategory="exec",PinType.PinSubCategoryObject=None,PinType.PinSubCategoryMemberReference=(),PinType.bIsArray=False,PinType.bIsReference=False,PinType.bIsConst=False,PinType.bIsWeakPointer=False,PersistentGuid=00000000000000000000000000000000,bHidden=False,bNotConnectable=False,bDefaultValueIsReadOnly=False,bDefaultValueIsIgnored=False,bAdvancedView=False,)
+			End Object
+			Begin Object Class=K2Node_CallFunction Name="K2Node_CallFunction_2902"
+			   bIsPureFunc=True
+			   FunctionReference=(MemberParent=Class'/Script/Engine.GameplayStatics',MemberName="GetPlayerController")
+			   NodePosX=2192
+			   NodePosY=-176
+			   NodeGuid=5089AE9B44A43150C97C5F92ECE42A73
+			   CustomProperties Pin (PinId=45DD77B44C0471A5CDD2C7BB7D1F8F28,PinName="self",PinFriendlyName=NSLOCTEXT("K2Node", "Target", "Target"),PinType.PinCategory="object",PinType.PinSubCategoryObject=Class'/Script/Engine.GameplayStatics',PinType.PinSubCategoryMemberReference=(),PinType.bIsArray=False,PinType.bIsReference=False,PinType.bIsConst=False,PinType.bIsWeakPointer=False,DefaultObject="/Script/Engine.Default__GameplayStatics",PersistentGuid=00000000000000000000000000000000,bHidden=True,bNotConnectable=False,bDefaultValueIsReadOnly=False,bDefaultValueIsIgnored=False,bAdvancedView=False,)
+			   CustomProperties Pin (PinId=56BF7CB143F1FAE6723E31A1D3FC3447,PinName="WorldContextObject",PinType.PinCategory="object",PinType.PinSubCategoryObject=Class'/Script/CoreUObject.Object',PinType.PinSubCategoryMemberReference=(),PinType.bIsArray=False,PinType.bIsReference=False,PinType.bIsConst=False,PinType.bIsWeakPointer=False,PersistentGuid=00000000000000000000000000000000,bHidden=True,bNotConnectable=False,bDefaultValueIsReadOnly=False,bDefaultValueIsIgnored=False,bAdvancedView=False,)
+			   CustomProperties Pin (PinId=6AA23C3242491D37AACF8CB51ECB365D,PinName="PlayerIndex",PinType.PinCategory="int",PinType.PinSubCategoryObject=None,PinType.PinSubCategoryMemberReference=(),PinType.bIsArray=False,PinType.bIsReference=False,PinType.bIsConst=False,PinType.bIsWeakPointer=False,DefaultValue="0",AutogeneratedDefaultValue="0",PersistentGuid=00000000000000000000000000000000,bHidden=False,bNotConnectable=False,bDefaultValueIsReadOnly=False,bDefaultValueIsIgnored=False,bAdvancedView=False,)
+			   CustomProperties Pin (PinId=3119A2234CC8CC2E58EDDA94ED78DB30,PinName="ReturnValue",Direction="EGPD_Output",PinType.PinCategory="object",PinType.PinSubCategoryObject=Class'/Script/Engine.PlayerController',PinType.PinSubCategoryMemberReference=(),PinType.bIsArray=False,PinType.bIsReference=False,PinType.bIsConst=False,PinType.bIsWeakPointer=False,LinkedTo=(K2Node_CallFunction_1861 8AAA25E3411AD93065E153BDF12317E1,),PersistentGuid=00000000000000000000000000000000,bHidden=False,bNotConnectable=False,bDefaultValueIsReadOnly=False,bDefaultValueIsIgnored=False,bAdvancedView=False,)
+			End Object
+			Begin Object Class=K2Node_CallFunction Name="K2Node_CallFunction_64853"
+			   bIsPureFunc=True
+			   FunctionReference=(MemberParent=Class'/Script/Engine.GameplayStatics',MemberName="GetPlayerCameraManager")
+			   NodePosX=-112
+			   NodePosY=-176
+			   NodeGuid=8798F1F24424791656D634AE744FBB71
+			   CustomProperties Pin (PinId=73A2796F477586504925ACB63B1C3373,PinName="self",PinFriendlyName=NSLOCTEXT("K2Node", "Target", "Target"),PinType.PinCategory="object",PinType.PinSubCategoryObject=Class'/Script/Engine.GameplayStatics',PinType.PinSubCategoryMemberReference=(),PinType.bIsArray=False,PinType.bIsReference=False,PinType.bIsConst=False,PinType.bIsWeakPointer=False,DefaultObject="/Script/Engine.Default__GameplayStatics",PersistentGuid=00000000000000000000000000000000,bHidden=True,bNotConnectable=False,bDefaultValueIsReadOnly=False,bDefaultValueIsIgnored=False,bAdvancedView=False,)
+			   CustomProperties Pin (PinId=A29AA4884E79B1B2CBFCAE93A185BFA1,PinName="WorldContextObject",PinType.PinCategory="object",PinType.PinSubCategoryObject=Class'/Script/CoreUObject.Object',PinType.PinSubCategoryMemberReference=(),PinType.bIsArray=False,PinType.bIsReference=False,PinType.bIsConst=False,PinType.bIsWeakPointer=False,DefaultValue="0",PersistentGuid=00000000000000000000000000000000,bHidden=True,bNotConnectable=False,bDefaultValueIsReadOnly=False,bDefaultValueIsIgnored=False,bAdvancedView=False,)
+			   CustomProperties Pin (PinId=F04D1A5446F837775E768082313D9559,PinName="PlayerIndex",PinType.PinCategory="int",PinType.PinSubCategoryObject=None,PinType.PinSubCategoryMemberReference=(),PinType.bIsArray=False,PinType.bIsReference=False,PinType.bIsConst=False,PinType.bIsWeakPointer=False,DefaultValue="0",AutogeneratedDefaultValue="0",PersistentGuid=00000000000000000000000000000000,bHidden=False,bNotConnectable=False,bDefaultValueIsReadOnly=False,bDefaultValueIsIgnored=False,bAdvancedView=False,)
+			   CustomProperties Pin (PinId=88B4C49D414D557C8728F39A3967490F,PinName="ReturnValue",Direction="EGPD_Output",PinType.PinCategory="object",PinType.PinSubCategoryObject=Class'/Script/Engine.PlayerCameraManager',PinType.PinSubCategoryMemberReference=(),PinType.bIsArray=False,PinType.bIsReference=False,PinType.bIsConst=False,PinType.bIsWeakPointer=False,LinkedTo=(K2Node_CallFunction_64919 136B71094B407B9186881DB137E50CB0,K2Node_CallFunction_65094 451585FC4960B7F55857FDB8A5A37392,),PersistentGuid=00000000000000000000000000000000,bHidden=False,bNotConnectable=False,bDefaultValueIsReadOnly=False,bDefaultValueIsIgnored=False,bAdvancedView=False,)
+			End Object
+			Begin Object Class=K2Node_CallFunction Name="K2Node_CallFunction_64919"
+			   bIsPureFunc=True
+			   bIsConstFunc=True
+			   FunctionReference=(MemberName="K2_GetActorLocation",bSelfContext=True)
+			   NodePosX=560
+			   NodePosY=-191
+			   NodeGuid=B1D294AF40FA8CB679B6D4BED0F11D23
+			   CustomProperties Pin (PinId=136B71094B407B9186881DB137E50CB0,PinName="self",PinFriendlyName=NSLOCTEXT("K2Node", "Target", "Target"),PinType.PinCategory="object",PinType.PinSubCategoryObject=Class'/Script/Engine.Actor',PinType.PinSubCategoryMemberReference=(),PinType.bIsArray=False,PinType.bIsReference=False,PinType.bIsConst=False,PinType.bIsWeakPointer=False,LinkedTo=(K2Node_CallFunction_64853 88B4C49D414D557C8728F39A3967490F,),PersistentGuid=00000000000000000000000000000000,bHidden=False,bNotConnectable=False,bDefaultValueIsReadOnly=False,bDefaultValueIsIgnored=False,bAdvancedView=False,)
+			   CustomProperties Pin (PinId=7C3AB90A4E8F92F9FAE0709612FDD604,PinName="ReturnValue",Direction="EGPD_Output",PinType.PinCategory="struct",PinType.PinSubCategoryObject=ScriptStruct'/Script/CoreUObject.Vector',PinType.PinSubCategoryMemberReference=(),PinType.bIsArray=False,PinType.bIsReference=False,PinType.bIsConst=False,PinType.bIsWeakPointer=False,DefaultValue="0, 0, 0",AutogeneratedDefaultValue="0, 0, 0",LinkedTo=(K2Node_CallFunction_65598 85DD87AD4631A8DC79F17588BD6ADD07,K2Node_CommutativeAssociativeBinaryOperator_1357 C2F85EED4A8F818C3CE355A3C56691F5,),PersistentGuid=00000000000000000000000000000000,bHidden=False,bNotConnectable=False,bDefaultValueIsReadOnly=False,bDefaultValueIsIgnored=False,bAdvancedView=False,)
+			End Object
+			Begin Object Class=K2Node_CallFunction Name="K2Node_CallFunction_65094"
+			   bIsPureFunc=True
+			   bIsConstFunc=True
+			   FunctionReference=(MemberName="K2_GetActorRotation",bSelfContext=True)
+			   NodePosX=176
+			   NodePosY=-191
+			   NodeGuid=B7237EBA447AD7195212088A6AF51A87
+			   CustomProperties Pin (PinId=451585FC4960B7F55857FDB8A5A37392,PinName="self",PinFriendlyName=NSLOCTEXT("K2Node", "Target", "Target"),PinType.PinCategory="object",PinType.PinSubCategoryObject=Class'/Script/Engine.Actor',PinType.PinSubCategoryMemberReference=(),PinType.bIsArray=False,PinType.bIsReference=False,PinType.bIsConst=False,PinType.bIsWeakPointer=False,LinkedTo=(K2Node_CallFunction_64853 88B4C49D414D557C8728F39A3967490F,),PersistentGuid=00000000000000000000000000000000,bHidden=False,bNotConnectable=False,bDefaultValueIsReadOnly=False,bDefaultValueIsIgnored=False,bAdvancedView=False,)
+			   CustomProperties Pin (PinId=3048B8A64F99A2395A2A78B01F534663,PinName="ReturnValue",Direction="EGPD_Output",PinType.PinCategory="struct",PinType.PinSubCategoryObject=ScriptStruct'/Script/CoreUObject.Rotator',PinType.PinSubCategoryMemberReference=(),PinType.bIsArray=False,PinType.bIsReference=False,PinType.bIsConst=False,PinType.bIsWeakPointer=False,DefaultValue="0, 0, 0",AutogeneratedDefaultValue="0, 0, 0",LinkedTo=(K2Node_CallFunction_65130 44D474914E7F590449BEE0850F8024CE,),PersistentGuid=00000000000000000000000000000000,bHidden=False,bNotConnectable=False,bDefaultValueIsReadOnly=False,bDefaultValueIsIgnored=False,bAdvancedView=False,)
+			End Object
+			Begin Object Class=K2Node_CallFunction Name="K2Node_CallFunction_65130"
+			   bIsPureFunc=True
+			   FunctionReference=(MemberParent=Class'/Script/Engine.KismetMathLibrary',MemberName="Conv_RotatorToVector")
+			   NodePosX=384
+			   NodePosY=-64
+			   NodeGuid=43732C9F46DD928BC6B78F92F47C8A12
+			   CustomProperties Pin (PinId=06AE3D444E2D6A5652B7C291690B1A65,PinName="self",PinFriendlyName=NSLOCTEXT("K2Node", "Target", "Target"),PinType.PinCategory="object",PinType.PinSubCategoryObject=Class'/Script/Engine.KismetMathLibrary',PinType.PinSubCategoryMemberReference=(),PinType.bIsArray=False,PinType.bIsReference=False,PinType.bIsConst=False,PinType.bIsWeakPointer=False,DefaultObject="/Script/Engine.Default__KismetMathLibrary",PersistentGuid=00000000000000000000000000000000,bHidden=True,bNotConnectable=False,bDefaultValueIsReadOnly=False,bDefaultValueIsIgnored=False,bAdvancedView=False,)
+			   CustomProperties Pin (PinId=44D474914E7F590449BEE0850F8024CE,PinName="InRot",PinType.PinCategory="struct",PinType.PinSubCategoryObject=ScriptStruct'/Script/CoreUObject.Rotator',PinType.PinSubCategoryMemberReference=(),PinType.bIsArray=False,PinType.bIsReference=False,PinType.bIsConst=False,PinType.bIsWeakPointer=False,DefaultValue="0, 0, 0",AutogeneratedDefaultValue="0, 0, 0",LinkedTo=(K2Node_CallFunction_65094 3048B8A64F99A2395A2A78B01F534663,),PersistentGuid=00000000000000000000000000000000,bHidden=False,bNotConnectable=False,bDefaultValueIsReadOnly=False,bDefaultValueIsIgnored=False,bAdvancedView=False,)
+			   CustomProperties Pin (PinId=6A4437574CCDB781A7E425AF26B29860,PinName="ReturnValue",Direction="EGPD_Output",PinType.PinCategory="struct",PinType.PinSubCategoryObject=ScriptStruct'/Script/CoreUObject.Vector',PinType.PinSubCategoryMemberReference=(),PinType.bIsArray=False,PinType.bIsReference=False,PinType.bIsConst=False,PinType.bIsWeakPointer=False,DefaultValue="0, 0, 0",AutogeneratedDefaultValue="0, 0, 0",LinkedTo=(K2Node_CallFunction_65459 CE2B40EF47E2476D00764EA442CB2F0D,),PersistentGuid=00000000000000000000000000000000,bHidden=False,bNotConnectable=False,bDefaultValueIsReadOnly=False,bDefaultValueIsIgnored=False,bAdvancedView=False,)
+			End Object
+			Begin Object Class=K2Node_CallFunction Name="K2Node_CallFunction_65349"
+			   bIsPureFunc=True
+			   FunctionReference=(MemberParent=Class'/Script/Engine.KismetMathLibrary',MemberName="Multiply_VectorFloat")
+			   NodePosX=784
+			   NodePosY=-16
+			   NodeGuid=9ADED0674817535FD7003DB28601EACF
+			   CustomProperties Pin (PinId=ABB067854B9FDB401D7AF2A71639E8A9,PinName="self",PinFriendlyName=NSLOCTEXT("K2Node", "Target", "Target"),PinType.PinCategory="object",PinType.PinSubCategoryObject=Class'/Script/Engine.KismetMathLibrary',PinType.PinSubCategoryMemberReference=(),PinType.bIsArray=False,PinType.bIsReference=False,PinType.bIsConst=False,PinType.bIsWeakPointer=False,DefaultObject="/Script/Engine.Default__KismetMathLibrary",PersistentGuid=00000000000000000000000000000000,bHidden=True,bNotConnectable=False,bDefaultValueIsReadOnly=False,bDefaultValueIsIgnored=False,bAdvancedView=False,)
+			   CustomProperties Pin (PinId=AE6D33C2476153AC065E9095CB766952,PinName="A",PinType.PinCategory="struct",PinType.PinSubCategoryObject=ScriptStruct'/Script/CoreUObject.Vector',PinType.PinSubCategoryMemberReference=(),PinType.bIsArray=False,PinType.bIsReference=False,PinType.bIsConst=False,PinType.bIsWeakPointer=False,DefaultValue="0, 0, 0",AutogeneratedDefaultValue="0, 0, 0",LinkedTo=(K2Node_CallFunction_65459 182E27F347849D96EFDA44BC8BAF228B,),PersistentGuid=00000000000000000000000000000000,bHidden=False,bNotConnectable=False,bDefaultValueIsReadOnly=False,bDefaultValueIsIgnored=False,bAdvancedView=False,)
+			   CustomProperties Pin (PinId=2C5718584516424BB52CF6A79E702026,PinName="B",PinType.PinCategory="float",PinType.PinSubCategoryObject=None,PinType.PinSubCategoryMemberReference=(),PinType.bIsArray=False,PinType.bIsReference=False,PinType.bIsConst=False,PinType.bIsWeakPointer=False,DefaultValue="0.0",AutogeneratedDefaultValue="0.0",LinkedTo=(K2Node_CommutativeAssociativeBinaryOperator_4227 0FAF44A34397D0F87EB02B88E8A653E3,),PersistentGuid=00000000000000000000000000000000,bHidden=False,bNotConnectable=False,bDefaultValueIsReadOnly=False,bDefaultValueIsIgnored=False,bAdvancedView=False,)
+			   CustomProperties Pin (PinId=40870E1F404EA24CCC3295A98D8A13B6,PinName="ReturnValue",Direction="EGPD_Output",PinType.PinCategory="struct",PinType.PinSubCategoryObject=ScriptStruct'/Script/CoreUObject.Vector',PinType.PinSubCategoryMemberReference=(),PinType.bIsArray=False,PinType.bIsReference=False,PinType.bIsConst=False,PinType.bIsWeakPointer=False,DefaultValue="0, 0, 0",AutogeneratedDefaultValue="0, 0, 0",LinkedTo=(K2Node_CommutativeAssociativeBinaryOperator_1357 4410C89D42A350AA1AF9EE830D6A174E,),PersistentGuid=00000000000000000000000000000000,bHidden=False,bNotConnectable=False,bDefaultValueIsReadOnly=False,bDefaultValueIsIgnored=False,bAdvancedView=False,)
+			End Object
+			Begin Object Class=K2Node_CallFunction Name="K2Node_CallFunction_65459"
+			   bIsPureFunc=True
+			   FunctionReference=(MemberParent=Class'/Script/Engine.KismetMathLibrary',MemberName="Normal")
+			   NodePosX=592
+			   NodePosY=-64
+			   NodeGuid=8F88CD2F4F7D9006417C01B2332B7C39
+			   CustomProperties Pin (PinId=E124F3A24616014EE593778FBFC7A74F,PinName="self",PinFriendlyName=NSLOCTEXT("K2Node", "Target", "Target"),PinToolTip="Target\nKismet Math Library Reference",PinType.PinCategory="object",PinType.PinSubCategoryObject=Class'/Script/Engine.KismetMathLibrary',PinType.PinSubCategoryMemberReference=(),PinType.bIsArray=False,PinType.bIsReference=False,PinType.bIsConst=False,PinType.bIsWeakPointer=False,DefaultObject="/Script/Engine.Default__KismetMathLibrary",PersistentGuid=00000000000000000000000000000000,bHidden=True,bNotConnectable=False,bDefaultValueIsReadOnly=False,bDefaultValueIsIgnored=False,bAdvancedView=False,)
+			   CustomProperties Pin (PinId=CE2B40EF47E2476D00764EA442CB2F0D,PinName="A",PinToolTip="A\nVector ",PinType.PinCategory="struct",PinType.PinSubCategoryObject=ScriptStruct'/Script/CoreUObject.Vector',PinType.PinSubCategoryMemberReference=(),PinType.bIsArray=False,PinType.bIsReference=False,PinType.bIsConst=False,PinType.bIsWeakPointer=False,DefaultValue="0, 0, 0",AutogeneratedDefaultValue="0, 0, 0",LinkedTo=(K2Node_CallFunction_65130 6A4437574CCDB781A7E425AF26B29860,),PersistentGuid=00000000000000000000000000000000,bHidden=False,bNotConnectable=False,bDefaultValueIsReadOnly=False,bDefaultValueIsIgnored=False,bAdvancedView=False,)
+			   CustomProperties Pin (PinId=182E27F347849D96EFDA44BC8BAF228B,PinName="ReturnValue",PinToolTip="Return Value\nVector ",Direction="EGPD_Output",PinType.PinCategory="struct",PinType.PinSubCategoryObject=ScriptStruct'/Script/CoreUObject.Vector',PinType.PinSubCategoryMemberReference=(),PinType.bIsArray=False,PinType.bIsReference=False,PinType.bIsConst=False,PinType.bIsWeakPointer=False,DefaultValue="0, 0, 0",AutogeneratedDefaultValue="0, 0, 0",LinkedTo=(K2Node_CallFunction_65349 AE6D33C2476153AC065E9095CB766952,),PersistentGuid=00000000000000000000000000000000,bHidden=False,bNotConnectable=False,bDefaultValueIsReadOnly=False,bDefaultValueIsIgnored=False,bAdvancedView=False,)
+			End Object
+			Begin Object Class=K2Node_CallFunction Name="K2Node_CallFunction_65598"
+			   FunctionReference=(MemberParent=Class'/Script/Engine.KismetSystemLibrary',MemberName="LineTraceSingle_NEW")
+			   NodePosX=1104
+			   NodePosY=-384
+			   NodeGuid=8BE120E14A1ACA5DD4C51B809102B477
+			   AdvancedPinDisplay=Hidden
+			   CustomProperties Pin (PinId=B3AD098C4E9F315318BE2C9A781CFD16,PinName="execute",PinToolTip="\nExec",PinType.PinCategory="exec",PinType.PinSubCategoryObject=None,PinType.PinSubCategoryMemberReference=(),PinType.bIsArray=False,PinType.bIsReference=False,PinType.bIsConst=False,PinType.bIsWeakPointer=False,LinkedTo=(K2Node_FunctionEntry_9 598C3D6A46B70A76463B448AA52DE23A,),PersistentGuid=00000000000000000000000000000000,bHidden=False,bNotConnectable=False,bDefaultValueIsReadOnly=False,bDefaultValueIsIgnored=False,bAdvancedView=False,)
+			   CustomProperties Pin (PinId=E001912C48BB5BEC656460843FE151FA,PinName="then",PinToolTip="\nExec",Direction="EGPD_Output",PinType.PinCategory="exec",PinType.PinSubCategoryObject=None,PinType.PinSubCategoryMemberReference=(),PinType.bIsArray=False,PinType.bIsReference=False,PinType.bIsConst=False,PinType.bIsWeakPointer=False,LinkedTo=(K2Node_IfThenElse_161 66D340DC4A52DC8D727C43ABA6E818C4,),PersistentGuid=00000000000000000000000000000000,bHidden=False,bNotConnectable=False,bDefaultValueIsReadOnly=False,bDefaultValueIsIgnored=False,bAdvancedView=False,)
+			   CustomProperties Pin (PinId=941FA0A3430F4462DB5E30B8A5A81DF4,PinName="self",PinFriendlyName=NSLOCTEXT("K2Node", "Target", "Target"),PinToolTip="Target\nKismet System Library Reference",PinType.PinCategory="object",PinType.PinSubCategoryObject=Class'/Script/Engine.KismetSystemLibrary',PinType.PinSubCategoryMemberReference=(),PinType.bIsArray=False,PinType.bIsReference=False,PinType.bIsConst=False,PinType.bIsWeakPointer=False,DefaultObject="/Script/Engine.Default__KismetSystemLibrary",PersistentGuid=00000000000000000000000000000000,bHidden=True,bNotConnectable=False,bDefaultValueIsReadOnly=False,bDefaultValueIsIgnored=False,bAdvancedView=False,)
+			   CustomProperties Pin (PinId=6F242C374EDACE5FC62C50B118086553,PinName="WorldContextObject",PinToolTip="World Context Object\nObject Reference",PinType.PinCategory="object",PinType.PinSubCategoryObject=Class'/Script/CoreUObject.Object',PinType.PinSubCategoryMemberReference=(),PinType.bIsArray=False,PinType.bIsReference=False,PinType.bIsConst=False,PinType.bIsWeakPointer=False,DefaultValue="0",PersistentGuid=00000000000000000000000000000000,bHidden=True,bNotConnectable=False,bDefaultValueIsReadOnly=False,bDefaultValueIsIgnored=False,bAdvancedView=False,)
+			   CustomProperties Pin (PinId=85DD87AD4631A8DC79F17588BD6ADD07,PinName="Start",PinToolTip="Start\nVector \n\nStart of line segment.",PinType.PinCategory="struct",PinType.PinSubCategoryObject=ScriptStruct'/Script/CoreUObject.Vector',PinType.PinSubCategoryMemberReference=(),PinType.bIsArray=False,PinType.bIsReference=False,PinType.bIsConst=True,PinType.bIsWeakPointer=False,DefaultValue="0, 0, 0",AutogeneratedDefaultValue="0, 0, 0",LinkedTo=(K2Node_CallFunction_64919 7C3AB90A4E8F92F9FAE0709612FDD604,),PersistentGuid=00000000000000000000000000000000,bHidden=False,bNotConnectable=False,bDefaultValueIsReadOnly=False,bDefaultValueIsIgnored=False,bAdvancedView=False,)
+			   CustomProperties Pin (PinId=54C3C6D04942C8BC62F2D696DFF0A12C,PinName="End",PinToolTip="End\nVector \n\nEnd of line segment.",PinType.PinCategory="struct",PinType.PinSubCategoryObject=ScriptStruct'/Script/CoreUObject.Vector',PinType.PinSubCategoryMemberReference=(),PinType.bIsArray=False,PinType.bIsReference=False,PinType.bIsConst=True,PinType.bIsWeakPointer=False,DefaultValue="0, 0, 0",AutogeneratedDefaultValue="0, 0, 0",LinkedTo=(K2Node_CommutativeAssociativeBinaryOperator_1357 1778AF2B4812E5B635D0DC915DFAA735,),PersistentGuid=00000000000000000000000000000000,bHidden=False,bNotConnectable=False,bDefaultValueIsReadOnly=False,bDefaultValueIsIgnored=False,bAdvancedView=False,)
+			   CustomProperties Pin (PinId=FAD144344A968D69191168B056747AB6,PinName="TraceChannel",PinToolTip="Trace Channel\nETraceTypeQuery Enum",PinType.PinCategory="byte",PinType.PinSubCategoryObject=Enum'/Script/Engine.ETraceTypeQuery',PinType.PinSubCategoryMemberReference=(),PinType.bIsArray=False,PinType.bIsReference=False,PinType.bIsConst=False,PinType.bIsWeakPointer=False,DefaultValue="TraceTypeQuery1",PersistentGuid=00000000000000000000000000000000,bHidden=False,bNotConnectable=False,bDefaultValueIsReadOnly=False,bDefaultValueIsIgnored=False,bAdvancedView=False,)
+			   CustomProperties Pin (PinId=8B576AA34608DE6F0B60328E9D9C1166,PinName="bTraceComplex",PinToolTip="Trace Complex\nBoolean\n\nTrue to test against complex collision, false to test against simplified collision.",PinType.PinCategory="bool",PinType.PinSubCategoryObject=None,PinType.PinSubCategoryMemberReference=(),PinType.bIsArray=False,PinType.bIsReference=False,PinType.bIsConst=False,PinType.bIsWeakPointer=False,DefaultValue="false",AutogeneratedDefaultValue="false",PersistentGuid=00000000000000000000000000000000,bHidden=False,bNotConnectable=False,bDefaultValueIsReadOnly=False,bDefaultValueIsIgnored=False,bAdvancedView=False,)
+			   CustomProperties Pin (PinId=E4B43C3146C4E49A9F72ECBAF170A48B,PinName="ActorsToIgnore",PinToolTip="Actors to Ignore\nArray of Actor References",PinType.PinCategory="object",PinType.PinSubCategoryObject=Class'/Script/Engine.Actor',PinType.PinSubCategoryMemberReference=(),PinType.bIsArray=True,PinType.bIsReference=True,PinType.bIsConst=True,PinType.bIsWeakPointer=False,PersistentGuid=00000000000000000000000000000000,bHidden=False,bNotConnectable=False,bDefaultValueIsReadOnly=False,bDefaultValueIsIgnored=True,bAdvancedView=False,)
+			   CustomProperties Pin (PinId=6DC8450447A07CB69FEA6ABBD6C12BED,PinName="DrawDebugType",PinToolTip="Draw Debug Type\nEDrawDebugTrace Enum",PinType.PinCategory="byte",PinType.PinSubCategoryObject=Enum'/Script/Engine.EDrawDebugTrace',PinType.PinSubCategoryMemberReference=(),PinType.bIsArray=False,PinType.bIsReference=False,PinType.bIsConst=False,PinType.bIsWeakPointer=False,DefaultValue="None",PersistentGuid=00000000000000000000000000000000,bHidden=False,bNotConnectable=False,bDefaultValueIsReadOnly=False,bDefaultValueIsIgnored=False,bAdvancedView=False,)
+			   CustomProperties Pin (PinId=43A3FD1943054F5E9C27A298298CE721,PinName="OutHit",PinToolTip="Out Hit\nHit Result Structure\n\nProperties of the trace hit.",Direction="EGPD_Output",PinType.PinCategory="struct",PinType.PinSubCategoryObject=ScriptStruct'/Script/Engine.HitResult',PinType.PinSubCategoryMemberReference=(),PinType.bIsArray=False,PinType.bIsReference=False,PinType.bIsConst=False,PinType.bIsWeakPointer=False,LinkedTo=(K2Node_CallFunction_65759 9BE54D1545016821EA539F83438E9077,K2Node_CallFunction_1861 7453F08741D187AE2B9629985E33A634,),PersistentGuid=00000000000000000000000000000000,bHidden=False,bNotConnectable=False,bDefaultValueIsReadOnly=False,bDefaultValueIsIgnored=False,bAdvancedView=False,)
+			   CustomProperties Pin (PinId=D4D0981940265BE6F98E8B9F2CB3C064,PinName="bIgnoreSelf",PinToolTip="Ignore Self\nBoolean",PinType.PinCategory="bool",PinType.PinSubCategoryObject=None,PinType.PinSubCategoryMemberReference=(),PinType.bIsArray=False,PinType.bIsReference=False,PinType.bIsConst=False,PinType.bIsWeakPointer=False,DefaultValue="false",AutogeneratedDefaultValue="true",PersistentGuid=00000000000000000000000000000000,bHidden=False,bNotConnectable=False,bDefaultValueIsReadOnly=False,bDefaultValueIsIgnored=False,bAdvancedView=False,)
+			   CustomProperties Pin (PinId=17643C41432C7B66D6FF7F9A6A4BA258,PinName="TraceColor",PinToolTip="Trace Color\nLinear Color Structure",PinType.PinCategory="struct",PinType.PinSubCategoryObject=ScriptStruct'/Script/CoreUObject.LinearColor',PinType.PinSubCategoryMemberReference=(),PinType.bIsArray=False,PinType.bIsReference=False,PinType.bIsConst=False,PinType.bIsWeakPointer=False,DefaultValue="(R=1.000000,G=0.000000,B=0.000000,A=1.000000)",AutogeneratedDefaultValue="(R=1.000000,G=0.000000,B=0.000000,A=1.000000)",PersistentGuid=00000000000000000000000000000000,bHidden=False,bNotConnectable=False,bDefaultValueIsReadOnly=False,bDefaultValueIsIgnored=False,bAdvancedView=True,)
+			   CustomProperties Pin (PinId=F62FE7C840096CE5353C2C98BF138F9C,PinName="TraceHitColor",PinToolTip="Trace Hit Color\nLinear Color Structure",PinType.PinCategory="struct",PinType.PinSubCategoryObject=ScriptStruct'/Script/CoreUObject.LinearColor',PinType.PinSubCategoryMemberReference=(),PinType.bIsArray=False,PinType.bIsReference=False,PinType.bIsConst=False,PinType.bIsWeakPointer=False,DefaultValue="(R=0.000000,G=1.000000,B=0.000000,A=1.000000)",AutogeneratedDefaultValue="(R=0.000000,G=1.000000,B=0.000000,A=1.000000)",PersistentGuid=00000000000000000000000000000000,bHidden=False,bNotConnectable=False,bDefaultValueIsReadOnly=False,bDefaultValueIsIgnored=False,bAdvancedView=True,)
+			   CustomProperties Pin (PinId=9A93C66A42BC103DD20019BCA763A17F,PinName="DrawTime",PinToolTip="Draw Time\nFloat",PinType.PinCategory="float",PinType.PinSubCategoryObject=None,PinType.PinSubCategoryMemberReference=(),PinType.bIsArray=False,PinType.bIsReference=False,PinType.bIsConst=False,PinType.bIsWeakPointer=False,DefaultValue="5.000000",AutogeneratedDefaultValue="5.000000",PersistentGuid=00000000000000000000000000000000,bHidden=False,bNotConnectable=False,bDefaultValueIsReadOnly=False,bDefaultValueIsIgnored=False,bAdvancedView=True,)
+			   CustomProperties Pin (PinId=3A64487743C81996C17E56B543685516,PinName="ReturnValue",PinToolTip="Return Value\nBoolean\n\nTrue if there was a hit, false otherwise.",Direction="EGPD_Output",PinType.PinCategory="bool",PinType.PinSubCategoryObject=None,PinType.PinSubCategoryMemberReference=(),PinType.bIsArray=False,PinType.bIsReference=False,PinType.bIsConst=False,PinType.bIsWeakPointer=False,DefaultValue="false",AutogeneratedDefaultValue="false",LinkedTo=(K2Node_IfThenElse_161 1144399140FB3DC67E1A66B5FEB4A1E1,),PersistentGuid=00000000000000000000000000000000,bHidden=False,bNotConnectable=False,bDefaultValueIsReadOnly=False,bDefaultValueIsIgnored=False,bAdvancedView=False,)
+			End Object
+			Begin Object Class=K2Node_CommutativeAssociativeBinaryOperator Name="K2Node_CommutativeAssociativeBinaryOperator_1357"
+			   bIsPureFunc=True
+			   FunctionReference=(MemberParent=Class'/Script/Engine.KismetMathLibrary',MemberName="Add_VectorVector")
+			   NodePosX=912
+			   NodePosY=-96
+			   NodeGuid=8DC37CA3400BFCC550BCC1BE62581E73
+			   CustomProperties Pin (PinId=B729CFBC424DA01AC5E021B25FE0F7CD,PinName="self",PinFriendlyName=NSLOCTEXT("K2Node", "Target", "Target"),PinType.PinCategory="object",PinType.PinSubCategoryObject=Class'/Script/Engine.KismetMathLibrary',PinType.PinSubCategoryMemberReference=(),PinType.bIsArray=False,PinType.bIsReference=False,PinType.bIsConst=False,PinType.bIsWeakPointer=False,DefaultObject="/Script/Engine.Default__KismetMathLibrary",PersistentGuid=00000000000000000000000000000000,bHidden=True,bNotConnectable=False,bDefaultValueIsReadOnly=False,bDefaultValueIsIgnored=False,bAdvancedView=False,)
+			   CustomProperties Pin (PinId=4410C89D42A350AA1AF9EE830D6A174E,PinName="A",PinType.PinCategory="struct",PinType.PinSubCategoryObject=ScriptStruct'/Script/CoreUObject.Vector',PinType.PinSubCategoryMemberReference=(),PinType.bIsArray=False,PinType.bIsReference=False,PinType.bIsConst=False,PinType.bIsWeakPointer=False,DefaultValue="0, 0, 0",AutogeneratedDefaultValue="0, 0, 0",LinkedTo=(K2Node_CallFunction_65349 40870E1F404EA24CCC3295A98D8A13B6,),PersistentGuid=00000000000000000000000000000000,bHidden=False,bNotConnectable=False,bDefaultValueIsReadOnly=False,bDefaultValueIsIgnored=False,bAdvancedView=False,)
+			   CustomProperties Pin (PinId=C2F85EED4A8F818C3CE355A3C56691F5,PinName="B",PinType.PinCategory="struct",PinType.PinSubCategoryObject=ScriptStruct'/Script/CoreUObject.Vector',PinType.PinSubCategoryMemberReference=(),PinType.bIsArray=False,PinType.bIsReference=False,PinType.bIsConst=False,PinType.bIsWeakPointer=False,DefaultValue="0, 0, 0",AutogeneratedDefaultValue="0, 0, 0",LinkedTo=(K2Node_CallFunction_64919 7C3AB90A4E8F92F9FAE0709612FDD604,),PersistentGuid=00000000000000000000000000000000,bHidden=False,bNotConnectable=False,bDefaultValueIsReadOnly=False,bDefaultValueIsIgnored=False,bAdvancedView=False,)
+			   CustomProperties Pin (PinId=1778AF2B4812E5B635D0DC915DFAA735,PinName="ReturnValue",Direction="EGPD_Output",PinType.PinCategory="struct",PinType.PinSubCategoryObject=ScriptStruct'/Script/CoreUObject.Vector',PinType.PinSubCategoryMemberReference=(),PinType.bIsArray=False,PinType.bIsReference=False,PinType.bIsConst=False,PinType.bIsWeakPointer=False,AutogeneratedDefaultValue="0, 0, 0",LinkedTo=(K2Node_CallFunction_65598 54C3C6D04942C8BC62F2D696DFF0A12C,),PersistentGuid=00000000000000000000000000000000,bHidden=False,bNotConnectable=False,bDefaultValueIsReadOnly=False,bDefaultValueIsIgnored=False,bAdvancedView=False,)
+			End Object
+			Begin Object Class=K2Node_CallFunction Name="K2Node_CallFunction_65759"
+			   bIsPureFunc=True
+			   FunctionReference=(MemberParent=Class'/Script/Engine.GameplayStatics',MemberName="BreakHitResult")
+			   NodePosX=1568
+			   NodePosY=-288
+			   NodeGuid=92BF07174B1E8AA8AC947F9AF1A1E6E7
+			   CustomProperties Pin (PinId=CF78D28A455F553D39502496BC89B70E,PinName="self",PinFriendlyName=NSLOCTEXT("K2Node", "Target", "Target"),PinToolTip="Target\nGameplay Statics Reference",PinType.PinCategory="object",PinType.PinSubCategoryObject=Class'/Script/Engine.GameplayStatics',PinType.PinSubCategoryMemberReference=(),PinType.bIsArray=False,PinType.bIsReference=False,PinType.bIsConst=False,PinType.bIsWeakPointer=False,DefaultObject="/Script/Engine.Default__GameplayStatics",PersistentGuid=00000000000000000000000000000000,bHidden=True,bNotConnectable=False,bDefaultValueIsReadOnly=False,bDefaultValueIsIgnored=False,bAdvancedView=False,)
+			   CustomProperties Pin (PinId=9BE54D1545016821EA539F83438E9077,PinName="Hit",PinToolTip="Hit\nHit Result Structure (by ref)\n\nThe source HitResult.",PinType.PinCategory="struct",PinType.PinSubCategoryObject=ScriptStruct'/Script/Engine.HitResult',PinType.PinSubCategoryMemberReference=(),PinType.bIsArray=False,PinType.bIsReference=True,PinType.bIsConst=True,PinType.bIsWeakPointer=False,LinkedTo=(K2Node_CallFunction_65598 43A3FD1943054F5E9C27A298298CE721,),PersistentGuid=00000000000000000000000000000000,bHidden=False,bNotConnectable=False,bDefaultValueIsReadOnly=False,bDefaultValueIsIgnored=True,bAdvancedView=False,)
+			   CustomProperties Pin (PinId=42FAE549473540FC07A70799D95A29B3,PinName="bBlockingHit",PinToolTip="Blocking Hit\nBoolean\n\nTrue if there was a blocking hit, false otherwise.",Direction="EGPD_Output",PinType.PinCategory="bool",PinType.PinSubCategoryObject=None,PinType.PinSubCategoryMemberReference=(),PinType.bIsArray=False,PinType.bIsReference=False,PinType.bIsConst=False,PinType.bIsWeakPointer=False,DefaultValue="false",AutogeneratedDefaultValue="false",PersistentGuid=00000000000000000000000000000000,bHidden=False,bNotConnectable=False,bDefaultValueIsReadOnly=False,bDefaultValueIsIgnored=False,bAdvancedView=False,)
+			   CustomProperties Pin (PinId=31399B19487BF09D8EF1088E8ED56118,PinName="bInitialOverlap",PinToolTip="Initial Overlap\nBoolean\n\nTrue if the hit started in an initial overlap. In this case some other values should be interpreted differently. Time will be 0, ImpactPoint will equal Location, and normals will be equal and indicate a depenetration vector.",Direction="EGPD_Output",PinType.PinCategory="bool",PinType.PinSubCategoryObject=None,PinType.PinSubCategoryMemberReference=(),PinType.bIsArray=False,PinType.bIsReference=False,PinType.bIsConst=False,PinType.bIsWeakPointer=False,DefaultValue="false",AutogeneratedDefaultValue="false",PersistentGuid=00000000000000000000000000000000,bHidden=False,bNotConnectable=False,bDefaultValueIsReadOnly=False,bDefaultValueIsIgnored=False,bAdvancedView=False,)
+			   CustomProperties Pin (PinId=4312416249D19F19294CADA26CDE64E5,PinName="Time",PinToolTip="Time\nFloat\n\n\'Time\' of impact along trace direction ranging from [0.0 to 1.0) if there is a hit, indicating time between start and end. Equals 1.0 if there is no hit.",Direction="EGPD_Output",PinType.PinCategory="float",PinType.PinSubCategoryObject=None,PinType.PinSubCategoryMemberReference=(),PinType.bIsArray=False,PinType.bIsReference=False,PinType.bIsConst=False,PinType.bIsWeakPointer=False,DefaultValue="0.0",AutogeneratedDefaultValue="0.0",PersistentGuid=00000000000000000000000000000000,bHidden=False,bNotConnectable=False,bDefaultValueIsReadOnly=False,bDefaultValueIsIgnored=False,bAdvancedView=False,)
+			   CustomProperties Pin (PinId=C3E1DE9D4B514E4FCDA1159FEA7E7832,PinName="Location",PinToolTip="Location\nVector \n\nLocation of the hit in world space. If this was a swept shape test, this is the location where we can place the shape in the world where it will not penetrate.",Direction="EGPD_Output",PinType.PinCategory="struct",PinType.PinSubCategoryObject=ScriptStruct'/Script/CoreUObject.Vector',PinType.PinSubCategoryMemberReference=(),PinType.bIsArray=False,PinType.bIsReference=False,PinType.bIsConst=False,PinType.bIsWeakPointer=False,AutogeneratedDefaultValue="0, 0, 0",PersistentGuid=00000000000000000000000000000000,bHidden=False,bNotConnectable=False,bDefaultValueIsReadOnly=False,bDefaultValueIsIgnored=False,bAdvancedView=False,)
+			   CustomProperties Pin (PinId=EB9F6AED44EDCC7D8F313FB63A3B4139,PinName="ImpactPoint",PinToolTip="Impact Point\nVector \n\nLocation of the actual contact point of the trace shape with the surface of the hit object. Equal to Location in the case of an initial overlap.",Direction="EGPD_Output",PinType.PinCategory="struct",PinType.PinSubCategoryObject=ScriptStruct'/Script/CoreUObject.Vector',PinType.PinSubCategoryMemberReference=(),PinType.bIsArray=False,PinType.bIsReference=False,PinType.bIsConst=False,PinType.bIsWeakPointer=False,AutogeneratedDefaultValue="0, 0, 0",PersistentGuid=00000000000000000000000000000000,bHidden=False,bNotConnectable=False,bDefaultValueIsReadOnly=False,bDefaultValueIsIgnored=False,bAdvancedView=False,)
+			   CustomProperties Pin (PinId=AC37F9C247A40EF82D1755B03A6EF0AF,PinName="Normal",PinToolTip="Normal\nVector \n\nNormal of the hit in world space, for the object that was swept (e.g. for a sphere trace this points towards the sphere\'s center). Equal to ImpactNormal for line tests.",Direction="EGPD_Output",PinType.PinCategory="struct",PinType.PinSubCategoryObject=ScriptStruct'/Script/CoreUObject.Vector',PinType.PinSubCategoryMemberReference=(),PinType.bIsArray=False,PinType.bIsReference=False,PinType.bIsConst=False,PinType.bIsWeakPointer=False,AutogeneratedDefaultValue="0, 0, 0",PersistentGuid=00000000000000000000000000000000,bHidden=False,bNotConnectable=False,bDefaultValueIsReadOnly=False,bDefaultValueIsIgnored=False,bAdvancedView=False,)
+			   CustomProperties Pin (PinId=43DE840740F7EDB5AEEE2BA59BA63807,PinName="ImpactNormal",PinToolTip="Impact Normal\nVector \n\nNormal of the hit in world space, for the object that was hit by the sweep.",Direction="EGPD_Output",PinType.PinCategory="struct",PinType.PinSubCategoryObject=ScriptStruct'/Script/CoreUObject.Vector',PinType.PinSubCategoryMemberReference=(),PinType.bIsArray=False,PinType.bIsReference=False,PinType.bIsConst=False,PinType.bIsWeakPointer=False,AutogeneratedDefaultValue="0, 0, 0",PersistentGuid=00000000000000000000000000000000,bHidden=False,bNotConnectable=False,bDefaultValueIsReadOnly=False,bDefaultValueIsIgnored=False,bAdvancedView=False,)
+			   CustomProperties Pin (PinId=532D29F44B5DCE5F021B698ED58D27C4,PinName="PhysMat",PinToolTip="Phys Mat\nPhysical Material Reference\n\nPhysical material that was hit. Must set bReturnPhysicalMaterial to true in the query params for this to be returned.",Direction="EGPD_Output",PinType.PinCategory="object",PinType.PinSubCategoryObject=Class'/Script/Engine.PhysicalMaterial',PinType.PinSubCategoryMemberReference=(),PinType.bIsArray=False,PinType.bIsReference=False,PinType.bIsConst=False,PinType.bIsWeakPointer=False,PersistentGuid=00000000000000000000000000000000,bHidden=False,bNotConnectable=False,bDefaultValueIsReadOnly=False,bDefaultValueIsIgnored=False,bAdvancedView=False,)
+			   CustomProperties Pin (PinId=F1D6985542953219C7EF09A3FBFA8AF0,PinName="HitActor",PinToolTip="Hit Actor\nActor Reference\n\nActor hit by the trace.",Direction="EGPD_Output",PinType.PinCategory="object",PinType.PinSubCategoryObject=Class'/Script/Engine.Actor',PinType.PinSubCategoryMemberReference=(),PinType.bIsArray=False,PinType.bIsReference=False,PinType.bIsConst=False,PinType.bIsWeakPointer=False,LinkedTo=(K2Node_MacroInstance_65 C3FB28AC47AAD2B5998C2AA9D4862C2E,K2Node_CallFunction_1861 BB747167446E2DDBEFFFDD982A9CD245,),PersistentGuid=00000000000000000000000000000000,bHidden=False,bNotConnectable=False,bDefaultValueIsReadOnly=False,bDefaultValueIsIgnored=False,bAdvancedView=False,)
+			   CustomProperties Pin (PinId=FDFC05184366EBC9BDC42AA3A95EF523,PinName="HitComponent",PinToolTip="Hit Component\nPrimitive Component Reference\n\nPrimitiveComponent hit by the trace.",Direction="EGPD_Output",PinType.PinCategory="object",PinType.PinSubCategoryObject=Class'/Script/Engine.PrimitiveComponent',PinType.PinSubCategoryMemberReference=(),PinType.bIsArray=False,PinType.bIsReference=False,PinType.bIsConst=False,PinType.bIsWeakPointer=False,PersistentGuid=00000000000000000000000000000000,bHidden=False,bNotConnectable=False,bDefaultValueIsReadOnly=False,bDefaultValueIsIgnored=False,bAdvancedView=False,)
+			   CustomProperties Pin (PinId=CE3689BE4446FD1C971B7A9942E7AC24,PinName="HitBoneName",PinToolTip="Hit Bone Name\nName\n\nName of the bone hit (valid only if we hit a skeletal mesh).",Direction="EGPD_Output",PinType.PinCategory="name",PinType.PinSubCategoryObject=None,PinType.PinSubCategoryMemberReference=(),PinType.bIsArray=False,PinType.bIsReference=False,PinType.bIsConst=False,PinType.bIsWeakPointer=False,DefaultValue="None",AutogeneratedDefaultValue="None",PersistentGuid=00000000000000000000000000000000,bHidden=False,bNotConnectable=False,bDefaultValueIsReadOnly=False,bDefaultValueIsIgnored=False,bAdvancedView=False,)
+			   CustomProperties Pin (PinId=1E9AAE5D433101E315ABFCB6EE0C3F4E,PinName="HitItem",PinToolTip="Hit Item\nInteger\n\nPrimitive-specific data recording which item in the primitive was hit",Direction="EGPD_Output",PinType.PinCategory="int",PinType.PinSubCategoryObject=None,PinType.PinSubCategoryMemberReference=(),PinType.bIsArray=False,PinType.bIsReference=False,PinType.bIsConst=False,PinType.bIsWeakPointer=False,DefaultValue="0",AutogeneratedDefaultValue="0",PersistentGuid=00000000000000000000000000000000,bHidden=False,bNotConnectable=False,bDefaultValueIsReadOnly=False,bDefaultValueIsIgnored=False,bAdvancedView=False,)
+			   CustomProperties Pin (PinId=FA0E13FF4A7001F230C36C84199334C5,PinName="FaceIndex",PinToolTip="Face Index\nInteger\n\nIf colliding with trimesh or landscape, index of face that was hit.",Direction="EGPD_Output",PinType.PinCategory="int",PinType.PinSubCategoryObject=None,PinType.PinSubCategoryMemberReference=(),PinType.bIsArray=False,PinType.bIsReference=False,PinType.bIsConst=False,PinType.bIsWeakPointer=False,DefaultValue="0",AutogeneratedDefaultValue="0",PersistentGuid=00000000000000000000000000000000,bHidden=False,bNotConnectable=False,bDefaultValueIsReadOnly=False,bDefaultValueIsIgnored=False,bAdvancedView=False,)
+			   CustomProperties Pin (PinId=6AFDED3441A726D6362809827812A060,PinName="TraceStart",PinToolTip="Trace Start\nVector ",Direction="EGPD_Output",PinType.PinCategory="struct",PinType.PinSubCategoryObject=ScriptStruct'/Script/CoreUObject.Vector',PinType.PinSubCategoryMemberReference=(),PinType.bIsArray=False,PinType.bIsReference=False,PinType.bIsConst=False,PinType.bIsWeakPointer=False,DefaultValue="0, 0, 0",AutogeneratedDefaultValue="0, 0, 0",PersistentGuid=00000000000000000000000000000000,bHidden=False,bNotConnectable=False,bDefaultValueIsReadOnly=False,bDefaultValueIsIgnored=False,bAdvancedView=False,)
+			   CustomProperties Pin (PinId=96D907364CB400491CEFF4B6B3921F92,PinName="TraceEnd",PinToolTip="Trace End\nVector ",Direction="EGPD_Output",PinType.PinCategory="struct",PinType.PinSubCategoryObject=ScriptStruct'/Script/CoreUObject.Vector',PinType.PinSubCategoryMemberReference=(),PinType.bIsArray=False,PinType.bIsReference=False,PinType.bIsConst=False,PinType.bIsWeakPointer=False,DefaultValue="0, 0, 0",AutogeneratedDefaultValue="0, 0, 0",PersistentGuid=00000000000000000000000000000000,bHidden=False,bNotConnectable=False,bDefaultValueIsReadOnly=False,bDefaultValueIsIgnored=False,bAdvancedView=False,)
+			End Object
+			Begin Object Class=K2Node_CommutativeAssociativeBinaryOperator Name="K2Node_CommutativeAssociativeBinaryOperator_4227"
+			   bIsPureFunc=True
+			   FunctionReference=(MemberParent=Class'/Script/Engine.KismetMathLibrary',MemberName="Multiply_FloatFloat")
+			   NodePosX=592
+			   NodePosY=16
+			   NodeGuid=E816C87A41EA136F98171CAC80832A2C
+			   CustomProperties Pin (PinId=0D4F3C2044C7E8EA342002AF074BCB55,PinName="self",PinFriendlyName=NSLOCTEXT("K2Node", "Target", "Target"),PinType.PinCategory="object",PinType.PinSubCategoryObject=Class'/Script/Engine.KismetMathLibrary',PinType.PinSubCategoryMemberReference=(),PinType.bIsArray=False,PinType.bIsReference=False,PinType.bIsConst=False,PinType.bIsWeakPointer=False,DefaultObject="/Script/Engine.Default__KismetMathLibrary",PersistentGuid=00000000000000000000000000000000,bHidden=True,bNotConnectable=False,bDefaultValueIsReadOnly=False,bDefaultValueIsIgnored=False,bAdvancedView=False,)
+			   CustomProperties Pin (PinId=AE00303B4D5F38B1B48C718ABEC37EF0,PinName="A",PinType.PinCategory="float",PinType.PinSubCategoryObject=None,PinType.PinSubCategoryMemberReference=(),PinType.bIsArray=False,PinType.bIsReference=False,PinType.bIsConst=False,PinType.bIsWeakPointer=False,DefaultValue="0.0",AutogeneratedDefaultValue="0.0",LinkedTo=(K2Node_VariableGet_281 F3DEDC984C80801BC87F67B0097115B8,),PersistentGuid=00000000000000000000000000000000,bHidden=False,bNotConnectable=False,bDefaultValueIsReadOnly=False,bDefaultValueIsIgnored=False,bAdvancedView=False,)
+			   CustomProperties Pin (PinId=290EC66447D9833EAF1FD8914F55B78C,PinName="B",PinType.PinCategory="float",PinType.PinSubCategoryObject=None,PinType.PinSubCategoryMemberReference=(),PinType.bIsArray=False,PinType.bIsReference=False,PinType.bIsConst=False,PinType.bIsWeakPointer=False,DefaultValue="10",AutogeneratedDefaultValue="0.0",PersistentGuid=00000000000000000000000000000000,bHidden=False,bNotConnectable=False,bDefaultValueIsReadOnly=False,bDefaultValueIsIgnored=False,bAdvancedView=False,)
+			   CustomProperties Pin (PinId=0FAF44A34397D0F87EB02B88E8A653E3,PinName="ReturnValue",Direction="EGPD_Output",PinType.PinCategory="float",PinType.PinSubCategoryObject=None,PinType.PinSubCategoryMemberReference=(),PinType.bIsArray=False,PinType.bIsReference=False,PinType.bIsConst=False,PinType.bIsWeakPointer=False,DefaultValue="0.0",AutogeneratedDefaultValue="0.0",LinkedTo=(K2Node_CallFunction_65349 2C5718584516424BB52CF6A79E702026,),PersistentGuid=00000000000000000000000000000000,bHidden=False,bNotConnectable=False,bDefaultValueIsReadOnly=False,bDefaultValueIsIgnored=False,bAdvancedView=False,)
+			End Object
+			Begin Object Class=K2Node_IfThenElse Name="K2Node_IfThenElse_161"
+			   NodePosX=1600
+			   NodePosY=-432
+			   ErrorType=4
+			   ErrorMsg="Info The execution path doesn\'t end with a return node.  False "
+			   NodeGuid=C447E33540D01EDA1CD72EAB0CFEE477
+			   CustomProperties Pin (PinId=66D340DC4A52DC8D727C43ABA6E818C4,PinName="execute",PinType.PinCategory="exec",PinType.PinSubCategoryObject=None,PinType.PinSubCategoryMemberReference=(),PinType.bIsArray=False,PinType.bIsReference=False,PinType.bIsConst=False,PinType.bIsWeakPointer=False,LinkedTo=(K2Node_CallFunction_65598 E001912C48BB5BEC656460843FE151FA,),PersistentGuid=00000000000000000000000000000000,bHidden=False,bNotConnectable=False,bDefaultValueIsReadOnly=False,bDefaultValueIsIgnored=False,bAdvancedView=False,)
+			   CustomProperties Pin (PinId=1144399140FB3DC67E1A66B5FEB4A1E1,PinName="Condition",PinType.PinCategory="bool",PinType.PinSubCategoryObject=None,PinType.PinSubCategoryMemberReference=(),PinType.bIsArray=False,PinType.bIsReference=False,PinType.bIsConst=False,PinType.bIsWeakPointer=False,DefaultValue="true",AutogeneratedDefaultValue="true",LinkedTo=(K2Node_CallFunction_65598 3A64487743C81996C17E56B543685516,),PersistentGuid=00000000000000000000000000000000,bHidden=False,bNotConnectable=False,bDefaultValueIsReadOnly=False,bDefaultValueIsIgnored=False,bAdvancedView=False,)
+			   CustomProperties Pin (PinId=8DD9A5BC40D019A1AFD6F284F8C15FFD,PinName="then",PinFriendlyName=NSLOCTEXT("K2Node", "true", "true"),Direction="EGPD_Output",PinType.PinCategory="exec",PinType.PinSubCategoryObject=None,PinType.PinSubCategoryMemberReference=(),PinType.bIsArray=False,PinType.bIsReference=False,PinType.bIsConst=False,PinType.bIsWeakPointer=False,LinkedTo=(K2Node_MacroInstance_65 54DCDF6A4C18DFFBF71F54AE24DA40C0,),PersistentGuid=00000000000000000000000000000000,bHidden=False,bNotConnectable=False,bDefaultValueIsReadOnly=False,bDefaultValueIsIgnored=False,bAdvancedView=False,)
+			   CustomProperties Pin (PinId=4CD062D04319B76E8FD9D6B0CDDC7EA9,PinName="else",PinFriendlyName=NSLOCTEXT("K2Node", "false", "false"),Direction="EGPD_Output",PinType.PinCategory="exec",PinType.PinSubCategoryObject=None,PinType.PinSubCategoryMemberReference=(),PinType.bIsArray=False,PinType.bIsReference=False,PinType.bIsConst=False,PinType.bIsWeakPointer=False,PersistentGuid=00000000000000000000000000000000,bHidden=False,bNotConnectable=False,bDefaultValueIsReadOnly=False,bDefaultValueIsIgnored=False,bAdvancedView=False,)
+			End Object
+			Begin Object Class=EdGraphNode_Comment Name="EdGraphNode_Comment_590"
+			   NodePosX=-160
+			   NodePosY=-528
+			   NodeWidth=1616
+			   NodeHeight=704
+			   NodeComment="Trace from the center of the screen in direction of the camera"
+			   NodeGuid=6B615868413EE32AD4B9FC8B9F5DB6B5
+			End Object
+			Begin Object Class=EdGraphNode_Comment Name="EdGraphNode_Comment_596"
+			   NodePosX=1488
+			   NodePosY=-528
+			   NodeWidth=1328
+			   NodeHeight=704
+			   NodeComment="If the Trace hits something apply point damage"
+			   NodeGuid=765187474B4D4B14235EAF883D241CEC
+			End Object
+			Begin Object Class=K2Node_CallFunction Name="K2Node_CallFunction_1861"
+			   FunctionReference=(MemberParent=Class'/Script/Engine.GameplayStatics',MemberName="ApplyPointDamage")
+			   NodePosX=2560
+			   NodePosY=-432
+			   NodeGuid=EA250A4244DD16B5695B73803C315BA1
+			   CustomProperties Pin (PinId=6F410C0B46D7B07315AD748B03EB86C8,PinName="execute",PinToolTip="\nExec",PinType.PinCategory="exec",PinType.PinSubCategoryObject=None,PinType.PinSubCategoryMemberReference=(),PinType.bIsArray=False,PinType.bIsReference=False,PinType.bIsConst=False,PinType.bIsWeakPointer=False,LinkedTo=(K2Node_MacroInstance_65 E75DA26F4EBE18C63E3CEDA6B1D88DC6,),PersistentGuid=00000000000000000000000000000000,bHidden=False,bNotConnectable=False,bDefaultValueIsReadOnly=False,bDefaultValueIsIgnored=False,bAdvancedView=False,)
+			   CustomProperties Pin (PinId=57573271406B9D0CCE9C28BF2DE94719,PinName="then",PinToolTip="\nExec",Direction="EGPD_Output",PinType.PinCategory="exec",PinType.PinSubCategoryObject=None,PinType.PinSubCategoryMemberReference=(),PinType.bIsArray=False,PinType.bIsReference=False,PinType.bIsConst=False,PinType.bIsWeakPointer=False,PersistentGuid=00000000000000000000000000000000,bHidden=False,bNotConnectable=False,bDefaultValueIsReadOnly=False,bDefaultValueIsIgnored=False,bAdvancedView=False,)
+			   CustomProperties Pin (PinId=52375B5841E301566AD3AA961C2B94D1,PinName="self",PinFriendlyName=NSLOCTEXT("K2Node", "Target", "Target"),PinToolTip="Target\nGameplay Statics Reference",PinType.PinCategory="object",PinType.PinSubCategoryObject=Class'/Script/Engine.GameplayStatics',PinType.PinSubCategoryMemberReference=(),PinType.bIsArray=False,PinType.bIsReference=False,PinType.bIsConst=False,PinType.bIsWeakPointer=False,DefaultObject="/Script/Engine.Default__GameplayStatics",PersistentGuid=00000000000000000000000000000000,bHidden=True,bNotConnectable=False,bDefaultValueIsReadOnly=False,bDefaultValueIsIgnored=False,bAdvancedView=False,)
+			   CustomProperties Pin (PinId=BB747167446E2DDBEFFFDD982A9CD245,PinName="DamagedActor",PinToolTip="Damaged Actor\nActor Reference\n\nActor that will be damaged.",PinType.PinCategory="object",PinType.PinSubCategoryObject=Class'/Script/Engine.Actor',PinType.PinSubCategoryMemberReference=(),PinType.bIsArray=False,PinType.bIsReference=False,PinType.bIsConst=False,PinType.bIsWeakPointer=False,LinkedTo=(K2Node_CallFunction_65759 F1D6985542953219C7EF09A3FBFA8AF0,),PersistentGuid=00000000000000000000000000000000,bHidden=False,bNotConnectable=False,bDefaultValueIsReadOnly=False,bDefaultValueIsIgnored=False,bAdvancedView=False,)
+			   CustomProperties Pin (PinId=E0E59CE348F9C5A31938AA864E1EDF8C,PinName="BaseDamage",PinToolTip="Base Damage\nFloat\n\nThe base damage to apply.",PinType.PinCategory="float",PinType.PinSubCategoryObject=None,PinType.PinSubCategoryMemberReference=(),PinType.bIsArray=False,PinType.bIsReference=False,PinType.bIsConst=False,PinType.bIsWeakPointer=False,DefaultValue="100",AutogeneratedDefaultValue="0.0",PersistentGuid=00000000000000000000000000000000,bHidden=False,bNotConnectable=False,bDefaultValueIsReadOnly=False,bDefaultValueIsIgnored=False,bAdvancedView=False,)
+			   CustomProperties Pin (PinId=744D73DE45C865A6B95D1A90F0DF7412,PinName="HitFromDirection",PinToolTip="Hit from Direction\nVector  (by ref)\n\nDirection the hit came FROM",PinType.PinCategory="struct",PinType.PinSubCategoryObject=ScriptStruct'/Script/CoreUObject.Vector',PinType.PinSubCategoryMemberReference=(),PinType.bIsArray=False,PinType.bIsReference=True,PinType.bIsConst=True,PinType.bIsWeakPointer=False,DefaultValue="0, 0, 0",AutogeneratedDefaultValue="0, 0, 0",LinkedTo=(K2Node_CallFunction_1862 9C37BD7C42AFE3084A157196D1784CD3,),PersistentGuid=00000000000000000000000000000000,bHidden=False,bNotConnectable=False,bDefaultValueIsReadOnly=False,bDefaultValueIsIgnored=True,bAdvancedView=False,)
+			   CustomProperties Pin (PinId=7453F08741D187AE2B9629985E33A634,PinName="HitInfo",PinToolTip="Hit Info\nHit Result Structure (by ref)\n\nCollision or trace result that describes the hit",PinType.PinCategory="struct",PinType.PinSubCategoryObject=ScriptStruct'/Script/Engine.HitResult',PinType.PinSubCategoryMemberReference=(),PinType.bIsArray=False,PinType.bIsReference=True,PinType.bIsConst=True,PinType.bIsWeakPointer=False,LinkedTo=(K2Node_CallFunction_65598 43A3FD1943054F5E9C27A298298CE721,),PersistentGuid=00000000000000000000000000000000,bHidden=False,bNotConnectable=False,bDefaultValueIsReadOnly=False,bDefaultValueIsIgnored=True,bAdvancedView=False,)
+			   CustomProperties Pin (PinId=8AAA25E3411AD93065E153BDF12317E1,PinName="EventInstigator",PinToolTip="Event Instigator\nController Reference\n\nController that was responsible for causing this damage (e.g. player who shot the weapon)",PinType.PinCategory="object",PinType.PinSubCategoryObject=Class'/Script/Engine.Controller',PinType.PinSubCategoryMemberReference=(),PinType.bIsArray=False,PinType.bIsReference=False,PinType.bIsConst=False,PinType.bIsWeakPointer=False,LinkedTo=(K2Node_CallFunction_2902 3119A2234CC8CC2E58EDDA94ED78DB30,),PersistentGuid=00000000000000000000000000000000,bHidden=False,bNotConnectable=False,bDefaultValueIsReadOnly=False,bDefaultValueIsIgnored=False,bAdvancedView=False,)
+			   CustomProperties Pin (PinId=471A4E874191001E241068AF8F5CA28F,PinName="DamageCauser",PinToolTip="Damage Causer\nActor Reference\n\nActor that actually caused the damage (e.g. the grenade that exploded)",PinType.PinCategory="object",PinType.PinSubCategoryObject=Class'/Script/Engine.Actor',PinType.PinSubCategoryMemberReference=(),PinType.bIsArray=False,PinType.bIsReference=False,PinType.bIsConst=False,PinType.bIsWeakPointer=False,LinkedTo=(K2Node_Self_15 CFA0BC05479F37C57BC5E890E311DF98,),PersistentGuid=00000000000000000000000000000000,bHidden=False,bNotConnectable=False,bDefaultValueIsReadOnly=False,bDefaultValueIsIgnored=False,bAdvancedView=False,)
+			   CustomProperties Pin (PinId=423FB9E642E33BBA76EDACAF6D63DCA0,PinName="DamageTypeClass",PinToolTip="Damage Type Class\nDamage Type Class\n\nClass that describes the damage that was done.",PinType.PinCategory="class",PinType.PinSubCategoryObject=Class'/Script/Engine.DamageType',PinType.PinSubCategoryMemberReference=(),PinType.bIsArray=False,PinType.bIsReference=False,PinType.bIsConst=False,PinType.bIsWeakPointer=False,DefaultObject="/Script/Engine.DamageType",PersistentGuid=00000000000000000000000000000000,bHidden=False,bNotConnectable=False,bDefaultValueIsReadOnly=False,bDefaultValueIsIgnored=False,bAdvancedView=False,)
+			End Object
+			Begin Object Class=K2Node_CallFunction Name="K2Node_CallFunction_1862"
+			   bIsPureFunc=True
+			   FunctionReference=(MemberParent=Class'/Script/Engine.KismetMathLibrary',MemberName="MakeVector")
+			   NodePosX=2240
+			   NodePosY=-352
+			   NodeGuid=90D1D25A44EFBE2A50A7C4AB73385384
+			   CustomProperties Pin (PinId=89111CA54D160B2C7528509D086B822B,PinName="self",PinFriendlyName=NSLOCTEXT("K2Node", "Target", "Target"),PinToolTip="Target\nKismet Math Library Reference",PinType.PinCategory="object",PinType.PinSubCategoryObject=Class'/Script/Engine.KismetMathLibrary',PinType.PinSubCategoryMemberReference=(),PinType.bIsArray=False,PinType.bIsReference=False,PinType.bIsConst=False,PinType.bIsWeakPointer=False,DefaultObject="/Script/Engine.Default__KismetMathLibrary",PersistentGuid=00000000000000000000000000000000,bHidden=True,bNotConnectable=False,bDefaultValueIsReadOnly=False,bDefaultValueIsIgnored=False,bAdvancedView=False,)
+			   CustomProperties Pin (PinId=396AF363497B394C5C3BA4BCB5A80C67,PinName="X",PinToolTip="X\nFloat",PinType.PinCategory="float",PinType.PinSubCategoryObject=None,PinType.PinSubCategoryMemberReference=(),PinType.bIsArray=False,PinType.bIsReference=False,PinType.bIsConst=False,PinType.bIsWeakPointer=False,DefaultValue="0.0",AutogeneratedDefaultValue="0.0",PersistentGuid=00000000000000000000000000000000,bHidden=False,bNotConnectable=False,bDefaultValueIsReadOnly=False,bDefaultValueIsIgnored=False,bAdvancedView=False,)
+			   CustomProperties Pin (PinId=24D8F5EE410BCDB0B82938825D4FCECC,PinName="Y",PinToolTip="Y\nFloat",PinType.PinCategory="float",PinType.PinSubCategoryObject=None,PinType.PinSubCategoryMemberReference=(),PinType.bIsArray=False,PinType.bIsReference=False,PinType.bIsConst=False,PinType.bIsWeakPointer=False,DefaultValue="0.0",AutogeneratedDefaultValue="0.0",PersistentGuid=00000000000000000000000000000000,bHidden=False,bNotConnectable=False,bDefaultValueIsReadOnly=False,bDefaultValueIsIgnored=False,bAdvancedView=False,)
+			   CustomProperties Pin (PinId=8C7BE8EB4BFE81E73D8821ABC5B5C51D,PinName="Z",PinToolTip="Z\nFloat",PinType.PinCategory="float",PinType.PinSubCategoryObject=None,PinType.PinSubCategoryMemberReference=(),PinType.bIsArray=False,PinType.bIsReference=False,PinType.bIsConst=False,PinType.bIsWeakPointer=False,DefaultValue="1",AutogeneratedDefaultValue="0.0",PersistentGuid=00000000000000000000000000000000,bHidden=False,bNotConnectable=False,bDefaultValueIsReadOnly=False,bDefaultValueIsIgnored=False,bAdvancedView=False,)
+			   CustomProperties Pin (PinId=9C37BD7C42AFE3084A157196D1784CD3,PinName="ReturnValue",PinToolTip="Return Value\nVector ",Direction="EGPD_Output",PinType.PinCategory="struct",PinType.PinSubCategoryObject=ScriptStruct'/Script/CoreUObject.Vector',PinType.PinSubCategoryMemberReference=(),PinType.bIsArray=False,PinType.bIsReference=False,PinType.bIsConst=False,PinType.bIsWeakPointer=False,AutogeneratedDefaultValue="0, 0, 0",LinkedTo=(K2Node_CallFunction_1861 744D73DE45C865A6B95D1A90F0DF7412,),PersistentGuid=00000000000000000000000000000000,bHidden=False,bNotConnectable=False,bDefaultValueIsReadOnly=False,bDefaultValueIsIgnored=False,bAdvancedView=False,)
+			End Object
+			Begin Object Class=K2Node_VariableGet Name="K2Node_VariableGet_281"
+			   VariableReference=(MemberName="InteractionDistance",MemberGuid=024096844ACA19D4D6CB4DA334906B98,bSelfContext=True)
+			   NodePosX=384
+			   NodePosY=64
+			   ErrorType=1
+			   ErrorMsg="Error The property associated with  InteractionDistance  could not be found"
+			   NodeGuid=EC502D5E41D525430986138A615256A9
+			   CustomProperties Pin (PinId=F3DEDC984C80801BC87F67B0097115B8,PinName="InteractionDistance",Direction="EGPD_Output",PinType.PinCategory="float",PinType.PinSubCategoryObject=None,PinType.PinSubCategoryMemberReference=(),PinType.bIsArray=False,PinType.bIsReference=False,PinType.bIsConst=False,PinType.bIsWeakPointer=False,AutogeneratedDefaultValue="0.0",LinkedTo=(K2Node_CommutativeAssociativeBinaryOperator_4227 AE00303B4D5F38B1B48C718ABEC37EF0,),PersistentGuid=00000000000000000000000000000000,bHidden=False,bNotConnectable=False,bDefaultValueIsReadOnly=False,bDefaultValueIsIgnored=False,bAdvancedView=False,)
+			   CustomProperties Pin (PinId=7C7F96A046F0B4BB8BDA5C86407AD84A,PinName="self",PinFriendlyName=NSLOCTEXT("K2Node", "Target", "Target"),PinType.PinCategory="object",PinType.PinSubCategoryObject=BlueprintGeneratedClass'/Game/ExampleContent/BlueprintFluidSurface/HeightfieldPainter.HeightfieldPainter_C',PinType.PinSubCategoryMemberReference=(),PinType.bIsArray=False,PinType.bIsReference=False,PinType.bIsConst=False,PinType.bIsWeakPointer=False,PersistentGuid=00000000000000000000000000000000,bHidden=True,bNotConnectable=False,bDefaultValueIsReadOnly=False,bDefaultValueIsIgnored=False,bAdvancedView=False,)
+			End Object
+			Begin Object Class=K2Node_Self Name="K2Node_Self_15"
+			   NodePosX=2336
+			   NodePosY=-80
+			   NodeGuid=519A1C834BAB76F02E68CE802A2ED70E
+			   CustomProperties Pin (PinId=CFA0BC05479F37C57BC5E890E311DF98,PinName="self",Direction="EGPD_Output",PinType.PinCategory="object",PinType.PinSubCategory="self",PinType.PinSubCategoryObject=None,PinType.PinSubCategoryMemberReference=(),PinType.bIsArray=False,PinType.bIsReference=False,PinType.bIsConst=False,PinType.bIsWeakPointer=False,LinkedTo=(K2Node_CallFunction_1861 471A4E874191001E241068AF8F5CA28F,),PersistentGuid=00000000000000000000000000000000,bHidden=False,bNotConnectable=False,bDefaultValueIsReadOnly=False,bDefaultValueIsIgnored=False,bAdvancedView=False,)
+			End Object
+
+	[/PARAM]
+[/OBJECT]
+[/EXCERPT:Function]
+
+
+
+[EXCERPT:ConScript]	
+[OBJECT:Snippet]
+	[PARAMLITERAL:id]
+		example_graph
+	[/PARAMLITERAL]
+	[PARAMLITERAL:type]
+		blueprint
+	[/PARAMLITERAL]
+	[PARAM:image]
+		![](Engine\Rendering\RenderTargets\BlueprintRenderTargets\HowTo\HeightFieldPainter\HFP_BP_Setup_03.png)
+	[/PARAM]
+	[PARAM:snippet]
+			BEGIN OBJECT
+							Begin Object Class=K2Node_FunctionEntry Name="K2Node_FunctionEntry_0"
+			   SignatureClass=Class'/Script/Engine.Actor'
+			   SignatureName="UserConstructionScript"
+			   NodePosY=-30
+			   NodeGuid=B0B1CB434FC9A9805F90B99BCC93EB5A
+			   CustomProperties Pin (PinId=5F831A4645C2410834600EB9C7A99B0D,PinName="then",Direction="EGPD_Output",PinType.PinCategory="exec",PinType.PinSubCategoryObject=None,PinType.PinSubCategoryMemberReference=(),PinType.bIsArray=False,PinType.bIsReference=False,PinType.bIsConst=False,PinType.bIsWeakPointer=False,LinkedTo=(K2Node_CallFunction_0 4BB5CF914F98AB4AD35FF7B67CFEAF4E,),PersistentGuid=00000000000000000000000000000000,bHidden=False,bNotConnectable=False,bDefaultValueIsReadOnly=False,bDefaultValueIsIgnored=False,bAdvancedView=False,)
+			End Object
+			Begin Object Class=K2Node_CallFunction Name="K2Node_CallFunction_0"
+			   FunctionReference=(MemberParent=Class'/Script/Engine.KismetRenderingLibrary',MemberName="CreateRenderTarget2D")
+			   NodePosX=208
+			   NodePosY=-30
+			   NodeGuid=A81BF0AE4E5D306079BC5DA665A30FB6
+			   CustomProperties Pin (PinId=4BB5CF914F98AB4AD35FF7B67CFEAF4E,PinName="execute",PinToolTip="\nExec",PinType.PinCategory="exec",PinType.PinSubCategoryObject=None,PinType.PinSubCategoryMemberReference=(),PinType.bIsArray=False,PinType.bIsReference=False,PinType.bIsConst=False,PinType.bIsWeakPointer=False,LinkedTo=(K2Node_FunctionEntry_0 5F831A4645C2410834600EB9C7A99B0D,),PersistentGuid=00000000000000000000000000000000,bHidden=False,bNotConnectable=False,bDefaultValueIsReadOnly=False,bDefaultValueIsIgnored=False,bAdvancedView=False,)
+			   CustomProperties Pin (PinId=1C1F266C47727A53B706C5963C47584D,PinName="then",PinToolTip="\nExec",Direction="EGPD_Output",PinType.PinCategory="exec",PinType.PinSubCategoryObject=None,PinType.PinSubCategoryMemberReference=(),PinType.bIsArray=False,PinType.bIsReference=False,PinType.bIsConst=False,PinType.bIsWeakPointer=False,LinkedTo=(K2Node_VariableSet_27 9C917A754928F23B03FABCA44048AAB2,),PersistentGuid=00000000000000000000000000000000,bHidden=False,bNotConnectable=False,bDefaultValueIsReadOnly=False,bDefaultValueIsIgnored=False,bAdvancedView=False,)
+			   CustomProperties Pin (PinId=57FB7A334A03A8BC40A3AD8571A16D56,PinName="self",PinFriendlyName=NSLOCTEXT("K2Node", "Target", "Target"),PinToolTip="Target\nKismet Rendering Library Reference",PinType.PinCategory="object",PinType.PinSubCategoryObject=Class'/Script/Engine.KismetRenderingLibrary',PinType.PinSubCategoryMemberReference=(),PinType.bIsArray=False,PinType.bIsReference=False,PinType.bIsConst=False,PinType.bIsWeakPointer=False,DefaultObject="/Script/Engine.Default__KismetRenderingLibrary",PersistentGuid=00000000000000000000000000000000,bHidden=True,bNotConnectable=False,bDefaultValueIsReadOnly=False,bDefaultValueIsIgnored=False,bAdvancedView=False,)
+			   CustomProperties Pin (PinId=32E396E642D88226B465C0BE447EC65C,PinName="WorldContextObject",PinToolTip="World Context Object\nObject Reference",PinType.PinCategory="object",PinType.PinSubCategoryObject=Class'/Script/CoreUObject.Object',PinType.PinSubCategoryMemberReference=(),PinType.bIsArray=False,PinType.bIsReference=False,PinType.bIsConst=False,PinType.bIsWeakPointer=False,PersistentGuid=00000000000000000000000000000000,bHidden=True,bNotConnectable=False,bDefaultValueIsReadOnly=False,bDefaultValueIsIgnored=False,bAdvancedView=False,)
+			   CustomProperties Pin (PinId=51636BFC453832DB294B3FAFD144FA68,PinName="Width",PinToolTip="Width\nInteger",PinType.PinCategory="int",PinType.PinSubCategoryObject=None,PinType.PinSubCategoryMemberReference=(),PinType.bIsArray=False,PinType.bIsReference=False,PinType.bIsConst=False,PinType.bIsWeakPointer=False,DefaultValue="256",AutogeneratedDefaultValue="256",PersistentGuid=00000000000000000000000000000000,bHidden=False,bNotConnectable=False,bDefaultValueIsReadOnly=False,bDefaultValueIsIgnored=False,bAdvancedView=False,)
+			   CustomProperties Pin (PinId=C88ADEAE4D60CB7889C29A97DF1CB399,PinName="Height",PinToolTip="Height\nInteger",PinType.PinCategory="int",PinType.PinSubCategoryObject=None,PinType.PinSubCategoryMemberReference=(),PinType.bIsArray=False,PinType.bIsReference=False,PinType.bIsConst=False,PinType.bIsWeakPointer=False,DefaultValue="256",AutogeneratedDefaultValue="256",PersistentGuid=00000000000000000000000000000000,bHidden=False,bNotConnectable=False,bDefaultValueIsReadOnly=False,bDefaultValueIsIgnored=False,bAdvancedView=False,)
+			   CustomProperties Pin (PinId=A9473E364E9F0CC00B30FA8D897DD7A3,PinName="ReturnValue",PinToolTip="Return Value\nTexture Render Target 2D Reference",Direction="EGPD_Output",PinType.PinCategory="object",PinType.PinSubCategoryObject=Class'/Script/Engine.TextureRenderTarget2D',PinType.PinSubCategoryMemberReference=(),PinType.bIsArray=False,PinType.bIsReference=False,PinType.bIsConst=False,PinType.bIsWeakPointer=False,LinkedTo=(K2Node_VariableSet_27 043C9ADD4D0D429377ACDE9DA6EDFFB5,),PersistentGuid=00000000000000000000000000000000,bHidden=False,bNotConnectable=False,bDefaultValueIsReadOnly=False,bDefaultValueIsIgnored=False,bAdvancedView=False,)
+			End Object
+			Begin Object Class=K2Node_VariableSet Name="K2Node_VariableSet_27"
+			   VariableReference=(MemberName="HeightfieldRT",MemberGuid=E71A12524941DA802B08039BF7491FF7,bSelfContext=True)
+			   NodePosX=480
+			   NodePosY=-14
+			   NodeGuid=2B08F5704E354C15504F67972BE7421E
+			   CustomProperties Pin (PinId=9C917A754928F23B03FABCA44048AAB2,PinName="execute",PinType.PinCategory="exec",PinType.PinSubCategoryObject=None,PinType.PinSubCategoryMemberReference=(),PinType.bIsArray=False,PinType.bIsReference=False,PinType.bIsConst=False,PinType.bIsWeakPointer=False,LinkedTo=(K2Node_CallFunction_0 1C1F266C47727A53B706C5963C47584D,),PersistentGuid=00000000000000000000000000000000,bHidden=False,bNotConnectable=False,bDefaultValueIsReadOnly=False,bDefaultValueIsIgnored=False,bAdvancedView=False,)
+			   CustomProperties Pin (PinId=A771633C4A2DF2AD370B4EA578B9121E,PinName="then",Direction="EGPD_Output",PinType.PinCategory="exec",PinType.PinSubCategoryObject=None,PinType.PinSubCategoryMemberReference=(),PinType.bIsArray=False,PinType.bIsReference=False,PinType.bIsConst=False,PinType.bIsWeakPointer=False,LinkedTo=(K2Node_CallMaterialParameterCollectionFunction_0 55284C3A40F1F1A33C806CA7209FCF48,),PersistentGuid=00000000000000000000000000000000,bHidden=False,bNotConnectable=False,bDefaultValueIsReadOnly=False,bDefaultValueIsIgnored=False,bAdvancedView=False,)
+			   CustomProperties Pin (PinId=043C9ADD4D0D429377ACDE9DA6EDFFB5,PinName="HeightfieldRT",PinType.PinCategory="object",PinType.PinSubCategoryObject=Class'/Script/Engine.TextureRenderTarget2D',PinType.PinSubCategoryMemberReference=(),PinType.bIsArray=False,PinType.bIsReference=False,PinType.bIsConst=False,PinType.bIsWeakPointer=False,LinkedTo=(K2Node_CallFunction_0 A9473E364E9F0CC00B30FA8D897DD7A3,),PersistentGuid=00000000000000000000000000000000,bHidden=False,bNotConnectable=False,bDefaultValueIsReadOnly=False,bDefaultValueIsIgnored=False,bAdvancedView=False,)
+			   CustomProperties Pin (PinId=14291FCB4A21ACDCB7731C914BDE67AD,PinName="Output_Get",PinToolTip="Retrieves the value of the variable, can use instead of a separate Get node",Direction="EGPD_Output",PinType.PinCategory="object",PinType.PinSubCategoryObject=Class'/Script/Engine.TextureRenderTarget2D',PinType.PinSubCategoryMemberReference=(),PinType.bIsArray=False,PinType.bIsReference=False,PinType.bIsConst=False,PinType.bIsWeakPointer=False,PersistentGuid=00000000000000000000000000000000,bHidden=False,bNotConnectable=False,bDefaultValueIsReadOnly=False,bDefaultValueIsIgnored=False,bAdvancedView=False,)
+			   CustomProperties Pin (PinId=395AB6754114998C89BA2FA5E9DF5E3C,PinName="self",PinFriendlyName=NSLOCTEXT("K2Node", "Target", "Target"),PinType.PinCategory="object",PinType.PinSubCategoryObject=BlueprintGeneratedClass'/Game/ExampleContent/BlueprintFluidSurface/HeightfieldPainter.HeightfieldPainter_C',PinType.PinSubCategoryMemberReference=(),PinType.bIsArray=False,PinType.bIsReference=False,PinType.bIsConst=False,PinType.bIsWeakPointer=False,PersistentGuid=00000000000000000000000000000000,bHidden=True,bNotConnectable=False,bDefaultValueIsReadOnly=False,bDefaultValueIsIgnored=False,bAdvancedView=False,)
+			End Object
+			Begin Object Class=K2Node_CallMaterialParameterCollectionFunction Name="K2Node_CallMaterialParameterCollectionFunction_0"
+			   FunctionReference=(MemberParent=Class'/Script/Engine.KismetMaterialLibrary',MemberName="CreateDynamicMaterialInstance")
+			   NodePosX=704
+			   NodePosY=-31
+			   NodeGuid=C8D401BA4B2296A5D4F72AAA61DC6A1A
+			   CustomProperties Pin (PinId=55284C3A40F1F1A33C806CA7209FCF48,PinName="execute",PinToolTip="\nExec",PinType.PinCategory="exec",PinType.PinSubCategoryObject=None,PinType.PinSubCategoryMemberReference=(),PinType.bIsArray=False,PinType.bIsReference=False,PinType.bIsConst=False,PinType.bIsWeakPointer=False,LinkedTo=(K2Node_VariableSet_27 A771633C4A2DF2AD370B4EA578B9121E,),PersistentGuid=00000000000000000000000000000000,bHidden=False,bNotConnectable=False,bDefaultValueIsReadOnly=False,bDefaultValueIsIgnored=False,bAdvancedView=False,)
+			   CustomProperties Pin (PinId=EB1773364D7C76BF4342B2B69834309E,PinName="then",PinToolTip="\nExec",Direction="EGPD_Output",PinType.PinCategory="exec",PinType.PinSubCategoryObject=None,PinType.PinSubCategoryMemberReference=(),PinType.bIsArray=False,PinType.bIsReference=False,PinType.bIsConst=False,PinType.bIsWeakPointer=False,LinkedTo=(K2Node_VariableSet_62 CE369A2E4A487F559D3CA1BFB60EF2C1,),PersistentGuid=00000000000000000000000000000000,bHidden=False,bNotConnectable=False,bDefaultValueIsReadOnly=False,bDefaultValueIsIgnored=False,bAdvancedView=False,)
+			   CustomProperties Pin (PinId=B062848843D00CAD4EA6A185646372A2,PinName="self",PinFriendlyName=NSLOCTEXT("K2Node", "Target", "Target"),PinToolTip="Target\nKismet Material Library Reference",PinType.PinCategory="object",PinType.PinSubCategoryObject=Class'/Script/Engine.KismetMaterialLibrary',PinType.PinSubCategoryMemberReference=(),PinType.bIsArray=False,PinType.bIsReference=False,PinType.bIsConst=False,PinType.bIsWeakPointer=False,DefaultObject="/Script/Engine.Default__KismetMaterialLibrary",PersistentGuid=00000000000000000000000000000000,bHidden=True,bNotConnectable=False,bDefaultValueIsReadOnly=False,bDefaultValueIsIgnored=False,bAdvancedView=False,)
+			   CustomProperties Pin (PinId=FD860AAB4D9A9E1F2E790185EF555B10,PinName="WorldContextObject",PinToolTip="World Context Object\nObject Reference",PinType.PinCategory="object",PinType.PinSubCategoryObject=Class'/Script/CoreUObject.Object',PinType.PinSubCategoryMemberReference=(),PinType.bIsArray=False,PinType.bIsReference=False,PinType.bIsConst=False,PinType.bIsWeakPointer=False,PersistentGuid=00000000000000000000000000000000,bHidden=True,bNotConnectable=False,bDefaultValueIsReadOnly=False,bDefaultValueIsIgnored=False,bAdvancedView=False,)
+			   CustomProperties Pin (PinId=1866D0004FF68655B91E79944DD00837,PinName="Parent",PinToolTip="Parent\nMaterial Interface Reference",PinType.PinCategory="object",PinType.PinSubCategoryObject=Class'/Script/Engine.MaterialInterface',PinType.PinSubCategoryMemberReference=(),PinType.bIsArray=False,PinType.bIsReference=False,PinType.bIsConst=False,PinType.bIsWeakPointer=False,DefaultObject="/Game/ExampleContent/BlueprintFluidSurface/HeightfieldPainterMaterial.HeightfieldPainterMaterial",PersistentGuid=00000000000000000000000000000000,bHidden=False,bNotConnectable=False,bDefaultValueIsReadOnly=False,bDefaultValueIsIgnored=False,bAdvancedView=False,)
+			   CustomProperties Pin (PinId=5C310A2941F5407764FD778BBCC6A295,PinName="ReturnValue",PinToolTip="Return Value\nMaterial Instance Dynamic Reference",Direction="EGPD_Output",PinType.PinCategory="object",PinType.PinSubCategoryObject=Class'/Script/Engine.MaterialInstanceDynamic',PinType.PinSubCategoryMemberReference=(),PinType.bIsArray=False,PinType.bIsReference=False,PinType.bIsConst=False,PinType.bIsWeakPointer=False,LinkedTo=(K2Node_VariableSet_62 4CC233BC4949E821C89377A307B8AA44,),PersistentGuid=00000000000000000000000000000000,bHidden=False,bNotConnectable=False,bDefaultValueIsReadOnly=False,bDefaultValueIsIgnored=False,bAdvancedView=False,)
+			End Object
+			Begin Object Class=K2Node_VariableSet Name="K2Node_VariableSet_62"
+			   VariableReference=(MemberName="PainterMaterialInstance",MemberGuid=9030CC2445DD109F0DC10BB00BA7797D,bSelfContext=True)
+			   NodePosX=1024
+			   NodePosY=-15
+			   NodeGuid=DB54514D4AB2D6265BEF468C399ACBFA
+			   CustomProperties Pin (PinId=CE369A2E4A487F559D3CA1BFB60EF2C1,PinName="execute",PinType.PinCategory="exec",PinType.PinSubCategoryObject=None,PinType.PinSubCategoryMemberReference=(),PinType.bIsArray=False,PinType.bIsReference=False,PinType.bIsConst=False,PinType.bIsWeakPointer=False,LinkedTo=(K2Node_CallMaterialParameterCollectionFunction_0 EB1773364D7C76BF4342B2B69834309E,),PersistentGuid=00000000000000000000000000000000,bHidden=False,bNotConnectable=False,bDefaultValueIsReadOnly=False,bDefaultValueIsIgnored=False,bAdvancedView=False,)
+			   CustomProperties Pin (PinId=5468191F471453588ED7E09E3381E913,PinName="then",Direction="EGPD_Output",PinType.PinCategory="exec",PinType.PinSubCategoryObject=None,PinType.PinSubCategoryMemberReference=(),PinType.bIsArray=False,PinType.bIsReference=False,PinType.bIsConst=False,PinType.bIsWeakPointer=False,LinkedTo=(K2Node_CallFunction_7 C230F0494A1C54C96FB6FD8F22B2E9F8,),PersistentGuid=00000000000000000000000000000000,bHidden=False,bNotConnectable=False,bDefaultValueIsReadOnly=False,bDefaultValueIsIgnored=False,bAdvancedView=False,)
+			   CustomProperties Pin (PinId=4CC233BC4949E821C89377A307B8AA44,PinName="PainterMaterialInstance",PinType.PinCategory="object",PinType.PinSubCategoryObject=Class'/Script/Engine.MaterialInstanceDynamic',PinType.PinSubCategoryMemberReference=(),PinType.bIsArray=False,PinType.bIsReference=False,PinType.bIsConst=False,PinType.bIsWeakPointer=False,LinkedTo=(K2Node_CallMaterialParameterCollectionFunction_0 5C310A2941F5407764FD778BBCC6A295,),PersistentGuid=00000000000000000000000000000000,bHidden=False,bNotConnectable=False,bDefaultValueIsReadOnly=False,bDefaultValueIsIgnored=False,bAdvancedView=False,)
+			   CustomProperties Pin (PinId=58ECD5564609BD67961F0691C09BDD97,PinName="Output_Get",PinToolTip="Retrieves the value of the variable, can use instead of a separate Get node",Direction="EGPD_Output",PinType.PinCategory="object",PinType.PinSubCategoryObject=Class'/Script/Engine.MaterialInstanceDynamic',PinType.PinSubCategoryMemberReference=(),PinType.bIsArray=False,PinType.bIsReference=False,PinType.bIsConst=False,PinType.bIsWeakPointer=False,LinkedTo=(K2Node_CallFunction_7 5449122648152C4B0D3527B990CD0693,),PersistentGuid=00000000000000000000000000000000,bHidden=False,bNotConnectable=False,bDefaultValueIsReadOnly=False,bDefaultValueIsIgnored=False,bAdvancedView=False,)
+			   CustomProperties Pin (PinId=75739663431A6A39638566A68D0D1D28,PinName="self",PinFriendlyName=NSLOCTEXT("K2Node", "Target", "Target"),PinType.PinCategory="object",PinType.PinSubCategoryObject=BlueprintGeneratedClass'/Game/ExampleContent/BlueprintFluidSurface/HeightfieldPainter.HeightfieldPainter_C',PinType.PinSubCategoryMemberReference=(),PinType.bIsArray=False,PinType.bIsReference=False,PinType.bIsConst=False,PinType.bIsWeakPointer=False,PersistentGuid=00000000000000000000000000000000,bHidden=True,bNotConnectable=False,bDefaultValueIsReadOnly=False,bDefaultValueIsIgnored=False,bAdvancedView=False,)
+			End Object
+			Begin Object Class=K2Node_CallFunction Name="K2Node_CallFunction_7"
+			   FunctionReference=(MemberParent=Class'/Script/Engine.PrimitiveComponent',MemberName="SetMaterial")
+			   NodePosX=1360
+			   NodePosY=-47
+			   NodeGuid=8BACB33648DD06793EFE40907DF1DF6D
+			   CustomProperties Pin (PinId=C230F0494A1C54C96FB6FD8F22B2E9F8,PinName="execute",PinType.PinCategory="exec",PinType.PinSubCategoryObject=None,PinType.PinSubCategoryMemberReference=(),PinType.bIsArray=False,PinType.bIsReference=False,PinType.bIsConst=False,PinType.bIsWeakPointer=False,LinkedTo=(K2Node_VariableSet_62 5468191F471453588ED7E09E3381E913,),PersistentGuid=00000000000000000000000000000000,bHidden=False,bNotConnectable=False,bDefaultValueIsReadOnly=False,bDefaultValueIsIgnored=False,bAdvancedView=False,)
+			   CustomProperties Pin (PinId=F2CC079C442B5A9B2D1144AA10496547,PinName="then",Direction="EGPD_Output",PinType.PinCategory="exec",PinType.PinSubCategoryObject=None,PinType.PinSubCategoryMemberReference=(),PinType.bIsArray=False,PinType.bIsReference=False,PinType.bIsConst=False,PinType.bIsWeakPointer=False,LinkedTo=(K2Node_CallMaterialParameterCollectionFunction_4 8171209541AC0BE2FE3AE3AB0E232778,),PersistentGuid=00000000000000000000000000000000,bHidden=False,bNotConnectable=False,bDefaultValueIsReadOnly=False,bDefaultValueIsIgnored=False,bAdvancedView=False,)
+			   CustomProperties Pin (PinId=44FC11A848C584BF32FBF9904A82DE73,PinName="self",PinFriendlyName=NSLOCTEXT("K2Node", "Target", "Target"),PinType.PinCategory="object",PinType.PinSubCategoryObject=Class'/Script/Engine.PrimitiveComponent',PinType.PinSubCategoryMemberReference=(),PinType.bIsArray=False,PinType.bIsReference=False,PinType.bIsConst=False,PinType.bIsWeakPointer=False,LinkedTo=(K2Node_VariableGet_3 621B77D54968165AC2B8CEB2FF6E7F57,),PersistentGuid=00000000000000000000000000000000,bHidden=False,bNotConnectable=False,bDefaultValueIsReadOnly=False,bDefaultValueIsIgnored=False,bAdvancedView=False,)
+			   CustomProperties Pin (PinId=7B23D20A402FF6510AF0AAA9FCF0B9CB,PinName="ElementIndex",PinType.PinCategory="int",PinType.PinSubCategoryObject=None,PinType.PinSubCategoryMemberReference=(),PinType.bIsArray=False,PinType.bIsReference=False,PinType.bIsConst=False,PinType.bIsWeakPointer=False,DefaultValue="0",AutogeneratedDefaultValue="0",PersistentGuid=00000000000000000000000000000000,bHidden=False,bNotConnectable=False,bDefaultValueIsReadOnly=False,bDefaultValueIsIgnored=False,bAdvancedView=False,)
+			   CustomProperties Pin (PinId=5449122648152C4B0D3527B990CD0693,PinName="Material",PinType.PinCategory="object",PinType.PinSubCategoryObject=Class'/Script/Engine.MaterialInterface',PinType.PinSubCategoryMemberReference=(),PinType.bIsArray=False,PinType.bIsReference=False,PinType.bIsConst=False,PinType.bIsWeakPointer=False,LinkedTo=(K2Node_VariableSet_62 58ECD5564609BD67961F0691C09BDD97,),PersistentGuid=00000000000000000000000000000000,bHidden=False,bNotConnectable=False,bDefaultValueIsReadOnly=False,bDefaultValueIsIgnored=False,bAdvancedView=False,)
+			End Object
+			Begin Object Class=K2Node_VariableGet Name="K2Node_VariableGet_3"
+			   VariableReference=(MemberName="StaticMeshComponent",bSelfContext=True)
+			   NodePosX=1088
+			   NodePosY=80
+			   NodeGuid=5B01DC5349925C33381EE5AA54F19521
+			   CustomProperties Pin (PinId=621B77D54968165AC2B8CEB2FF6E7F57,PinName="StaticMeshComponent",Direction="EGPD_Output",PinType.PinCategory="object",PinType.PinSubCategoryObject=Class'/Script/Engine.StaticMeshComponent',PinType.PinSubCategoryMemberReference=(),PinType.bIsArray=False,PinType.bIsReference=False,PinType.bIsConst=False,PinType.bIsWeakPointer=False,LinkedTo=(K2Node_CallFunction_7 44FC11A848C584BF32FBF9904A82DE73,),PersistentGuid=00000000000000000000000000000000,bHidden=False,bNotConnectable=False,bDefaultValueIsReadOnly=False,bDefaultValueIsIgnored=False,bAdvancedView=False,)
+			   CustomProperties Pin (PinId=F76E900D4E69951E166D27A38BEF0036,PinName="self",PinFriendlyName=NSLOCTEXT("K2Node", "Target", "Target"),PinType.PinCategory="object",PinType.PinSubCategoryObject=Class'/Script/Engine.StaticMeshActor',PinType.PinSubCategoryMemberReference=(),PinType.bIsArray=False,PinType.bIsReference=False,PinType.bIsConst=False,PinType.bIsWeakPointer=False,PersistentGuid=00000000000000000000000000000000,bHidden=True,bNotConnectable=False,bDefaultValueIsReadOnly=False,bDefaultValueIsIgnored=False,bAdvancedView=False,)
+			End Object
+			Begin Object Class=K2Node_CallMaterialParameterCollectionFunction Name="K2Node_CallMaterialParameterCollectionFunction_4"
+			   FunctionReference=(MemberParent=Class'/Script/Engine.KismetMaterialLibrary',MemberName="CreateDynamicMaterialInstance")
+			   NodePosX=1616
+			   NodePosY=-32
+			   NodeGuid=DFF317B74AF637BD1362D6A2690ABE63
+			   CustomProperties Pin (PinId=8171209541AC0BE2FE3AE3AB0E232778,PinName="execute",PinToolTip="\nExec",PinType.PinCategory="exec",PinType.PinSubCategoryObject=None,PinType.PinSubCategoryMemberReference=(),PinType.bIsArray=False,PinType.bIsReference=False,PinType.bIsConst=False,PinType.bIsWeakPointer=False,LinkedTo=(K2Node_CallFunction_7 F2CC079C442B5A9B2D1144AA10496547,),PersistentGuid=00000000000000000000000000000000,bHidden=False,bNotConnectable=False,bDefaultValueIsReadOnly=False,bDefaultValueIsIgnored=False,bAdvancedView=False,)
+			   CustomProperties Pin (PinId=E21E08D942166B036E5FC38DE9949801,PinName="then",PinToolTip="\nExec",Direction="EGPD_Output",PinType.PinCategory="exec",PinType.PinSubCategoryObject=None,PinType.PinSubCategoryMemberReference=(),PinType.bIsArray=False,PinType.bIsReference=False,PinType.bIsConst=False,PinType.bIsWeakPointer=False,LinkedTo=(K2Node_VariableSet_87 0325DE6C44B81F336F343AA55852BE19,),PersistentGuid=00000000000000000000000000000000,bHidden=False,bNotConnectable=False,bDefaultValueIsReadOnly=False,bDefaultValueIsIgnored=False,bAdvancedView=False,)
+			   CustomProperties Pin (PinId=70A6FE7E4656F2E19DB4409A253B5E03,PinName="self",PinFriendlyName=NSLOCTEXT("K2Node", "Target", "Target"),PinToolTip="Target\nKismet Material Library Reference",PinType.PinCategory="object",PinType.PinSubCategoryObject=Class'/Script/Engine.KismetMaterialLibrary',PinType.PinSubCategoryMemberReference=(),PinType.bIsArray=False,PinType.bIsReference=False,PinType.bIsConst=False,PinType.bIsWeakPointer=False,DefaultObject="/Script/Engine.Default__KismetMaterialLibrary",PersistentGuid=00000000000000000000000000000000,bHidden=True,bNotConnectable=False,bDefaultValueIsReadOnly=False,bDefaultValueIsIgnored=False,bAdvancedView=False,)
+			   CustomProperties Pin (PinId=686C837444AAE4CBE52D3781B2D8B908,PinName="WorldContextObject",PinToolTip="World Context Object\nObject Reference",PinType.PinCategory="object",PinType.PinSubCategoryObject=Class'/Script/CoreUObject.Object',PinType.PinSubCategoryMemberReference=(),PinType.bIsArray=False,PinType.bIsReference=False,PinType.bIsConst=False,PinType.bIsWeakPointer=False,PersistentGuid=00000000000000000000000000000000,bHidden=True,bNotConnectable=False,bDefaultValueIsReadOnly=False,bDefaultValueIsIgnored=False,bAdvancedView=False,)
+			   CustomProperties Pin (PinId=AF370D77468D9FC15882E0997AF8C55D,PinName="Parent",PinToolTip="Parent\nMaterial Interface Reference",PinType.PinCategory="object",PinType.PinSubCategoryObject=Class'/Script/Engine.MaterialInterface',PinType.PinSubCategoryMemberReference=(),PinType.bIsArray=False,PinType.bIsReference=False,PinType.bIsConst=False,PinType.bIsWeakPointer=False,DefaultObject="/Game/ExampleContent/BlueprintFluidSurface/ForceSplat.ForceSplat",PersistentGuid=00000000000000000000000000000000,bHidden=False,bNotConnectable=False,bDefaultValueIsReadOnly=False,bDefaultValueIsIgnored=False,bAdvancedView=False,)
+			   CustomProperties Pin (PinId=0E4E15EF4A45EF8796AD66B098B2BF5D,PinName="ReturnValue",PinToolTip="Return Value\nMaterial Instance Dynamic Reference",Direction="EGPD_Output",PinType.PinCategory="object",PinType.PinSubCategoryObject=Class'/Script/Engine.MaterialInstanceDynamic',PinType.PinSubCategoryMemberReference=(),PinType.bIsArray=False,PinType.bIsReference=False,PinType.bIsConst=False,PinType.bIsWeakPointer=False,LinkedTo=(K2Node_VariableSet_87 1B72E0D243E28349EA1939B5937A1FA3,),PersistentGuid=00000000000000000000000000000000,bHidden=False,bNotConnectable=False,bDefaultValueIsReadOnly=False,bDefaultValueIsIgnored=False,bAdvancedView=False,)
+			End Object
+			Begin Object Class=K2Node_VariableSet Name="K2Node_VariableSet_87"
+			   VariableReference=(MemberName="ForceSplatMID",MemberGuid=20FE6B2D4954BEA0CD113E92370334C9,bSelfContext=True)
+			   NodePosX=1984
+			   NodePosY=-16
+			   NodeGuid=7150A4F544FD83115DC4E1AC0F232497
+			   CustomProperties Pin (PinId=0325DE6C44B81F336F343AA55852BE19,PinName="execute",PinType.PinCategory="exec",PinType.PinSubCategoryObject=None,PinType.PinSubCategoryMemberReference=(),PinType.bIsArray=False,PinType.bIsReference=False,PinType.bIsConst=False,PinType.bIsWeakPointer=False,LinkedTo=(K2Node_CallMaterialParameterCollectionFunction_4 E21E08D942166B036E5FC38DE9949801,),PersistentGuid=00000000000000000000000000000000,bHidden=False,bNotConnectable=False,bDefaultValueIsReadOnly=False,bDefaultValueIsIgnored=False,bAdvancedView=False,)
+			   CustomProperties Pin (PinId=9F8DC7C3470FFEEE9D3663AA495FDB06,PinName="then",Direction="EGPD_Output",PinType.PinCategory="exec",PinType.PinSubCategoryObject=None,PinType.PinSubCategoryMemberReference=(),PinType.bIsArray=False,PinType.bIsReference=False,PinType.bIsConst=False,PinType.bIsWeakPointer=False,PersistentGuid=00000000000000000000000000000000,bHidden=False,bNotConnectable=False,bDefaultValueIsReadOnly=False,bDefaultValueIsIgnored=False,bAdvancedView=False,)
+			   CustomProperties Pin (PinId=1B72E0D243E28349EA1939B5937A1FA3,PinName="ForceSplatMID",PinType.PinCategory="object",PinType.PinSubCategoryObject=Class'/Script/Engine.MaterialInstanceDynamic',PinType.PinSubCategoryMemberReference=(),PinType.bIsArray=False,PinType.bIsReference=False,PinType.bIsConst=False,PinType.bIsWeakPointer=False,LinkedTo=(K2Node_CallMaterialParameterCollectionFunction_4 0E4E15EF4A45EF8796AD66B098B2BF5D,),PersistentGuid=00000000000000000000000000000000,bHidden=False,bNotConnectable=False,bDefaultValueIsReadOnly=False,bDefaultValueIsIgnored=False,bAdvancedView=False,)
+			   CustomProperties Pin (PinId=50FE09C64720F61910ACF390E9D9F8C5,PinName="Output_Get",PinToolTip="Retrieves the value of the variable, can use instead of a separate Get node",Direction="EGPD_Output",PinType.PinCategory="object",PinType.PinSubCategoryObject=Class'/Script/Engine.MaterialInstanceDynamic',PinType.PinSubCategoryMemberReference=(),PinType.bIsArray=False,PinType.bIsReference=False,PinType.bIsConst=False,PinType.bIsWeakPointer=False,PersistentGuid=00000000000000000000000000000000,bHidden=False,bNotConnectable=False,bDefaultValueIsReadOnly=False,bDefaultValueIsIgnored=False,bAdvancedView=False,)
+			   CustomProperties Pin (PinId=D309D0AD478263FA889346B800C4AB96,PinName="self",PinFriendlyName=NSLOCTEXT("K2Node", "Target", "Target"),PinType.PinCategory="object",PinType.PinSubCategoryObject=BlueprintGeneratedClass'/Game/ExampleContent/BlueprintFluidSurface/HeightfieldPainter.HeightfieldPainter_C',PinType.PinSubCategoryMemberReference=(),PinType.bIsArray=False,PinType.bIsReference=False,PinType.bIsConst=False,PinType.bIsWeakPointer=False,PersistentGuid=00000000000000000000000000000000,bHidden=True,bNotConnectable=False,bDefaultValueIsReadOnly=False,bDefaultValueIsIgnored=False,bAdvancedView=False,)
+			End Object
+			Begin Object Class=EdGraphNode_Comment Name="EdGraphNode_Comment_2"
+			   NodePosX=-48
+			   NodePosY=-96
+			   NodeWidth=2279
+			   NodeHeight=261
+			   NodeComment="Setup Render Target Static Mesh and Dynamic Material Instances"
+			   NodeGuid=A65A76EE472D19FC053F51A39842B39E
+			End Object
+
+			END OBJECT
+	[/PARAM]
+[/OBJECT]
+[/EXCERPT:ConScript]
+
+[EXCERPT:EG]	
+[OBJECT:Snippet]
+	[PARAMLITERAL:id]
+		EventGraph
+	[/PARAMLITERAL]
+	[PARAMLITERAL:type]
+		blueprint
+	[/PARAMLITERAL]
+	[PARAM:image]
+		![](Engine\Rendering\RenderTargets\BlueprintRenderTargets\HowTo\HeightFieldPainter\HFP_BP_Setup_07.png)
+	[/PARAM]
+	[PARAM:snippet]
+			BEGIN OBJECT
+				Begin Object Class=K2Node_Event Name="K2Node_Event_0"
+			   EventReference=(MemberParent=Class'/Script/Engine.Actor',MemberName="ReceiveBeginPlay")
+			   bOverrideFunction=True
+			   NodePosX=-224
+			   NodePosY=1119
+			   bCommentBubblePinned=True
+			   NodeGuid=DBF784C743D229EAF492849964D015AF
+			   CustomProperties Pin (PinId=2BE1644F47C274111DE6EBA0A1855DDE,PinName="OutputDelegate",Direction="EGPD_Output",PinType.PinCategory="delegate",PinType.PinSubCategoryObject=None,PinType.PinSubCategoryMemberReference=(MemberParent=Class'/Script/Engine.Actor',MemberName="ReceiveBeginPlay"),PinType.bIsArray=False,PinType.bIsReference=False,PinType.bIsConst=False,PinType.bIsWeakPointer=False,PersistentGuid=00000000000000000000000000000000,bHidden=False,bNotConnectable=False,bDefaultValueIsReadOnly=False,bDefaultValueIsIgnored=False,bAdvancedView=False,)
+			   CustomProperties Pin (PinId=4C1A3CD54C93D4A74803118ED8E9BE79,PinName="then",Direction="EGPD_Output",PinType.PinCategory="exec",PinType.PinSubCategoryObject=None,PinType.PinSubCategoryMemberReference=(),PinType.bIsArray=False,PinType.bIsReference=False,PinType.bIsConst=False,PinType.bIsWeakPointer=False,LinkedTo=(K2Node_CallFunction_371 59EBD0D84FA02620751F0CAC38ADA5FF,),PersistentGuid=00000000000000000000000000000000,bHidden=False,bNotConnectable=False,bDefaultValueIsReadOnly=False,bDefaultValueIsIgnored=False,bAdvancedView=False,)
+			End Object
+			Begin Object Class=K2Node_Event Name="K2Node_Event_57"
+			   EventReference=(MemberParent=Class'/Script/Engine.Actor',MemberName="ReceiveTick")
+			   bOverrideFunction=True
+			   NodePosX=400
+			   NodePosY=1611
+			   bCommentBubblePinned=True
+			   NodeGuid=B6C641E94DC295D2CB53B189AE27253F
+			   CustomProperties Pin (PinId=40C56D29448EE7B7BC3BD4A6B5EEBE71,PinName="OutputDelegate",Direction="EGPD_Output",PinType.PinCategory="delegate",PinType.PinSubCategoryObject=None,PinType.PinSubCategoryMemberReference=(MemberParent=Class'/Script/Engine.Actor',MemberName="ReceiveTick"),PinType.bIsArray=False,PinType.bIsReference=False,PinType.bIsConst=False,PinType.bIsWeakPointer=False,PersistentGuid=00000000000000000000000000000000,bHidden=False,bNotConnectable=False,bDefaultValueIsReadOnly=False,bDefaultValueIsIgnored=False,bAdvancedView=False,)
+			   CustomProperties Pin (PinId=0371202549E0ACFD644FD4A3E8B79FD5,PinName="then",Direction="EGPD_Output",PinType.PinCategory="exec",PinType.PinSubCategoryObject=None,PinType.PinSubCategoryMemberReference=(),PinType.bIsArray=False,PinType.bIsReference=False,PinType.bIsConst=False,PinType.bIsWeakPointer=False,LinkedTo=(K2Node_IfThenElse_25 997FDA76493A40CEE900CD81079C2BFC,),PersistentGuid=00000000000000000000000000000000,bHidden=False,bNotConnectable=False,bDefaultValueIsReadOnly=False,bDefaultValueIsIgnored=False,bAdvancedView=False,)
+			   CustomProperties Pin (PinId=8C4E9AE4471B9BF4C1279EBA05819B9C,PinName="DeltaSeconds",PinToolTip="Delta Seconds\nFloat",Direction="EGPD_Output",PinType.PinCategory="float",PinType.PinSubCategoryObject=None,PinType.PinSubCategoryMemberReference=(),PinType.bIsArray=False,PinType.bIsReference=False,PinType.bIsConst=False,PinType.bIsWeakPointer=False,DefaultValue="0.0",AutogeneratedDefaultValue="0.0",PersistentGuid=00000000000000000000000000000000,bHidden=False,bNotConnectable=False,bDefaultValueIsReadOnly=False,bDefaultValueIsIgnored=False,bAdvancedView=False,)
+			End Object
+			Begin Object Class=K2Node_CallFunction Name="K2Node_CallFunction_359"
+			   FunctionReference=(MemberParent=Class'/Script/Engine.KismetRenderingLibrary',MemberName="ClearRenderTarget2D")
+			   NodePosX=592
+			   NodePosY=1119
+			   NodeGuid=A9AC87F647DE4F49AD4EB3ADBAE0AC82
+			   CustomProperties Pin (PinId=50F50A054429DFDB03CCE884B76C4E0A,PinName="execute",PinType.PinCategory="exec",PinType.PinSubCategoryObject=None,PinType.PinSubCategoryMemberReference=(),PinType.bIsArray=False,PinType.bIsReference=False,PinType.bIsConst=False,PinType.bIsWeakPointer=False,LinkedTo=(K2Node_CallFunction_371 C8543EC14C417EB94C0A40806572C540,),PersistentGuid=00000000000000000000000000000000,bHidden=False,bNotConnectable=False,bDefaultValueIsReadOnly=False,bDefaultValueIsIgnored=False,bAdvancedView=False,)
+			   CustomProperties Pin (PinId=43C94C3C4FB04A4DEADC37A9E3BDE0B2,PinName="then",Direction="EGPD_Output",PinType.PinCategory="exec",PinType.PinSubCategoryObject=None,PinType.PinSubCategoryMemberReference=(),PinType.bIsArray=False,PinType.bIsReference=False,PinType.bIsConst=False,PinType.bIsWeakPointer=False,LinkedTo=(K2Node_CallFunction_370 542CC73148D38A37BCCECE90262AFC49,),PersistentGuid=00000000000000000000000000000000,bHidden=False,bNotConnectable=False,bDefaultValueIsReadOnly=False,bDefaultValueIsIgnored=False,bAdvancedView=False,)
+			   CustomProperties Pin (PinId=E0E9F0EF4A1840A57B003C8528C144D0,PinName="self",PinFriendlyName=NSLOCTEXT("K2Node", "Target", "Target"),PinType.PinCategory="object",PinType.PinSubCategoryObject=Class'/Script/Engine.KismetRenderingLibrary',PinType.PinSubCategoryMemberReference=(),PinType.bIsArray=False,PinType.bIsReference=False,PinType.bIsConst=False,PinType.bIsWeakPointer=False,DefaultObject="/Script/Engine.Default__KismetRenderingLibrary",PersistentGuid=00000000000000000000000000000000,bHidden=True,bNotConnectable=False,bDefaultValueIsReadOnly=False,bDefaultValueIsIgnored=False,bAdvancedView=False,)
+			   CustomProperties Pin (PinId=79C9EF7C406FCCFB6D58E7971F16CDC8,PinName="WorldContextObject",PinType.PinCategory="object",PinType.PinSubCategoryObject=Class'/Script/CoreUObject.Object',PinType.PinSubCategoryMemberReference=(),PinType.bIsArray=False,PinType.bIsReference=False,PinType.bIsConst=False,PinType.bIsWeakPointer=False,PersistentGuid=00000000000000000000000000000000,bHidden=True,bNotConnectable=False,bDefaultValueIsReadOnly=False,bDefaultValueIsIgnored=False,bAdvancedView=False,)
+			   CustomProperties Pin (PinId=314A49414552DC481563F19C80A6DDDF,PinName="TextureRenderTarget",PinType.PinCategory="object",PinType.PinSubCategoryObject=Class'/Script/Engine.TextureRenderTarget2D',PinType.PinSubCategoryMemberReference=(),PinType.bIsArray=False,PinType.bIsReference=False,PinType.bIsConst=False,PinType.bIsWeakPointer=False,LinkedTo=(K2Node_VariableGet_178 1015D25442FF3C8A3B214DA5835B8950,),PersistentGuid=00000000000000000000000000000000,bHidden=False,bNotConnectable=False,bDefaultValueIsReadOnly=False,bDefaultValueIsIgnored=False,bAdvancedView=False,)
+			   CustomProperties Pin (PinId=F48572574E7B44CBAA94708BF8EAAAE5,PinName="ClearColor",PinType.PinCategory="struct",PinType.PinSubCategoryObject=ScriptStruct'/Script/CoreUObject.LinearColor',PinType.PinSubCategoryMemberReference=(),PinType.bIsArray=False,PinType.bIsReference=False,PinType.bIsConst=False,PinType.bIsWeakPointer=False,PersistentGuid=00000000000000000000000000000000,bHidden=False,bNotConnectable=False,bDefaultValueIsReadOnly=False,bDefaultValueIsIgnored=False,bAdvancedView=False,)
+			End Object
+			Begin Object Class=K2Node_VariableGet Name="K2Node_VariableGet_178"
+			   VariableReference=(MemberName="HeightfieldRT",MemberGuid=E71A12524941DA802B08039BF7491FF7,bSelfContext=True)
+			   NodePosX=384
+			   NodePosY=1200
+			   NodeGuid=A024EAB542C2204C79B095B1968A22D6
+			   CustomProperties Pin (PinId=1015D25442FF3C8A3B214DA5835B8950,PinName="HeightfieldRT",Direction="EGPD_Output",PinType.PinCategory="object",PinType.PinSubCategoryObject=Class'/Script/Engine.TextureRenderTarget2D',PinType.PinSubCategoryMemberReference=(),PinType.bIsArray=False,PinType.bIsReference=False,PinType.bIsConst=False,PinType.bIsWeakPointer=False,LinkedTo=(K2Node_CallFunction_359 314A49414552DC481563F19C80A6DDDF,),PersistentGuid=00000000000000000000000000000000,bHidden=False,bNotConnectable=False,bDefaultValueIsReadOnly=False,bDefaultValueIsIgnored=False,bAdvancedView=False,)
+			   CustomProperties Pin (PinId=07915D7C4D3DC0942D902586DF54E782,PinName="self",PinFriendlyName=NSLOCTEXT("K2Node", "Target", "Target"),PinType.PinCategory="object",PinType.PinSubCategoryObject=BlueprintGeneratedClass'/Game/ExampleContent/BlueprintFluidSurface/HeightfieldPainter.HeightfieldPainter_C',PinType.PinSubCategoryMemberReference=(),PinType.bIsArray=False,PinType.bIsReference=False,PinType.bIsConst=False,PinType.bIsWeakPointer=False,PersistentGuid=00000000000000000000000000000000,bHidden=True,bNotConnectable=False,bDefaultValueIsReadOnly=False,bDefaultValueIsIgnored=False,bAdvancedView=False,)
+			End Object
+			Begin Object Class=K2Node_InputKey Name="K2Node_InputKey_19"
+			   InputKey=LeftMouseButton
+			   NodePosX=-224
+			   NodePosY=1584
+			   NodeGuid=87A3C11B4FEFE0378592B58E6DA37F8B
+			   CustomProperties Pin (PinId=4D5FE2B549E275EDF0E659896A9261E4,PinName="Pressed",Direction="EGPD_Output",PinType.PinCategory="exec",PinType.PinSubCategoryObject=None,PinType.PinSubCategoryMemberReference=(),PinType.bIsArray=False,PinType.bIsReference=False,PinType.bIsConst=False,PinType.bIsWeakPointer=False,LinkedTo=(K2Node_VariableSet_41 DF34F0B548E46FCCB2AB54856F7CFFC5,),PersistentGuid=00000000000000000000000000000000,bHidden=False,bNotConnectable=False,bDefaultValueIsReadOnly=False,bDefaultValueIsIgnored=False,bAdvancedView=False,)
+			   CustomProperties Pin (PinId=B150F43D4C370A3041E5B9A2DD8FC834,PinName="Released",Direction="EGPD_Output",PinType.PinCategory="exec",PinType.PinSubCategoryObject=None,PinType.PinSubCategoryMemberReference=(),PinType.bIsArray=False,PinType.bIsReference=False,PinType.bIsConst=False,PinType.bIsWeakPointer=False,LinkedTo=(K2Node_VariableSet_42 9312E8C94B08C9FAF76638810A0F61A1,),PersistentGuid=00000000000000000000000000000000,bHidden=False,bNotConnectable=False,bDefaultValueIsReadOnly=False,bDefaultValueIsIgnored=False,bAdvancedView=False,)
+			   CustomProperties Pin (PinId=43D7BA3045BD03142E77E597FD7EAD86,PinName="Key",Direction="EGPD_Output",PinType.PinCategory="struct",PinType.PinSubCategoryObject=ScriptStruct'/Script/InputCore.Key',PinType.PinSubCategoryMemberReference=(),PinType.bIsArray=False,PinType.bIsReference=False,PinType.bIsConst=False,PinType.bIsWeakPointer=False,DefaultValue="AnyKey",PersistentGuid=00000000000000000000000000000000,bHidden=False,bNotConnectable=False,bDefaultValueIsReadOnly=False,bDefaultValueIsIgnored=False,bAdvancedView=False,)
+			End Object
+			Begin Object Class=K2Node_VariableSet Name="K2Node_VariableSet_41"
+			   VariableReference=(MemberName="MouseDown",MemberGuid=209E1D534C15EFFA81760DA2C7F9D56C,bSelfContext=True)
+			   NodePosX=80
+			   NodePosY=1568
+			   NodeGuid=729490D24E2E13816637E2B6C9D6CFD8
+			   CustomProperties Pin (PinId=DF34F0B548E46FCCB2AB54856F7CFFC5,PinName="execute",PinType.PinCategory="exec",PinType.PinSubCategoryObject=None,PinType.PinSubCategoryMemberReference=(),PinType.bIsArray=False,PinType.bIsReference=False,PinType.bIsConst=False,PinType.bIsWeakPointer=False,LinkedTo=(K2Node_InputKey_19 4D5FE2B549E275EDF0E659896A9261E4,),PersistentGuid=00000000000000000000000000000000,bHidden=False,bNotConnectable=False,bDefaultValueIsReadOnly=False,bDefaultValueIsIgnored=False,bAdvancedView=False,)
+			   CustomProperties Pin (PinId=A68EDFFF4BFBADBD6B1C60AB8EEE7A3D,PinName="then",Direction="EGPD_Output",PinType.PinCategory="exec",PinType.PinSubCategoryObject=None,PinType.PinSubCategoryMemberReference=(),PinType.bIsArray=False,PinType.bIsReference=False,PinType.bIsConst=False,PinType.bIsWeakPointer=False,PersistentGuid=00000000000000000000000000000000,bHidden=False,bNotConnectable=False,bDefaultValueIsReadOnly=False,bDefaultValueIsIgnored=False,bAdvancedView=False,)
+			   CustomProperties Pin (PinId=8F7952DE47C7BACDD7565587C283D7EE,PinName="MouseDown",PinType.PinCategory="bool",PinType.PinSubCategoryObject=None,PinType.PinSubCategoryMemberReference=(),PinType.bIsArray=False,PinType.bIsReference=False,PinType.bIsConst=False,PinType.bIsWeakPointer=False,DefaultValue="true",AutogeneratedDefaultValue="false",PersistentGuid=00000000000000000000000000000000,bHidden=False,bNotConnectable=False,bDefaultValueIsReadOnly=False,bDefaultValueIsIgnored=False,bAdvancedView=False,)
+			   CustomProperties Pin (PinId=6799AEA64465D17287DE96888FB9C1D6,PinName="Output_Get",PinToolTip="Retrieves the value of the variable, can use instead of a separate Get node",Direction="EGPD_Output",PinType.PinCategory="bool",PinType.PinSubCategoryObject=None,PinType.PinSubCategoryMemberReference=(),PinType.bIsArray=False,PinType.bIsReference=False,PinType.bIsConst=False,PinType.bIsWeakPointer=False,DefaultValue="false",AutogeneratedDefaultValue="false",PersistentGuid=00000000000000000000000000000000,bHidden=False,bNotConnectable=False,bDefaultValueIsReadOnly=False,bDefaultValueIsIgnored=False,bAdvancedView=False,)
+			   CustomProperties Pin (PinId=CC197367408EF6FCFBCED7B3192F4460,PinName="self",PinFriendlyName=NSLOCTEXT("K2Node", "Target", "Target"),PinType.PinCategory="object",PinType.PinSubCategoryObject=BlueprintGeneratedClass'/Game/ExampleContent/BlueprintFluidSurface/HeightfieldPainter.HeightfieldPainter_C',PinType.PinSubCategoryMemberReference=(),PinType.bIsArray=False,PinType.bIsReference=False,PinType.bIsConst=False,PinType.bIsWeakPointer=False,PersistentGuid=00000000000000000000000000000000,bHidden=True,bNotConnectable=False,bDefaultValueIsReadOnly=False,bDefaultValueIsIgnored=False,bAdvancedView=False,)
+			End Object
+			Begin Object Class=K2Node_VariableSet Name="K2Node_VariableSet_42"
+			   VariableReference=(MemberName="MouseDown",MemberGuid=209E1D534C15EFFA81760DA2C7F9D56C,bSelfContext=True)
+			   NodePosX=80
+			   NodePosY=1664
+			   NodeGuid=90955F284B3D2283A650FFB0DC894590
+			   CustomProperties Pin (PinId=9312E8C94B08C9FAF76638810A0F61A1,PinName="execute",PinType.PinCategory="exec",PinType.PinSubCategoryObject=None,PinType.PinSubCategoryMemberReference=(),PinType.bIsArray=False,PinType.bIsReference=False,PinType.bIsConst=False,PinType.bIsWeakPointer=False,LinkedTo=(K2Node_InputKey_19 B150F43D4C370A3041E5B9A2DD8FC834,),PersistentGuid=00000000000000000000000000000000,bHidden=False,bNotConnectable=False,bDefaultValueIsReadOnly=False,bDefaultValueIsIgnored=False,bAdvancedView=False,)
+			   CustomProperties Pin (PinId=CB1027A1495B8948A22EB58160AA4C3E,PinName="then",Direction="EGPD_Output",PinType.PinCategory="exec",PinType.PinSubCategoryObject=None,PinType.PinSubCategoryMemberReference=(),PinType.bIsArray=False,PinType.bIsReference=False,PinType.bIsConst=False,PinType.bIsWeakPointer=False,PersistentGuid=00000000000000000000000000000000,bHidden=False,bNotConnectable=False,bDefaultValueIsReadOnly=False,bDefaultValueIsIgnored=False,bAdvancedView=False,)
+			   CustomProperties Pin (PinId=3754DE984E6A106CC45C95831AEDF03A,PinName="MouseDown",PinType.PinCategory="bool",PinType.PinSubCategoryObject=None,PinType.PinSubCategoryMemberReference=(),PinType.bIsArray=False,PinType.bIsReference=False,PinType.bIsConst=False,PinType.bIsWeakPointer=False,DefaultValue="false",AutogeneratedDefaultValue="false",PersistentGuid=00000000000000000000000000000000,bHidden=False,bNotConnectable=False,bDefaultValueIsReadOnly=False,bDefaultValueIsIgnored=False,bAdvancedView=False,)
+			   CustomProperties Pin (PinId=CD9E960E4A4357C090CAFF8D5F8D9AB9,PinName="Output_Get",PinToolTip="Retrieves the value of the variable, can use instead of a separate Get node",Direction="EGPD_Output",PinType.PinCategory="bool",PinType.PinSubCategoryObject=None,PinType.PinSubCategoryMemberReference=(),PinType.bIsArray=False,PinType.bIsReference=False,PinType.bIsConst=False,PinType.bIsWeakPointer=False,DefaultValue="false",AutogeneratedDefaultValue="false",PersistentGuid=00000000000000000000000000000000,bHidden=False,bNotConnectable=False,bDefaultValueIsReadOnly=False,bDefaultValueIsIgnored=False,bAdvancedView=False,)
+			   CustomProperties Pin (PinId=CF16942B4E378883567B419CFD5CCC7B,PinName="self",PinFriendlyName=NSLOCTEXT("K2Node", "Target", "Target"),PinType.PinCategory="object",PinType.PinSubCategoryObject=BlueprintGeneratedClass'/Game/ExampleContent/BlueprintFluidSurface/HeightfieldPainter.HeightfieldPainter_C',PinType.PinSubCategoryMemberReference=(),PinType.bIsArray=False,PinType.bIsReference=False,PinType.bIsConst=False,PinType.bIsWeakPointer=False,PersistentGuid=00000000000000000000000000000000,bHidden=True,bNotConnectable=False,bDefaultValueIsReadOnly=False,bDefaultValueIsIgnored=False,bAdvancedView=False,)
+			End Object
+			Begin Object Class=K2Node_CallFunction Name="K2Node_CallFunction_360"
+			   FunctionReference=(MemberName="TraceFromCamera",MemberGuid=C96978014B119460D779C7A90DE17A56,bSelfContext=True)
+			   NodePosX=864
+			   NodePosY=1595
+			   NodeGuid=B47980294D235C085989D1A54C876ECE
+			   CustomProperties Pin (PinId=9F3F898442AC24B5F566A9A6E3D64A03,PinName="execute",PinToolTip="\nExec",PinType.PinCategory="exec",PinType.PinSubCategoryObject=None,PinType.PinSubCategoryMemberReference=(),PinType.bIsArray=False,PinType.bIsReference=False,PinType.bIsConst=False,PinType.bIsWeakPointer=False,LinkedTo=(K2Node_IfThenElse_25 358D51FF4B717F8D875E549790D394DB,),PersistentGuid=00000000000000000000000000000000,bHidden=False,bNotConnectable=False,bDefaultValueIsReadOnly=False,bDefaultValueIsIgnored=False,bAdvancedView=False,)
+			   CustomProperties Pin (PinId=BF6710D74C9EC10D6D5690B01743CC3D,PinName="then",PinToolTip="\nExec",Direction="EGPD_Output",PinType.PinCategory="exec",PinType.PinSubCategoryObject=None,PinType.PinSubCategoryMemberReference=(),PinType.bIsArray=False,PinType.bIsReference=False,PinType.bIsConst=False,PinType.bIsWeakPointer=False,PersistentGuid=00000000000000000000000000000000,bHidden=False,bNotConnectable=False,bDefaultValueIsReadOnly=False,bDefaultValueIsIgnored=False,bAdvancedView=False,)
+			   CustomProperties Pin (PinId=53418ABD4C9A655761FB14B9217B591D,PinName="self",PinFriendlyName=NSLOCTEXT("K2Node", "Target", "Target"),PinToolTip="Target\nSelf Reference",PinType.PinCategory="object",PinType.PinSubCategory="self",PinType.PinSubCategoryObject=None,PinType.PinSubCategoryMemberReference=(),PinType.bIsArray=False,PinType.bIsReference=False,PinType.bIsConst=False,PinType.bIsWeakPointer=False,PersistentGuid=00000000000000000000000000000000,bHidden=False,bNotConnectable=False,bDefaultValueIsReadOnly=False,bDefaultValueIsIgnored=False,bAdvancedView=False,)
+			End Object
+			Begin Object Class=K2Node_VariableGet Name="K2Node_VariableGet_179"
+			   VariableReference=(MemberName="MouseDown",MemberGuid=209E1D534C15EFFA81760DA2C7F9D56C,bSelfContext=True)
+			   NodePosX=400
+			   NodePosY=1723
+			   NodeGuid=1D4ECF694B52923CFD925F9190AF2D6E
+			   CustomProperties Pin (PinId=200A09224BC64E9B688537BFADAE5937,PinName="MouseDown",Direction="EGPD_Output",PinType.PinCategory="bool",PinType.PinSubCategoryObject=None,PinType.PinSubCategoryMemberReference=(),PinType.bIsArray=False,PinType.bIsReference=False,PinType.bIsConst=False,PinType.bIsWeakPointer=False,AutogeneratedDefaultValue="false",LinkedTo=(K2Node_IfThenElse_25 DBE9930D445160424F99CAB156B1E26F,),PersistentGuid=00000000000000000000000000000000,bHidden=False,bNotConnectable=False,bDefaultValueIsReadOnly=False,bDefaultValueIsIgnored=False,bAdvancedView=False,)
+			   CustomProperties Pin (PinId=C2F73D4D4B582D11AFB6A49B9D958DDE,PinName="self",PinFriendlyName=NSLOCTEXT("K2Node", "Target", "Target"),PinType.PinCategory="object",PinType.PinSubCategoryObject=BlueprintGeneratedClass'/Game/ExampleContent/BlueprintFluidSurface/HeightfieldPainter.HeightfieldPainter_C',PinType.PinSubCategoryMemberReference=(),PinType.bIsArray=False,PinType.bIsReference=False,PinType.bIsConst=False,PinType.bIsWeakPointer=False,PersistentGuid=00000000000000000000000000000000,bHidden=True,bNotConnectable=False,bDefaultValueIsReadOnly=False,bDefaultValueIsIgnored=False,bAdvancedView=False,)
+			End Object
+			Begin Object Class=K2Node_IfThenElse Name="K2Node_IfThenElse_25"
+			   NodePosX=624
+			   NodePosY=1611
+			   NodeGuid=3274D0A149940BA4CA7B23A98CF589B2
+			   CustomProperties Pin (PinId=997FDA76493A40CEE900CD81079C2BFC,PinName="execute",PinType.PinCategory="exec",PinType.PinSubCategoryObject=None,PinType.PinSubCategoryMemberReference=(),PinType.bIsArray=False,PinType.bIsReference=False,PinType.bIsConst=False,PinType.bIsWeakPointer=False,LinkedTo=(K2Node_Event_57 0371202549E0ACFD644FD4A3E8B79FD5,),PersistentGuid=00000000000000000000000000000000,bHidden=False,bNotConnectable=False,bDefaultValueIsReadOnly=False,bDefaultValueIsIgnored=False,bAdvancedView=False,)
+			   CustomProperties Pin (PinId=DBE9930D445160424F99CAB156B1E26F,PinName="Condition",PinType.PinCategory="bool",PinType.PinSubCategoryObject=None,PinType.PinSubCategoryMemberReference=(),PinType.bIsArray=False,PinType.bIsReference=False,PinType.bIsConst=False,PinType.bIsWeakPointer=False,DefaultValue="false",AutogeneratedDefaultValue="true",LinkedTo=(K2Node_VariableGet_179 200A09224BC64E9B688537BFADAE5937,),PersistentGuid=00000000000000000000000000000000,bHidden=False,bNotConnectable=False,bDefaultValueIsReadOnly=False,bDefaultValueIsIgnored=False,bAdvancedView=False,)
+			   CustomProperties Pin (PinId=358D51FF4B717F8D875E549790D394DB,PinName="then",PinFriendlyName=NSLOCTEXT("K2Node", "true", "true"),Direction="EGPD_Output",PinType.PinCategory="exec",PinType.PinSubCategoryObject=None,PinType.PinSubCategoryMemberReference=(),PinType.bIsArray=False,PinType.bIsReference=False,PinType.bIsConst=False,PinType.bIsWeakPointer=False,LinkedTo=(K2Node_CallFunction_360 9F3F898442AC24B5F566A9A6E3D64A03,),PersistentGuid=00000000000000000000000000000000,bHidden=False,bNotConnectable=False,bDefaultValueIsReadOnly=False,bDefaultValueIsIgnored=False,bAdvancedView=False,)
+			   CustomProperties Pin (PinId=53B9B8674F2C0FDC804A088C1129D0E3,PinName="else",PinFriendlyName=NSLOCTEXT("K2Node", "false", "false"),Direction="EGPD_Output",PinType.PinCategory="exec",PinType.PinSubCategoryObject=None,PinType.PinSubCategoryMemberReference=(),PinType.bIsArray=False,PinType.bIsReference=False,PinType.bIsConst=False,PinType.bIsWeakPointer=False,PersistentGuid=00000000000000000000000000000000,bHidden=False,bNotConnectable=False,bDefaultValueIsReadOnly=False,bDefaultValueIsIgnored=False,bAdvancedView=False,)
+			End Object
+			Begin Object Class=EdGraphNode_Comment Name="EdGraphNode_Comment_119"
+			   NodePosX=352
+			   NodePosY=1552
+			   NodeWidth=768
+			   NodeHeight=224
+			   NodeComment="Shoot when mouse is down"
+			   NodeGuid=009623A5485444F142AA62A6D595AAA8
+			End Object
+			Begin Object Class=K2Node_Event Name="K2Node_Event_58"
+			   EventReference=(MemberParent=Class'/Script/Engine.Actor',MemberName="ReceivePointDamage")
+			   bOverrideFunction=True
+			   NodePosX=1280
+			   NodePosY=1152
+			   NodeGuid=DCD3DAD24F490085E5BEBEBEB61EAE3C
+			   CustomProperties Pin (PinId=CEA11B954E55C6E301F0DDB73B737331,PinName="OutputDelegate",Direction="EGPD_Output",PinType.PinCategory="delegate",PinType.PinSubCategoryObject=None,PinType.PinSubCategoryMemberReference=(MemberParent=Class'/Script/Engine.Actor',MemberName="ReceivePointDamage"),PinType.bIsArray=False,PinType.bIsReference=False,PinType.bIsConst=False,PinType.bIsWeakPointer=False,PersistentGuid=00000000000000000000000000000000,bHidden=False,bNotConnectable=False,bDefaultValueIsReadOnly=False,bDefaultValueIsIgnored=False,bAdvancedView=False,)
+			   CustomProperties Pin (PinId=7BC1398142AE6DCA11411BBF6C1C7F4E,PinName="then",Direction="EGPD_Output",PinType.PinCategory="exec",PinType.PinSubCategoryObject=None,PinType.PinSubCategoryMemberReference=(),PinType.bIsArray=False,PinType.bIsReference=False,PinType.bIsConst=False,PinType.bIsWeakPointer=False,LinkedTo=(K2Node_CallFunction_368 527F9720481C49179DFF49B166FB5074,),PersistentGuid=00000000000000000000000000000000,bHidden=False,bNotConnectable=False,bDefaultValueIsReadOnly=False,bDefaultValueIsIgnored=False,bAdvancedView=False,)
+			   CustomProperties Pin (PinId=60CC33914AE4CFBF54C5D2807C150F9C,PinName="Damage",PinToolTip="Damage\nFloat",Direction="EGPD_Output",PinType.PinCategory="float",PinType.PinSubCategoryObject=None,PinType.PinSubCategoryMemberReference=(),PinType.bIsArray=False,PinType.bIsReference=False,PinType.bIsConst=False,PinType.bIsWeakPointer=False,DefaultValue="0.0",AutogeneratedDefaultValue="0.0",PersistentGuid=00000000000000000000000000000000,bHidden=False,bNotConnectable=False,bDefaultValueIsReadOnly=False,bDefaultValueIsIgnored=False,bAdvancedView=False,)
+			   CustomProperties Pin (PinId=16E9F4B64E717E33F3F8AB9D1FE60265,PinName="DamageType",PinToolTip="Damage Type\nDamage Type Reference",Direction="EGPD_Output",PinType.PinCategory="object",PinType.PinSubCategoryObject=Class'/Script/Engine.DamageType',PinType.PinSubCategoryMemberReference=(),PinType.bIsArray=False,PinType.bIsReference=False,PinType.bIsConst=True,PinType.bIsWeakPointer=False,PersistentGuid=00000000000000000000000000000000,bHidden=False,bNotConnectable=False,bDefaultValueIsReadOnly=False,bDefaultValueIsIgnored=False,bAdvancedView=False,)
+			   CustomProperties Pin (PinId=640542634B4797AB7140EAB2A04EF819,PinName="HitLocation",PinToolTip="Hit Location\nVector ",Direction="EGPD_Output",PinType.PinCategory="struct",PinType.PinSubCategoryObject=ScriptStruct'/Script/CoreUObject.Vector',PinType.PinSubCategoryMemberReference=(),PinType.bIsArray=False,PinType.bIsReference=False,PinType.bIsConst=False,PinType.bIsWeakPointer=False,AutogeneratedDefaultValue="0, 0, 0",LinkedTo=(K2Node_CallFunction_361 850583744D608EF2952864ABE9A12E9E,),PersistentGuid=00000000000000000000000000000000,bHidden=False,bNotConnectable=False,bDefaultValueIsReadOnly=False,bDefaultValueIsIgnored=False,bAdvancedView=False,)
+			   CustomProperties Pin (PinId=977D30BB43F35C9B14B82D96656AEE48,PinName="HitNormal",PinToolTip="Hit Normal\nVector ",Direction="EGPD_Output",PinType.PinCategory="struct",PinType.PinSubCategoryObject=ScriptStruct'/Script/CoreUObject.Vector',PinType.PinSubCategoryMemberReference=(),PinType.bIsArray=False,PinType.bIsReference=False,PinType.bIsConst=False,PinType.bIsWeakPointer=False,DefaultValue="0, 0, 0",AutogeneratedDefaultValue="0, 0, 0",PersistentGuid=00000000000000000000000000000000,bHidden=False,bNotConnectable=False,bDefaultValueIsReadOnly=False,bDefaultValueIsIgnored=False,bAdvancedView=False,)
+			   CustomProperties Pin (PinId=F3CA2D394CE3DE4376083189EE70EE01,PinName="HitComponent",PinToolTip="Hit Component\nPrimitive Component Reference",Direction="EGPD_Output",PinType.PinCategory="object",PinType.PinSubCategoryObject=Class'/Script/Engine.PrimitiveComponent',PinType.PinSubCategoryMemberReference=(),PinType.bIsArray=False,PinType.bIsReference=False,PinType.bIsConst=False,PinType.bIsWeakPointer=False,PersistentGuid=00000000000000000000000000000000,bHidden=False,bNotConnectable=False,bDefaultValueIsReadOnly=False,bDefaultValueIsIgnored=False,bAdvancedView=False,)
+			   CustomProperties Pin (PinId=E09A9FB14D431DADE2BB769CBACDB17A,PinName="BoneName",PinToolTip="Bone Name\nName",Direction="EGPD_Output",PinType.PinCategory="name",PinType.PinSubCategoryObject=None,PinType.PinSubCategoryMemberReference=(),PinType.bIsArray=False,PinType.bIsReference=False,PinType.bIsConst=False,PinType.bIsWeakPointer=False,DefaultValue="None",AutogeneratedDefaultValue="None",PersistentGuid=00000000000000000000000000000000,bHidden=False,bNotConnectable=False,bDefaultValueIsReadOnly=False,bDefaultValueIsIgnored=False,bAdvancedView=False,)
+			   CustomProperties Pin (PinId=80A5C0F84FD774A9E91907821B24714C,PinName="ShotFromDirection",PinToolTip="Shot from Direction\nVector ",Direction="EGPD_Output",PinType.PinCategory="struct",PinType.PinSubCategoryObject=ScriptStruct'/Script/CoreUObject.Vector',PinType.PinSubCategoryMemberReference=(),PinType.bIsArray=False,PinType.bIsReference=False,PinType.bIsConst=False,PinType.bIsWeakPointer=False,DefaultValue="0, 0, 0",AutogeneratedDefaultValue="0, 0, 0",PersistentGuid=00000000000000000000000000000000,bHidden=False,bNotConnectable=False,bDefaultValueIsReadOnly=False,bDefaultValueIsIgnored=False,bAdvancedView=False,)
+			   CustomProperties Pin (PinId=5D9B6A9F49DC7BCFD58603AB5847E975,PinName="InstigatedBy",PinToolTip="Instigated By\nController Reference",Direction="EGPD_Output",PinType.PinCategory="object",PinType.PinSubCategoryObject=Class'/Script/Engine.Controller',PinType.PinSubCategoryMemberReference=(),PinType.bIsArray=False,PinType.bIsReference=False,PinType.bIsConst=False,PinType.bIsWeakPointer=False,PersistentGuid=00000000000000000000000000000000,bHidden=False,bNotConnectable=False,bDefaultValueIsReadOnly=False,bDefaultValueIsIgnored=False,bAdvancedView=False,)
+			   CustomProperties Pin (PinId=CC306F3F459174FB247CD18B4B3919C5,PinName="DamageCauser",PinToolTip="Damage Causer\nActor Reference",Direction="EGPD_Output",PinType.PinCategory="object",PinType.PinSubCategoryObject=Class'/Script/Engine.Actor',PinType.PinSubCategoryMemberReference=(),PinType.bIsArray=False,PinType.bIsReference=False,PinType.bIsConst=False,PinType.bIsWeakPointer=False,PersistentGuid=00000000000000000000000000000000,bHidden=False,bNotConnectable=False,bDefaultValueIsReadOnly=False,bDefaultValueIsIgnored=False,bAdvancedView=False,)
+			End Object
+			Begin Object Class=EdGraphNode_Comment Name="EdGraphNode_Comment_120"
+			   NodePosX=1520
+			   NodePosY=1552
+			   NodeWidth=1312
+			   NodeHeight=272
+			   NodeComment="Convert the hit location into texture space for this mesh (should be replaced by line trace UV to work on any mesh)"
+			   NodeGuid=415A113C427D100C9141DBB5992DEE4B
+			End Object
+			Begin Object Class=K2Node_CallFunction Name="K2Node_CallFunction_361"
+			   bIsPureFunc=True
+			   FunctionReference=(MemberParent=Class'/Script/Engine.KismetMathLibrary',MemberName="TransformLocation")
+			   NodePosX=2256
+			   NodePosY=1642
+			   NodeGuid=04C287E249294FA75C5AAA89C7423EBC
+			   CustomProperties Pin (PinId=CEB3CA4647BBC3C737152680770F0BE1,PinName="self",PinFriendlyName=NSLOCTEXT("K2Node", "Target", "Target"),PinType.PinCategory="object",PinType.PinSubCategoryObject=Class'/Script/Engine.KismetMathLibrary',PinType.PinSubCategoryMemberReference=(),PinType.bIsArray=False,PinType.bIsReference=False,PinType.bIsConst=False,PinType.bIsWeakPointer=False,DefaultObject="/Script/Engine.Default__KismetMathLibrary",PersistentGuid=00000000000000000000000000000000,bHidden=True,bNotConnectable=False,bDefaultValueIsReadOnly=False,bDefaultValueIsIgnored=False,bAdvancedView=False,)
+			   CustomProperties Pin (PinId=8A4E9BDC43985D182203EEBE4D64ABED,PinName="T",PinType.PinCategory="struct",PinType.PinSubCategoryObject=ScriptStruct'/Script/CoreUObject.Transform',PinType.PinSubCategoryMemberReference=(),PinType.bIsArray=False,PinType.bIsReference=True,PinType.bIsConst=True,PinType.bIsWeakPointer=False,LinkedTo=(K2Node_CallFunction_363 BDA890074B460518CA47E0B5319A61D4,),PersistentGuid=00000000000000000000000000000000,bHidden=False,bNotConnectable=False,bDefaultValueIsReadOnly=False,bDefaultValueIsIgnored=True,bAdvancedView=False,)
+			   CustomProperties Pin (PinId=850583744D608EF2952864ABE9A12E9E,PinName="Location",PinType.PinCategory="struct",PinType.PinSubCategoryObject=ScriptStruct'/Script/CoreUObject.Vector',PinType.PinSubCategoryMemberReference=(),PinType.bIsArray=False,PinType.bIsReference=False,PinType.bIsConst=False,PinType.bIsWeakPointer=False,DefaultValue="0, 0, 0",AutogeneratedDefaultValue="0, 0, 0",LinkedTo=(K2Node_Event_58 640542634B4797AB7140EAB2A04EF819,),PersistentGuid=00000000000000000000000000000000,bHidden=False,bNotConnectable=False,bDefaultValueIsReadOnly=False,bDefaultValueIsIgnored=False,bAdvancedView=False,)
+			   CustomProperties Pin (PinId=7335087343EF48C6242BCDBFC09F2B35,PinName="ReturnValue",Direction="EGPD_Output",PinType.PinCategory="struct",PinType.PinSubCategoryObject=ScriptStruct'/Script/CoreUObject.Vector',PinType.PinSubCategoryMemberReference=(),PinType.bIsArray=False,PinType.bIsReference=False,PinType.bIsConst=False,PinType.bIsWeakPointer=False,AutogeneratedDefaultValue="0, 0, 0",LinkedTo=(K2Node_CallFunction_365 694D50D04A4FE8E7DC8CA689C5CDBE67,),PersistentGuid=00000000000000000000000000000000,bHidden=False,bNotConnectable=False,bDefaultValueIsReadOnly=False,bDefaultValueIsIgnored=False,bAdvancedView=False,)
+			End Object
+			Begin Object Class=K2Node_CallFunction Name="K2Node_CallFunction_362"
+			   bIsPureFunc=True
+			   bIsConstFunc=True
+			   FunctionReference=(MemberParent=Class'/Script/Engine.SceneComponent',MemberName="K2_GetComponentToWorld")
+			   NodePosX=1776
+			   NodePosY=1626
+			   NodeGuid=B37297A64EDD7A388CA750852B7BB044
+			   CustomProperties Pin (PinId=5A1010634C4DE3EAF3EF409F23EC1215,PinName="self",PinFriendlyName=NSLOCTEXT("K2Node", "Target", "Target"),PinType.PinCategory="object",PinType.PinSubCategoryObject=Class'/Script/Engine.SceneComponent',PinType.PinSubCategoryMemberReference=(),PinType.bIsArray=False,PinType.bIsReference=False,PinType.bIsConst=False,PinType.bIsWeakPointer=False,LinkedTo=(K2Node_VariableGet_180 78F118CD49FC76D027795D843DFE71FF,),PersistentGuid=00000000000000000000000000000000,bHidden=False,bNotConnectable=False,bDefaultValueIsReadOnly=False,bDefaultValueIsIgnored=False,bAdvancedView=False,)
+			   CustomProperties Pin (PinId=5150BE354EF59F455E6CA988E1351299,PinName="ReturnValue",Direction="EGPD_Output",PinType.PinCategory="struct",PinType.PinSubCategoryObject=ScriptStruct'/Script/CoreUObject.Transform',PinType.PinSubCategoryMemberReference=(),PinType.bIsArray=False,PinType.bIsReference=False,PinType.bIsConst=False,PinType.bIsWeakPointer=False,LinkedTo=(K2Node_CallFunction_363 557C12A64EC70024F7CC519F97D5D932,),PersistentGuid=00000000000000000000000000000000,bHidden=False,bNotConnectable=False,bDefaultValueIsReadOnly=False,bDefaultValueIsIgnored=False,bAdvancedView=False,)
+			End Object
+			Begin Object Class=K2Node_CallFunction Name="K2Node_CallFunction_363"
+			   bIsPureFunc=True
+			   FunctionReference=(MemberParent=Class'/Script/Engine.KismetMathLibrary',MemberName="InvertTransform")
+			   NodePosX=2016
+			   NodePosY=1642
+			   NodeGuid=8F6DC3A042C9F2D67A46239DBC74BF8A
+			   CustomProperties Pin (PinId=C4A955704E1399074BD903B945BE2245,PinName="self",PinFriendlyName=NSLOCTEXT("K2Node", "Target", "Target"),PinType.PinCategory="object",PinType.PinSubCategoryObject=Class'/Script/Engine.KismetMathLibrary',PinType.PinSubCategoryMemberReference=(),PinType.bIsArray=False,PinType.bIsReference=False,PinType.bIsConst=False,PinType.bIsWeakPointer=False,DefaultObject="/Script/Engine.Default__KismetMathLibrary",PersistentGuid=00000000000000000000000000000000,bHidden=True,bNotConnectable=False,bDefaultValueIsReadOnly=False,bDefaultValueIsIgnored=False,bAdvancedView=False,)
+			   CustomProperties Pin (PinId=557C12A64EC70024F7CC519F97D5D932,PinName="T",PinType.PinCategory="struct",PinType.PinSubCategoryObject=ScriptStruct'/Script/CoreUObject.Transform',PinType.PinSubCategoryMemberReference=(),PinType.bIsArray=False,PinType.bIsReference=True,PinType.bIsConst=True,PinType.bIsWeakPointer=False,LinkedTo=(K2Node_CallFunction_362 5150BE354EF59F455E6CA988E1351299,),PersistentGuid=00000000000000000000000000000000,bHidden=False,bNotConnectable=False,bDefaultValueIsReadOnly=False,bDefaultValueIsIgnored=True,bAdvancedView=False,)
+			   CustomProperties Pin (PinId=BDA890074B460518CA47E0B5319A61D4,PinName="ReturnValue",Direction="EGPD_Output",PinType.PinCategory="struct",PinType.PinSubCategoryObject=ScriptStruct'/Script/CoreUObject.Transform',PinType.PinSubCategoryMemberReference=(),PinType.bIsArray=False,PinType.bIsReference=False,PinType.bIsConst=False,PinType.bIsWeakPointer=False,LinkedTo=(K2Node_CallFunction_361 8A4E9BDC43985D182203EEBE4D64ABED,),PersistentGuid=00000000000000000000000000000000,bHidden=False,bNotConnectable=False,bDefaultValueIsReadOnly=False,bDefaultValueIsIgnored=False,bAdvancedView=False,)
+			End Object
+			Begin Object Class=K2Node_CallFunction Name="K2Node_CallFunction_364"
+			   bIsPureFunc=True
+			   FunctionReference=(MemberParent=Class'/Script/Engine.KismetMathLibrary',MemberName="Conv_VectorToLinearColor")
+			   NodePosX=2688
+			   NodePosY=1666
+			   NodeGuid=1B4DF094495714359B4ECB9EEE8204DB
+			   CustomProperties Pin (PinId=E21437414C7397632C246B96B157F7B6,PinName="self",PinFriendlyName=NSLOCTEXT("K2Node", "Target", "Target"),PinType.PinCategory="object",PinType.PinSubCategoryObject=Class'/Script/Engine.KismetMathLibrary',PinType.PinSubCategoryMemberReference=(),PinType.bIsArray=False,PinType.bIsReference=False,PinType.bIsConst=False,PinType.bIsWeakPointer=False,DefaultObject="/Script/Engine.Default__KismetMathLibrary",PersistentGuid=00000000000000000000000000000000,bHidden=True,bNotConnectable=False,bDefaultValueIsReadOnly=False,bDefaultValueIsIgnored=False,bAdvancedView=False,)
+			   CustomProperties Pin (PinId=1E73613645F8697CABA715886E923B95,PinName="InVec",PinType.PinCategory="struct",PinType.PinSubCategoryObject=ScriptStruct'/Script/CoreUObject.Vector',PinType.PinSubCategoryMemberReference=(),PinType.bIsArray=False,PinType.bIsReference=False,PinType.bIsConst=False,PinType.bIsWeakPointer=False,DefaultValue="0, 0, 0",AutogeneratedDefaultValue="0, 0, 0",LinkedTo=(K2Node_CallFunction_365 62F0E4DF49FC6CCA925D4E9A9BC03D67,),PersistentGuid=00000000000000000000000000000000,bHidden=False,bNotConnectable=False,bDefaultValueIsReadOnly=False,bDefaultValueIsIgnored=False,bAdvancedView=False,)
+			   CustomProperties Pin (PinId=86266FF5416699BCD33F46A84671E715,PinName="ReturnValue",Direction="EGPD_Output",PinType.PinCategory="struct",PinType.PinSubCategoryObject=ScriptStruct'/Script/CoreUObject.LinearColor',PinType.PinSubCategoryMemberReference=(),PinType.bIsArray=False,PinType.bIsReference=False,PinType.bIsConst=False,PinType.bIsWeakPointer=False,LinkedTo=(K2Node_CallFunction_367 1523A2CA42868EC87B90F39450925710,),PersistentGuid=00000000000000000000000000000000,bHidden=False,bNotConnectable=False,bDefaultValueIsReadOnly=False,bDefaultValueIsIgnored=False,bAdvancedView=False,)
+			End Object
+			Begin Object Class=K2Node_CallFunction Name="K2Node_CallFunction_365"
+			   bIsPureFunc=True
+			   FunctionReference=(MemberParent=Class'/Script/Engine.KismetMathLibrary',MemberName="Divide_VectorFloat")
+			   NodePosX=2512
+			   NodePosY=1666
+			   NodeComment="Size of the plane mesh in local space"
+			   bCommentBubbleVisible=True
+			   NodeGuid=20399D3E4EDE3014A1150998D3145D97
+			   CustomProperties Pin (PinId=17EBF57640D6253D4AE04282E40F95B0,PinName="self",PinFriendlyName=NSLOCTEXT("K2Node", "Target", "Target"),PinType.PinCategory="object",PinType.PinSubCategoryObject=Class'/Script/Engine.KismetMathLibrary',PinType.PinSubCategoryMemberReference=(),PinType.bIsArray=False,PinType.bIsReference=False,PinType.bIsConst=False,PinType.bIsWeakPointer=False,DefaultObject="/Script/Engine.Default__KismetMathLibrary",PersistentGuid=00000000000000000000000000000000,bHidden=True,bNotConnectable=False,bDefaultValueIsReadOnly=False,bDefaultValueIsIgnored=False,bAdvancedView=False,)
+			   CustomProperties Pin (PinId=694D50D04A4FE8E7DC8CA689C5CDBE67,PinName="A",PinType.PinCategory="struct",PinType.PinSubCategoryObject=ScriptStruct'/Script/CoreUObject.Vector',PinType.PinSubCategoryMemberReference=(),PinType.bIsArray=False,PinType.bIsReference=False,PinType.bIsConst=False,PinType.bIsWeakPointer=False,DefaultValue="0, 0, 0",AutogeneratedDefaultValue="0, 0, 0",LinkedTo=(K2Node_CallFunction_361 7335087343EF48C6242BCDBFC09F2B35,),PersistentGuid=00000000000000000000000000000000,bHidden=False,bNotConnectable=False,bDefaultValueIsReadOnly=False,bDefaultValueIsIgnored=False,bAdvancedView=False,)
+			   CustomProperties Pin (PinId=4FB4472943DA9A99ECFFC992F975CAAA,PinName="B",PinType.PinCategory="float",PinType.PinSubCategoryObject=None,PinType.PinSubCategoryMemberReference=(),PinType.bIsArray=False,PinType.bIsReference=False,PinType.bIsConst=False,PinType.bIsWeakPointer=False,DefaultValue="1000",AutogeneratedDefaultValue="1.000000",PersistentGuid=00000000000000000000000000000000,bHidden=False,bNotConnectable=False,bDefaultValueIsReadOnly=False,bDefaultValueIsIgnored=False,bAdvancedView=False,)
+			   CustomProperties Pin (PinId=62F0E4DF49FC6CCA925D4E9A9BC03D67,PinName="ReturnValue",Direction="EGPD_Output",PinType.PinCategory="struct",PinType.PinSubCategoryObject=ScriptStruct'/Script/CoreUObject.Vector',PinType.PinSubCategoryMemberReference=(),PinType.bIsArray=False,PinType.bIsReference=False,PinType.bIsConst=False,PinType.bIsWeakPointer=False,AutogeneratedDefaultValue="0, 0, 0",LinkedTo=(K2Node_CallFunction_364 1E73613645F8697CABA715886E923B95,),PersistentGuid=00000000000000000000000000000000,bHidden=False,bNotConnectable=False,bDefaultValueIsReadOnly=False,bDefaultValueIsIgnored=False,bAdvancedView=False,)
+			End Object
+			Begin Object Class=K2Node_CallFunction Name="K2Node_CallFunction_366"
+			   FunctionReference=(MemberParent=Class'/Script/Engine.KismetRenderingLibrary',MemberName="DrawMaterialToRenderTarget")
+			   NodePosX=2992
+			   NodePosY=1153
+			   NodeGuid=2F07EF694E802F738B3AB58A0D8E8623
+			   CustomProperties Pin (PinId=D89DC81641CA4093CA72DAAF5CACC388,PinName="execute",PinToolTip="\nExec",PinType.PinCategory="exec",PinType.PinSubCategoryObject=None,PinType.PinSubCategoryMemberReference=(),PinType.bIsArray=False,PinType.bIsReference=False,PinType.bIsConst=False,PinType.bIsWeakPointer=False,LinkedTo=(K2Node_CallFunction_367 F5065258491F11F07B8E7B9BF2A38140,),PersistentGuid=00000000000000000000000000000000,bHidden=False,bNotConnectable=False,bDefaultValueIsReadOnly=False,bDefaultValueIsIgnored=False,bAdvancedView=False,)
+			   CustomProperties Pin (PinId=18C8362646DD324E0B3EE09FE2E9691F,PinName="then",PinToolTip="\nExec",Direction="EGPD_Output",PinType.PinCategory="exec",PinType.PinSubCategoryObject=None,PinType.PinSubCategoryMemberReference=(),PinType.bIsArray=False,PinType.bIsReference=False,PinType.bIsConst=False,PinType.bIsWeakPointer=False,PersistentGuid=00000000000000000000000000000000,bHidden=False,bNotConnectable=False,bDefaultValueIsReadOnly=False,bDefaultValueIsIgnored=False,bAdvancedView=False,)
+			   CustomProperties Pin (PinId=3E4C4C0C47F9C599C77C0A94DAE5CDBC,PinName="self",PinFriendlyName=NSLOCTEXT("K2Node", "Target", "Target"),PinToolTip="Target\nKismet Rendering Library Reference",PinType.PinCategory="object",PinType.PinSubCategoryObject=Class'/Script/Engine.KismetRenderingLibrary',PinType.PinSubCategoryMemberReference=(),PinType.bIsArray=False,PinType.bIsReference=False,PinType.bIsConst=False,PinType.bIsWeakPointer=False,DefaultObject="/Script/Engine.Default__KismetRenderingLibrary",PersistentGuid=00000000000000000000000000000000,bHidden=True,bNotConnectable=False,bDefaultValueIsReadOnly=False,bDefaultValueIsIgnored=False,bAdvancedView=False,)
+			   CustomProperties Pin (PinId=5EE9579647D9A511B63300AEE24E89BF,PinName="WorldContextObject",PinToolTip="World Context Object\nObject Reference",PinType.PinCategory="object",PinType.PinSubCategoryObject=Class'/Script/CoreUObject.Object',PinType.PinSubCategoryMemberReference=(),PinType.bIsArray=False,PinType.bIsReference=False,PinType.bIsConst=False,PinType.bIsWeakPointer=False,PersistentGuid=00000000000000000000000000000000,bHidden=True,bNotConnectable=False,bDefaultValueIsReadOnly=False,bDefaultValueIsIgnored=False,bAdvancedView=False,)
+			   CustomProperties Pin (PinId=B57C770C4D827E19FDE1F885B23E018C,PinName="TextureRenderTarget",PinToolTip="Texture Render Target\nTexture Render Target 2D Reference",PinType.PinCategory="object",PinType.PinSubCategoryObject=Class'/Script/Engine.TextureRenderTarget2D',PinType.PinSubCategoryMemberReference=(),PinType.bIsArray=False,PinType.bIsReference=False,PinType.bIsConst=False,PinType.bIsWeakPointer=False,LinkedTo=(K2Node_VariableGet_182 A993DA7E42B925003560FEAF90CAD2D6,),PersistentGuid=00000000000000000000000000000000,bHidden=False,bNotConnectable=False,bDefaultValueIsReadOnly=False,bDefaultValueIsIgnored=False,bAdvancedView=False,)
+			   CustomProperties Pin (PinId=D99B80D64C0171739FBC10A4E4C540EC,PinName="Material",PinToolTip="Material\nMaterial Interface Reference",PinType.PinCategory="object",PinType.PinSubCategoryObject=Class'/Script/Engine.MaterialInterface',PinType.PinSubCategoryMemberReference=(),PinType.bIsArray=False,PinType.bIsReference=False,PinType.bIsConst=False,PinType.bIsWeakPointer=False,LinkedTo=(K2Node_VariableGet_181 16610BCA473F60D2FB851F8B3E4736A4,),PersistentGuid=00000000000000000000000000000000,bHidden=False,bNotConnectable=False,bDefaultValueIsReadOnly=False,bDefaultValueIsIgnored=False,bAdvancedView=False,)
+			End Object
+			Begin Object Class=EdGraphNode_Comment Name="EdGraphNode_Comment_121"
+			   CommentDepth=-2
+			   NodePosX=1232
+			   NodePosY=1056
+			   NodeWidth=2080
+			   NodeHeight=800
+			   NodeComment="Take Damage - paint additive splat into heightfield"
+			   NodeGuid=3BD36F9D40901D7760B550A693D7DFD5
+			End Object
+			Begin Object Class=K2Node_CallFunction Name="K2Node_CallFunction_367"
+			   FunctionReference=(MemberParent=Class'/Script/Engine.MaterialInstanceDynamic',MemberName="SetVectorParameterValue")
+			   NodePosX=2608
+			   NodePosY=1137
+			   NodeGuid=5110AB1044D444E02184F4BB1AFDB4D6
+			   CustomProperties Pin (PinId=88C40CCB4DEBDA28F201A388C6DF16B9,PinName="execute",PinType.PinCategory="exec",PinType.PinSubCategoryObject=None,PinType.PinSubCategoryMemberReference=(),PinType.bIsArray=False,PinType.bIsReference=False,PinType.bIsConst=False,PinType.bIsWeakPointer=False,LinkedTo=(K2Node_CallFunction_369 5826E451454A5FA8D03AB1B3928B3090,),PersistentGuid=00000000000000000000000000000000,bHidden=False,bNotConnectable=False,bDefaultValueIsReadOnly=False,bDefaultValueIsIgnored=False,bAdvancedView=False,)
+			   CustomProperties Pin (PinId=F5065258491F11F07B8E7B9BF2A38140,PinName="then",Direction="EGPD_Output",PinType.PinCategory="exec",PinType.PinSubCategoryObject=None,PinType.PinSubCategoryMemberReference=(),PinType.bIsArray=False,PinType.bIsReference=False,PinType.bIsConst=False,PinType.bIsWeakPointer=False,LinkedTo=(K2Node_CallFunction_366 D89DC81641CA4093CA72DAAF5CACC388,),PersistentGuid=00000000000000000000000000000000,bHidden=False,bNotConnectable=False,bDefaultValueIsReadOnly=False,bDefaultValueIsIgnored=False,bAdvancedView=False,)
+			   CustomProperties Pin (PinId=FF3048F146D54F1C4AA4A48C221B5457,PinName="self",PinFriendlyName=NSLOCTEXT("K2Node", "Target", "Target"),PinType.PinCategory="object",PinType.PinSubCategoryObject=Class'/Script/Engine.MaterialInstanceDynamic',PinType.PinSubCategoryMemberReference=(),PinType.bIsArray=False,PinType.bIsReference=False,PinType.bIsConst=False,PinType.bIsWeakPointer=False,LinkedTo=(K2Node_VariableGet_181 16610BCA473F60D2FB851F8B3E4736A4,),PersistentGuid=00000000000000000000000000000000,bHidden=False,bNotConnectable=False,bDefaultValueIsReadOnly=False,bDefaultValueIsIgnored=False,bAdvancedView=False,)
+			   CustomProperties Pin (PinId=50E108F844D08F285B9864B792806572,PinName="ParameterName",PinType.PinCategory="name",PinType.PinSubCategoryObject=None,PinType.PinSubCategoryMemberReference=(),PinType.bIsArray=False,PinType.bIsReference=False,PinType.bIsConst=False,PinType.bIsWeakPointer=False,DefaultValue="ForcePosition",AutogeneratedDefaultValue="None",PersistentGuid=00000000000000000000000000000000,bHidden=False,bNotConnectable=False,bDefaultValueIsReadOnly=False,bDefaultValueIsIgnored=False,bAdvancedView=False,)
+			   CustomProperties Pin (PinId=1523A2CA42868EC87B90F39450925710,PinName="Value",PinType.PinCategory="struct",PinType.PinSubCategoryObject=ScriptStruct'/Script/CoreUObject.LinearColor',PinType.PinSubCategoryMemberReference=(),PinType.bIsArray=False,PinType.bIsReference=False,PinType.bIsConst=False,PinType.bIsWeakPointer=False,LinkedTo=(K2Node_CallFunction_364 86266FF5416699BCD33F46A84671E715,),PersistentGuid=00000000000000000000000000000000,bHidden=False,bNotConnectable=False,bDefaultValueIsReadOnly=False,bDefaultValueIsIgnored=False,bAdvancedView=False,)
+			End Object
+			Begin Object Class=K2Node_VariableGet Name="K2Node_VariableGet_180"
+			   VariableReference=(MemberName="StaticMeshComponent",bSelfContext=True)
+			   NodePosX=1568
+			   NodePosY=1664
+			   NodeGuid=2AA0148A4039C35E011CCD8930FFC55E
+			   CustomProperties Pin (PinId=78F118CD49FC76D027795D843DFE71FF,PinName="StaticMeshComponent",Direction="EGPD_Output",PinType.PinCategory="object",PinType.PinSubCategoryObject=Class'/Script/Engine.StaticMeshComponent',PinType.PinSubCategoryMemberReference=(),PinType.bIsArray=False,PinType.bIsReference=False,PinType.bIsConst=False,PinType.bIsWeakPointer=False,LinkedTo=(K2Node_CallFunction_362 5A1010634C4DE3EAF3EF409F23EC1215,),PersistentGuid=00000000000000000000000000000000,bHidden=False,bNotConnectable=False,bDefaultValueIsReadOnly=False,bDefaultValueIsIgnored=False,bAdvancedView=False,)
+			   CustomProperties Pin (PinId=629DFDFB49CD19075CAD4B8F8CE8BD97,PinName="self",PinFriendlyName=NSLOCTEXT("K2Node", "Target", "Target"),PinType.PinCategory="object",PinType.PinSubCategoryObject=Class'/Script/Engine.StaticMeshActor',PinType.PinSubCategoryMemberReference=(),PinType.bIsArray=False,PinType.bIsReference=False,PinType.bIsConst=False,PinType.bIsWeakPointer=False,PersistentGuid=00000000000000000000000000000000,bHidden=True,bNotConnectable=False,bDefaultValueIsReadOnly=False,bDefaultValueIsIgnored=False,bAdvancedView=False,)
+			End Object
+			Begin Object Class=K2Node_VariableGet Name="K2Node_VariableGet_181"
+			   VariableReference=(MemberName="ForceSplatMID",MemberGuid=20FE6B2D4954BEA0CD113E92370334C9,bSelfContext=True)
+			   NodePosX=1536
+			   NodePosY=1488
+			   NodeGuid=845AA5BB4717DD847C6338B46486C900
+			   CustomProperties Pin (PinId=16610BCA473F60D2FB851F8B3E4736A4,PinName="ForceSplatMID",Direction="EGPD_Output",PinType.PinCategory="object",PinType.PinSubCategoryObject=Class'/Script/Engine.MaterialInstanceDynamic',PinType.PinSubCategoryMemberReference=(),PinType.bIsArray=False,PinType.bIsReference=False,PinType.bIsConst=False,PinType.bIsWeakPointer=False,LinkedTo=(K2Node_CallFunction_367 FF3048F146D54F1C4AA4A48C221B5457,K2Node_CallFunction_366 D99B80D64C0171739FBC10A4E4C540EC,K2Node_CallFunction_368 55F338D54A4933AB828D2DB26782C387,K2Node_CallFunction_369 55F338D54A4933AB828D2DB26782C387,),PersistentGuid=00000000000000000000000000000000,bHidden=False,bNotConnectable=False,bDefaultValueIsReadOnly=False,bDefaultValueIsIgnored=False,bAdvancedView=False,)
+			   CustomProperties Pin (PinId=2D0979DA456C78A5E1AF5BA55CB4E59D,PinName="self",PinFriendlyName=NSLOCTEXT("K2Node", "Target", "Target"),PinType.PinCategory="object",PinType.PinSubCategoryObject=BlueprintGeneratedClass'/Game/ExampleContent/BlueprintFluidSurface/HeightfieldPainter.HeightfieldPainter_C',PinType.PinSubCategoryMemberReference=(),PinType.bIsArray=False,PinType.bIsReference=False,PinType.bIsConst=False,PinType.bIsWeakPointer=False,PersistentGuid=00000000000000000000000000000000,bHidden=True,bNotConnectable=False,bDefaultValueIsReadOnly=False,bDefaultValueIsIgnored=False,bAdvancedView=False,)
+			End Object
+			Begin Object Class=K2Node_VariableGet Name="K2Node_VariableGet_182"
+			   VariableReference=(MemberName="HeightfieldRT",MemberGuid=E71A12524941DA802B08039BF7491FF7,bSelfContext=True)
+			   NodePosX=2784
+			   NodePosY=1408
+			   NodeGuid=85104BAC40625F22EB93A897106C0266
+			   CustomProperties Pin (PinId=A993DA7E42B925003560FEAF90CAD2D6,PinName="HeightfieldRT",Direction="EGPD_Output",PinType.PinCategory="object",PinType.PinSubCategoryObject=Class'/Script/Engine.TextureRenderTarget2D',PinType.PinSubCategoryMemberReference=(),PinType.bIsArray=False,PinType.bIsReference=False,PinType.bIsConst=False,PinType.bIsWeakPointer=False,LinkedTo=(K2Node_CallFunction_366 B57C770C4D827E19FDE1F885B23E018C,),PersistentGuid=00000000000000000000000000000000,bHidden=False,bNotConnectable=False,bDefaultValueIsReadOnly=False,bDefaultValueIsIgnored=False,bAdvancedView=False,)
+			   CustomProperties Pin (PinId=85E5E99140F7971A74486596FC7731E1,PinName="self",PinFriendlyName=NSLOCTEXT("K2Node", "Target", "Target"),PinType.PinCategory="object",PinType.PinSubCategoryObject=BlueprintGeneratedClass'/Game/ExampleContent/BlueprintFluidSurface/HeightfieldPainter.HeightfieldPainter_C',PinType.PinSubCategoryMemberReference=(),PinType.bIsArray=False,PinType.bIsReference=False,PinType.bIsConst=False,PinType.bIsWeakPointer=False,PersistentGuid=00000000000000000000000000000000,bHidden=True,bNotConnectable=False,bDefaultValueIsReadOnly=False,bDefaultValueIsIgnored=False,bAdvancedView=False,)
+			End Object
+			Begin Object Class=K2Node_CallFunction Name="K2Node_CallFunction_368"
+			   FunctionReference=(MemberParent=Class'/Script/Engine.MaterialInstanceDynamic',MemberName="SetScalarParameterValue")
+			   NodePosX=1728
+			   NodePosY=1137
+			   NodeGuid=02D6BB494A645B39BEB43CA9F510AC32
+			   CustomProperties Pin (PinId=527F9720481C49179DFF49B166FB5074,PinName="execute",PinToolTip="\nExec",PinType.PinCategory="exec",PinType.PinSubCategoryObject=None,PinType.PinSubCategoryMemberReference=(),PinType.bIsArray=False,PinType.bIsReference=False,PinType.bIsConst=False,PinType.bIsWeakPointer=False,LinkedTo=(K2Node_Event_58 7BC1398142AE6DCA11411BBF6C1C7F4E,),PersistentGuid=00000000000000000000000000000000,bHidden=False,bNotConnectable=False,bDefaultValueIsReadOnly=False,bDefaultValueIsIgnored=False,bAdvancedView=False,)
+			   CustomProperties Pin (PinId=5826E451454A5FA8D03AB1B3928B3090,PinName="then",PinToolTip="\nExec",Direction="EGPD_Output",PinType.PinCategory="exec",PinType.PinSubCategoryObject=None,PinType.PinSubCategoryMemberReference=(),PinType.bIsArray=False,PinType.bIsReference=False,PinType.bIsConst=False,PinType.bIsWeakPointer=False,LinkedTo=(K2Node_CallFunction_369 527F9720481C49179DFF49B166FB5074,),PersistentGuid=00000000000000000000000000000000,bHidden=False,bNotConnectable=False,bDefaultValueIsReadOnly=False,bDefaultValueIsIgnored=False,bAdvancedView=False,)
+			   CustomProperties Pin (PinId=55F338D54A4933AB828D2DB26782C387,PinName="self",PinFriendlyName=NSLOCTEXT("K2Node", "Target", "Target"),PinToolTip="Target\nMaterial Instance Dynamic Reference",PinType.PinCategory="object",PinType.PinSubCategoryObject=Class'/Script/Engine.MaterialInstanceDynamic',PinType.PinSubCategoryMemberReference=(),PinType.bIsArray=False,PinType.bIsReference=False,PinType.bIsConst=False,PinType.bIsWeakPointer=False,LinkedTo=(K2Node_VariableGet_181 16610BCA473F60D2FB851F8B3E4736A4,),PersistentGuid=00000000000000000000000000000000,bHidden=False,bNotConnectable=False,bDefaultValueIsReadOnly=False,bDefaultValueIsIgnored=False,bAdvancedView=False,)
+			   CustomProperties Pin (PinId=7E1E1625448C9EC8718DA083EEE57602,PinName="ParameterName",PinToolTip="Parameter Name\nName",PinType.PinCategory="name",PinType.PinSubCategoryObject=None,PinType.PinSubCategoryMemberReference=(),PinType.bIsArray=False,PinType.bIsReference=False,PinType.bIsConst=False,PinType.bIsWeakPointer=False,DefaultValue="ForceSize",AutogeneratedDefaultValue="None",PersistentGuid=00000000000000000000000000000000,bHidden=False,bNotConnectable=False,bDefaultValueIsReadOnly=False,bDefaultValueIsIgnored=False,bAdvancedView=False,)
+			   CustomProperties Pin (PinId=2997AC484441B6721690EC97301505A9,PinName="Value",PinToolTip="Value\nFloat",PinType.PinCategory="float",PinType.PinSubCategoryObject=None,PinType.PinSubCategoryMemberReference=(),PinType.bIsArray=False,PinType.bIsReference=False,PinType.bIsConst=False,PinType.bIsWeakPointer=False,DefaultValue="0.0",AutogeneratedDefaultValue="0.0",LinkedTo=(K2Node_VariableGet_183 4CA0FD884AE58C829F746B926BBF77AE,),PersistentGuid=00000000000000000000000000000000,bHidden=False,bNotConnectable=False,bDefaultValueIsReadOnly=False,bDefaultValueIsIgnored=False,bAdvancedView=False,)
+			End Object
+			Begin Object Class=K2Node_VariableGet Name="K2Node_VariableGet_183"
+			   VariableReference=(MemberName="BrushSize",MemberGuid=59B8C8C74D55E70D194D7CB1103368E6,bSelfContext=True)
+			   NodePosX=1536
+			   NodePosY=1344
+			   NodeGuid=FA6C7D664E3B0FAD5E09F087082A941B
+			   CustomProperties Pin (PinId=4CA0FD884AE58C829F746B926BBF77AE,PinName="BrushSize",Direction="EGPD_Output",PinType.PinCategory="float",PinType.PinSubCategoryObject=None,PinType.PinSubCategoryMemberReference=(),PinType.bIsArray=False,PinType.bIsReference=False,PinType.bIsConst=False,PinType.bIsWeakPointer=False,AutogeneratedDefaultValue="0.0",LinkedTo=(K2Node_CallFunction_368 2997AC484441B6721690EC97301505A9,),PersistentGuid=00000000000000000000000000000000,bHidden=False,bNotConnectable=False,bDefaultValueIsReadOnly=False,bDefaultValueIsIgnored=False,bAdvancedView=False,)
+			   CustomProperties Pin (PinId=0224EB0545DA2DE7C7B5479CBC69E65B,PinName="self",PinFriendlyName=NSLOCTEXT("K2Node", "Target", "Target"),PinType.PinCategory="object",PinType.PinSubCategoryObject=BlueprintGeneratedClass'/Game/ExampleContent/BlueprintFluidSurface/HeightfieldPainter.HeightfieldPainter_C',PinType.PinSubCategoryMemberReference=(),PinType.bIsArray=False,PinType.bIsReference=False,PinType.bIsConst=False,PinType.bIsWeakPointer=False,PersistentGuid=00000000000000000000000000000000,bHidden=True,bNotConnectable=False,bDefaultValueIsReadOnly=False,bDefaultValueIsIgnored=False,bAdvancedView=False,)
+			End Object
+			Begin Object Class=K2Node_CallFunction Name="K2Node_CallFunction_369"
+			   FunctionReference=(MemberParent=Class'/Script/Engine.MaterialInstanceDynamic',MemberName="SetScalarParameterValue")
+			   NodePosX=2096
+			   NodePosY=1137
+			   NodeGuid=887CAD4340528EBDF094F395254C2E87
+			   CustomProperties Pin (PinId=527F9720481C49179DFF49B166FB5074,PinName="execute",PinToolTip="\nExec",PinType.PinCategory="exec",PinType.PinSubCategoryObject=None,PinType.PinSubCategoryMemberReference=(),PinType.bIsArray=False,PinType.bIsReference=False,PinType.bIsConst=False,PinType.bIsWeakPointer=False,LinkedTo=(K2Node_CallFunction_368 5826E451454A5FA8D03AB1B3928B3090,),PersistentGuid=00000000000000000000000000000000,bHidden=False,bNotConnectable=False,bDefaultValueIsReadOnly=False,bDefaultValueIsIgnored=False,bAdvancedView=False,)
+			   CustomProperties Pin (PinId=5826E451454A5FA8D03AB1B3928B3090,PinName="then",PinToolTip="\nExec",Direction="EGPD_Output",PinType.PinCategory="exec",PinType.PinSubCategoryObject=None,PinType.PinSubCategoryMemberReference=(),PinType.bIsArray=False,PinType.bIsReference=False,PinType.bIsConst=False,PinType.bIsWeakPointer=False,LinkedTo=(K2Node_CallFunction_367 88C40CCB4DEBDA28F201A388C6DF16B9,),PersistentGuid=00000000000000000000000000000000,bHidden=False,bNotConnectable=False,bDefaultValueIsReadOnly=False,bDefaultValueIsIgnored=False,bAdvancedView=False,)
+			   CustomProperties Pin (PinId=55F338D54A4933AB828D2DB26782C387,PinName="self",PinFriendlyName=NSLOCTEXT("K2Node", "Target", "Target"),PinToolTip="Target\nMaterial Instance Dynamic Reference",PinType.PinCategory="object",PinType.PinSubCategoryObject=Class'/Script/Engine.MaterialInstanceDynamic',PinType.PinSubCategoryMemberReference=(),PinType.bIsArray=False,PinType.bIsReference=False,PinType.bIsConst=False,PinType.bIsWeakPointer=False,LinkedTo=(K2Node_VariableGet_181 16610BCA473F60D2FB851F8B3E4736A4,),PersistentGuid=00000000000000000000000000000000,bHidden=False,bNotConnectable=False,bDefaultValueIsReadOnly=False,bDefaultValueIsIgnored=False,bAdvancedView=False,)
+			   CustomProperties Pin (PinId=7E1E1625448C9EC8718DA083EEE57602,PinName="ParameterName",PinToolTip="Parameter Name\nName",PinType.PinCategory="name",PinType.PinSubCategoryObject=None,PinType.PinSubCategoryMemberReference=(),PinType.bIsArray=False,PinType.bIsReference=False,PinType.bIsConst=False,PinType.bIsWeakPointer=False,DefaultValue="ForceStrength",AutogeneratedDefaultValue="None",PersistentGuid=00000000000000000000000000000000,bHidden=False,bNotConnectable=False,bDefaultValueIsReadOnly=False,bDefaultValueIsIgnored=False,bAdvancedView=False,)
+			   CustomProperties Pin (PinId=2997AC484441B6721690EC97301505A9,PinName="Value",PinToolTip="Value\nFloat",PinType.PinCategory="float",PinType.PinSubCategoryObject=None,PinType.PinSubCategoryMemberReference=(),PinType.bIsArray=False,PinType.bIsReference=False,PinType.bIsConst=False,PinType.bIsWeakPointer=False,DefaultValue="0.0",AutogeneratedDefaultValue="0.0",LinkedTo=(K2Node_VariableGet_184 C94B43A14D6EFD6EC40E638FADF9DEDA,),PersistentGuid=00000000000000000000000000000000,bHidden=False,bNotConnectable=False,bDefaultValueIsReadOnly=False,bDefaultValueIsIgnored=False,bAdvancedView=False,)
+			End Object
+			Begin Object Class=K2Node_VariableGet Name="K2Node_VariableGet_184"
+			   VariableReference=(MemberName="BrushStrength",MemberGuid=ABE2467C458A059B3991E19E6762E572,bSelfContext=True)
+			   NodePosX=1936
+			   NodePosY=1360
+			   NodeGuid=045727FC42FA8980011242B314C7AC24
+			   CustomProperties Pin (PinId=C94B43A14D6EFD6EC40E638FADF9DEDA,PinName="BrushStrength",Direction="EGPD_Output",PinType.PinCategory="float",PinType.PinSubCategoryObject=None,PinType.PinSubCategoryMemberReference=(),PinType.bIsArray=False,PinType.bIsReference=False,PinType.bIsConst=False,PinType.bIsWeakPointer=False,AutogeneratedDefaultValue="0.0",LinkedTo=(K2Node_CallFunction_369 2997AC484441B6721690EC97301505A9,),PersistentGuid=00000000000000000000000000000000,bHidden=False,bNotConnectable=False,bDefaultValueIsReadOnly=False,bDefaultValueIsIgnored=False,bAdvancedView=False,)
+			   CustomProperties Pin (PinId=F8B1890E4DC2C88C7F0E44BDD28680BF,PinName="self",PinFriendlyName=NSLOCTEXT("K2Node", "Target", "Target"),PinType.PinCategory="object",PinType.PinSubCategoryObject=BlueprintGeneratedClass'/Game/ExampleContent/BlueprintFluidSurface/HeightfieldPainter.HeightfieldPainter_C',PinType.PinSubCategoryMemberReference=(),PinType.bIsArray=False,PinType.bIsReference=False,PinType.bIsConst=False,PinType.bIsWeakPointer=False,PersistentGuid=00000000000000000000000000000000,bHidden=True,bNotConnectable=False,bDefaultValueIsReadOnly=False,bDefaultValueIsIgnored=False,bAdvancedView=False,)
+			End Object
+			Begin Object Class=K2Node_CallFunction Name="K2Node_CallFunction_370"
+			   FunctionReference=(MemberParent=Class'/Script/Engine.MaterialInstanceDynamic',MemberName="SetTextureParameterValue")
+			   NodePosX=880
+			   NodePosY=1104
+			   ErrorType=1
+			   ErrorMsg="Error This blueprint (self) is not a MaterialInstanceDynamic, therefore \' self \' must have a connection."
+			   NodeGuid=26E935614A4A9B56907DDDA712D30F8D
+			   CustomProperties Pin (PinId=542CC73148D38A37BCCECE90262AFC49,PinName="execute",PinType.PinCategory="exec",PinType.PinSubCategoryObject=None,PinType.PinSubCategoryMemberReference=(),PinType.bIsArray=False,PinType.bIsReference=False,PinType.bIsConst=False,PinType.bIsWeakPointer=False,LinkedTo=(K2Node_CallFunction_359 43C94C3C4FB04A4DEADC37A9E3BDE0B2,),PersistentGuid=00000000000000000000000000000000,bHidden=False,bNotConnectable=False,bDefaultValueIsReadOnly=False,bDefaultValueIsIgnored=False,bAdvancedView=False,)
+			   CustomProperties Pin (PinId=A183C34442666D0C8D878F8716DDAC38,PinName="then",Direction="EGPD_Output",PinType.PinCategory="exec",PinType.PinSubCategoryObject=None,PinType.PinSubCategoryMemberReference=(),PinType.bIsArray=False,PinType.bIsReference=False,PinType.bIsConst=False,PinType.bIsWeakPointer=False,PersistentGuid=00000000000000000000000000000000,bHidden=False,bNotConnectable=False,bDefaultValueIsReadOnly=False,bDefaultValueIsIgnored=False,bAdvancedView=False,)
+			   CustomProperties Pin (PinId=63A3B28B4B0AAB8799485AA1BF3EAE97,PinName="self",PinFriendlyName=NSLOCTEXT("K2Node", "Target", "Target"),PinType.PinCategory="object",PinType.PinSubCategoryObject=Class'/Script/Engine.MaterialInstanceDynamic',PinType.PinSubCategoryMemberReference=(),PinType.bIsArray=False,PinType.bIsReference=False,PinType.bIsConst=False,PinType.bIsWeakPointer=False,LinkedTo=(K2Node_VariableGet_186 D6A13F044975C578FFF6A98328676BC9,),PersistentGuid=00000000000000000000000000000000,bHidden=False,bNotConnectable=False,bDefaultValueIsReadOnly=False,bDefaultValueIsIgnored=False,bAdvancedView=False,)
+			   CustomProperties Pin (PinId=92D5DD9C482CFAAA46C3FE93F76AC002,PinName="ParameterName",PinType.PinCategory="name",PinType.PinSubCategoryObject=None,PinType.PinSubCategoryMemberReference=(),PinType.bIsArray=False,PinType.bIsReference=False,PinType.bIsConst=False,PinType.bIsWeakPointer=False,DefaultValue="Heightfield",AutogeneratedDefaultValue="None",PersistentGuid=00000000000000000000000000000000,bHidden=False,bNotConnectable=False,bDefaultValueIsReadOnly=False,bDefaultValueIsIgnored=False,bAdvancedView=False,)
+			   CustomProperties Pin (PinId=FAC8C8B64E1E50423C9F4493F3C27926,PinName="Value",PinType.PinCategory="object",PinType.PinSubCategoryObject=Class'/Script/Engine.Texture',PinType.PinSubCategoryMemberReference=(),PinType.bIsArray=False,PinType.bIsReference=False,PinType.bIsConst=False,PinType.bIsWeakPointer=False,LinkedTo=(K2Node_VariableGet_185 3214E57245F463BA1D70C1845534BA2C,),PersistentGuid=00000000000000000000000000000000,bHidden=False,bNotConnectable=False,bDefaultValueIsReadOnly=False,bDefaultValueIsIgnored=False,bAdvancedView=False,)
+			End Object
+			Begin Object Class=K2Node_VariableGet Name="K2Node_VariableGet_185"
+			   VariableReference=(MemberName="HeightfieldRT",MemberGuid=E71A12524941DA802B08039BF7491FF7,bSelfContext=True)
+			   NodePosX=720
+			   NodePosY=1361
+			   NodeGuid=26D737594C6DDDCB502FD7A3C15904D5
+			   CustomProperties Pin (PinId=3214E57245F463BA1D70C1845534BA2C,PinName="HeightfieldRT",Direction="EGPD_Output",PinType.PinCategory="object",PinType.PinSubCategoryObject=Class'/Script/Engine.TextureRenderTarget2D',PinType.PinSubCategoryMemberReference=(),PinType.bIsArray=False,PinType.bIsReference=False,PinType.bIsConst=False,PinType.bIsWeakPointer=False,LinkedTo=(K2Node_CallFunction_370 FAC8C8B64E1E50423C9F4493F3C27926,),PersistentGuid=00000000000000000000000000000000,bHidden=False,bNotConnectable=False,bDefaultValueIsReadOnly=False,bDefaultValueIsIgnored=False,bAdvancedView=False,)
+			   CustomProperties Pin (PinId=D100E93E41E555D10A0028BE2218C562,PinName="self",PinFriendlyName=NSLOCTEXT("K2Node", "Target", "Target"),PinType.PinCategory="object",PinType.PinSubCategoryObject=BlueprintGeneratedClass'/Game/ExampleContent/BlueprintFluidSurface/HeightfieldPainter.HeightfieldPainter_C',PinType.PinSubCategoryMemberReference=(),PinType.bIsArray=False,PinType.bIsReference=False,PinType.bIsConst=False,PinType.bIsWeakPointer=False,PersistentGuid=00000000000000000000000000000000,bHidden=True,bNotConnectable=False,bDefaultValueIsReadOnly=False,bDefaultValueIsIgnored=False,bAdvancedView=False,)
+			End Object
+			Begin Object Class=K2Node_VariableGet Name="K2Node_VariableGet_186"
+			   VariableReference=(MemberName="PainterMaterialInstance",MemberGuid=9030CC2445DD109F0DC10BB00BA7797D,bSelfContext=True)
+			   NodePosX=640
+			   NodePosY=1313
+			   NodeGuid=24E9BBBF47D8AB17DF28F2BF2D7DF62E
+			   CustomProperties Pin (PinId=D6A13F044975C578FFF6A98328676BC9,PinName="PainterMaterialInstance",Direction="EGPD_Output",PinType.PinCategory="object",PinType.PinSubCategoryObject=Class'/Script/Engine.MaterialInstanceDynamic',PinType.PinSubCategoryMemberReference=(),PinType.bIsArray=False,PinType.bIsReference=False,PinType.bIsConst=False,PinType.bIsWeakPointer=False,LinkedTo=(K2Node_CallFunction_370 63A3B28B4B0AAB8799485AA1BF3EAE97,),PersistentGuid=00000000000000000000000000000000,bHidden=False,bNotConnectable=False,bDefaultValueIsReadOnly=False,bDefaultValueIsIgnored=False,bAdvancedView=False,)
+			   CustomProperties Pin (PinId=18307DBC4877CAE1B09745A501B84B7E,PinName="self",PinFriendlyName=NSLOCTEXT("K2Node", "Target", "Target"),PinType.PinCategory="object",PinType.PinSubCategoryObject=BlueprintGeneratedClass'/Game/ExampleContent/BlueprintFluidSurface/HeightfieldPainter.HeightfieldPainter_C',PinType.PinSubCategoryMemberReference=(),PinType.bIsArray=False,PinType.bIsReference=False,PinType.bIsConst=False,PinType.bIsWeakPointer=False,PersistentGuid=00000000000000000000000000000000,bHidden=True,bNotConnectable=False,bDefaultValueIsReadOnly=False,bDefaultValueIsIgnored=False,bAdvancedView=False,)
+			End Object
+			Begin Object Class=K2Node_CallFunction Name="K2Node_CallFunction_371"
+			   FunctionReference=(MemberName="EnableInput",bSelfContext=True)
+			   NodePosX=96
+			   NodePosY=1104
+			   NodeGuid=A72420974F0F631A77B768ADF20BC186
+			   CustomProperties Pin (PinId=59EBD0D84FA02620751F0CAC38ADA5FF,PinName="execute",PinToolTip="\nExec",PinType.PinCategory="exec",PinType.PinSubCategoryObject=None,PinType.PinSubCategoryMemberReference=(),PinType.bIsArray=False,PinType.bIsReference=False,PinType.bIsConst=False,PinType.bIsWeakPointer=False,LinkedTo=(K2Node_Event_0 4C1A3CD54C93D4A74803118ED8E9BE79,),PersistentGuid=00000000000000000000000000000000,bHidden=False,bNotConnectable=False,bDefaultValueIsReadOnly=False,bDefaultValueIsIgnored=False,bAdvancedView=False,)
+			   CustomProperties Pin (PinId=C8543EC14C417EB94C0A40806572C540,PinName="then",PinToolTip="\nExec",Direction="EGPD_Output",PinType.PinCategory="exec",PinType.PinSubCategoryObject=None,PinType.PinSubCategoryMemberReference=(),PinType.bIsArray=False,PinType.bIsReference=False,PinType.bIsConst=False,PinType.bIsWeakPointer=False,LinkedTo=(K2Node_CallFunction_359 50F50A054429DFDB03CCE884B76C4E0A,),PersistentGuid=00000000000000000000000000000000,bHidden=False,bNotConnectable=False,bDefaultValueIsReadOnly=False,bDefaultValueIsIgnored=False,bAdvancedView=False,)
+			   CustomProperties Pin (PinId=2472C14544442ACFC3F3CAAE2BFC249B,PinName="self",PinFriendlyName=NSLOCTEXT("K2Node", "Target", "Target"),PinToolTip="Target\nActor Reference",PinType.PinCategory="object",PinType.PinSubCategoryObject=Class'/Script/Engine.Actor',PinType.PinSubCategoryMemberReference=(),PinType.bIsArray=False,PinType.bIsReference=False,PinType.bIsConst=False,PinType.bIsWeakPointer=False,PersistentGuid=00000000000000000000000000000000,bHidden=False,bNotConnectable=False,bDefaultValueIsReadOnly=False,bDefaultValueIsIgnored=False,bAdvancedView=False,)
+			   CustomProperties Pin (PinId=538663B143E85B7459172FBD475EDB53,PinName="PlayerController",PinToolTip="Player Controller\nPlayer Controller Reference\n\nThe PlayerController whose input events we want to receive.",PinType.PinCategory="object",PinType.PinSubCategoryObject=Class'/Script/Engine.PlayerController',PinType.PinSubCategoryMemberReference=(),PinType.bIsArray=False,PinType.bIsReference=False,PinType.bIsConst=False,PinType.bIsWeakPointer=False,LinkedTo=(K2Node_CallFunction_372 DB581A434D6C7153C9DF6F9236D045F9,),PersistentGuid=00000000000000000000000000000000,bHidden=False,bNotConnectable=False,bDefaultValueIsReadOnly=False,bDefaultValueIsIgnored=False,bAdvancedView=False,)
+			End Object
+			Begin Object Class=K2Node_CallFunction Name="K2Node_CallFunction_372"
+			   bIsPureFunc=True
+			   FunctionReference=(MemberParent=Class'/Script/Engine.GameplayStatics',MemberName="GetPlayerController")
+			   NodePosX=-192
+			   NodePosY=1264
+			   NodeGuid=F0D3AAFF42D494B9C98436BC4F1B965B
+			   CustomProperties Pin (PinId=2843ED7743C312C47BFD3AB629A45D70,PinName="self",PinFriendlyName=NSLOCTEXT("K2Node", "Target", "Target"),PinToolTip="Target\nGameplay Statics Reference",PinType.PinCategory="object",PinType.PinSubCategoryObject=Class'/Script/Engine.GameplayStatics',PinType.PinSubCategoryMemberReference=(),PinType.bIsArray=False,PinType.bIsReference=False,PinType.bIsConst=False,PinType.bIsWeakPointer=False,DefaultObject="/Script/Engine.Default__GameplayStatics",PersistentGuid=00000000000000000000000000000000,bHidden=True,bNotConnectable=False,bDefaultValueIsReadOnly=False,bDefaultValueIsIgnored=False,bAdvancedView=False,)
+			   CustomProperties Pin (PinId=5BA2917945614FA7C1CCA2ACD9C09317,PinName="WorldContextObject",PinToolTip="World Context Object\nObject Reference",PinType.PinCategory="object",PinType.PinSubCategoryObject=Class'/Script/CoreUObject.Object',PinType.PinSubCategoryMemberReference=(),PinType.bIsArray=False,PinType.bIsReference=False,PinType.bIsConst=False,PinType.bIsWeakPointer=False,PersistentGuid=00000000000000000000000000000000,bHidden=True,bNotConnectable=False,bDefaultValueIsReadOnly=False,bDefaultValueIsIgnored=False,bAdvancedView=False,)
+			   CustomProperties Pin (PinId=9D63B5834AA0892C28999DAE7EAD90DD,PinName="PlayerIndex",PinToolTip="Player Index\nInteger",PinType.PinCategory="int",PinType.PinSubCategoryObject=None,PinType.PinSubCategoryMemberReference=(),PinType.bIsArray=False,PinType.bIsReference=False,PinType.bIsConst=False,PinType.bIsWeakPointer=False,DefaultValue="0",AutogeneratedDefaultValue="0",PersistentGuid=00000000000000000000000000000000,bHidden=False,bNotConnectable=False,bDefaultValueIsReadOnly=False,bDefaultValueIsIgnored=False,bAdvancedView=False,)
+			   CustomProperties Pin (PinId=DB581A434D6C7153C9DF6F9236D045F9,PinName="ReturnValue",PinToolTip="Return Value\nPlayer Controller Reference",Direction="EGPD_Output",PinType.PinCategory="object",PinType.PinSubCategoryObject=Class'/Script/Engine.PlayerController',PinType.PinSubCategoryMemberReference=(),PinType.bIsArray=False,PinType.bIsReference=False,PinType.bIsConst=False,PinType.bIsWeakPointer=False,LinkedTo=(K2Node_CallFunction_371 538663B143E85B7459172FBD475EDB53,),PersistentGuid=00000000000000000000000000000000,bHidden=False,bNotConnectable=False,bDefaultValueIsReadOnly=False,bDefaultValueIsIgnored=False,bAdvancedView=False,)
+			End Object
+			Begin Object Class=EdGraphNode_Comment Name="EdGraphNode_Comment_122"
+			   NodePosX=-240
+			   NodePosY=1056
+			   NodeWidth=544
+			   NodeHeight=320
+			   NodeComment="Enable input on this Actor"
+			   NodeGuid=F84B49FF441A8064BC205E8D99F1A7E4
+			End Object
+			Begin Object Class=EdGraphNode_Comment Name="EdGraphNode_Comment_123"
+			   NodePosX=-256
+			   NodePosY=1504
+			   NodeWidth=576
+			   NodeHeight=288
+			   NodeComment="Check to see if Left Mouse Button is held down"
+			   NodeGuid=CFE482F448840AC8ACDAAB8013C7A7B1
+			End Object
+			Begin Object Class=EdGraphNode_Comment Name="EdGraphNode_Comment_124"
+			   NodePosX=336
+			   NodePosY=1056
+			   NodeWidth=852
+			   NodeHeight=390
+			   NodeComment="Clear the Render Target off all previous data"
+			   NodeGuid=5BFD8E0C41D20A929E66B397257C5783
+			End Object
+
+			END OBJECT
+	[/PARAM]
+[/OBJECT]
+[/EXCERPT:EG]
+-->
\ No newline at end of file
