Copying //UE4/Dev-Documentation to Samples-Main (//UE4/Samples-Main)

diff --git a/Engine/Documentation/Source/Engine/Rendering/RenderTargets/BlueprintRenderTargets/HowTo/HeightFieldPainter/4/HFP4.INT.udn b/Engine/Documentation/Source/Engine/Rendering/RenderTargets/BlueprintRenderTargets/HowTo/HeightFieldPainter/4/HFP4.INT.udn
new file mode 100644
--- /dev/null
+++ b/Engine/Documentation/Source/Engine/Rendering/RenderTargets/BlueprintRenderTargets/HowTo/HeightFieldPainter/4/HFP4.INT.udn
@@ -0,0 +1,79 @@
+availability:docs
+title:4. End Result
+description:Going over how to use the Blueprint Height Map Painter.
+crumbs:%ROOT%, Engine, Engine\Rendering\RenderTargets
+version:4.13
+type:multi-step
+SkillLevel:Intermediate
+parent:Engine/Rendering/RenderTargets/BlueprintRenderTargets/HowTo/HeightFieldPainter
+Navigation:multi-step
+tags: Render Targets
+order:4
+checkpoint:HFP
+
+[NAV]
+
+Now that all of the required parts for the Height Field Painter Blueprint have been setup, it is now time to see the Height Field Painter in action. In the following section, we will go over how to use the Height Field Painter in your UE4 project.
+
+## Using the Height Field Painter Blueprint
+
+1. Locate the Height Field Painter Blueprint in the Content Browser and then drag it into your level, positioning it to fit your level's needs.
+
+	![](Engine\Rendering\RenderTargets\BlueprintRenderTargets\HowTo\HeightFieldPainter\HFP_BP_Use_00.png)
+
+1. Press the **Play** button on the Main Toolbar and then navigate your player so that they are in front (or very close) to the Height Field Painter Blueprint that was placed in the level.
+
+	![](Engine\Rendering\RenderTargets\BlueprintRenderTargets\HowTo\HeightFieldPainter\HFP_BP_Use_01.png)
+
+## End Result
+
+To raise the height of the surface position so that the crosshair is at a desired position, hold down the left mouse button to raise the height of the Static Mesh. If you continuously hold down the left mouse button while moving the mouse around, 
+you will continue to raise the height of the Static Mesh. If everything has been setup correctly, you should have something similar to the following video.
+
+[OBJECT:EmbeddedVideo]
+	[PARAMLITERAL:width]
+	640
+	[/PARAMLITERAL]
+	[PARAMLITERAL:height]
+	360
+	[/PARAMLITERAL]
+	[PARAMLITERAL:videoid]
+	K60O1o8UqzE
+	[/PARAMLITERAL]
+[/OBJECT]
+
+While this is just one way you can manipulate Static Meshes vertices using Blueprints & Render Targets, the creative possibilities are almost endless. From creating impact effects that can deform both players and the environment, to allowing players to have a real effect on the worlds they play in, the Height Field Painter Blueprint can be adopted to fit many different gameplay roles and requirements.
+
+
+
+
+
+
+
+
+
+
+
+
+
+
+
+
+
+
+
+
+
+
+
+
+
+
+
+
+
+
+
+
+
+[NAV]
\ No newline at end of file
diff --git a/Engine/Documentation/Source/Engine/Rendering/RenderTargets/BlueprintRenderTargets/HowTo/HeightFieldPainter/HeightFieldPainterHowTo.INT.udn b/Engine/Documentation/Source/Engine/Rendering/RenderTargets/BlueprintRenderTargets/HowTo/HeightFieldPainter/HeightFieldPainterHowTo.INT.udn
--- a/Engine/Documentation/Source/Engine/Rendering/RenderTargets/BlueprintRenderTargets/HowTo/HeightFieldPainter/HeightFieldPainterHowTo.INT.udn
+++ b/Engine/Documentation/Source/Engine/Rendering/RenderTargets/BlueprintRenderTargets/HowTo/HeightFieldPainter/HeightFieldPainterHowTo.INT.udn
@@ -1,38 +1,26 @@
 availability:docs
-title:Creating Deformable Surfaces with Blueprints & Render Targets
-description:How-To use Blueprints and Render Targets to create realtime deformable surfaces.
-crumbs:
+title:Creating a Height Field Painter with Blueprints and Render Targets
+description:How-To use Blueprints and Render Targets to create realtime Height Filed Painter.
+crumbs:%ROOT%, Engine, Engine\Rendering\RenderTargets
 version:4.13
-Type:topic
+Type:how-to
 SkillLevel:Intermediate
-parent:Engine
-order:2
-Navigation:topic
-
-
-
-
-
-
-
-
-
-
-
-
-
-
-
-
-
-
-
-
-
-
-
-
-
+parent:Engine/Rendering/RenderTargets/BlueprintRenderTargets/HowTo
+order:1
+Navigation:landing
+Related:Gameplay/Framework/GameMode/
+Related:Engine/Blueprints/
+Related:Engine/Rendering/Materials/
+Related:Engine/Blueprints/UserGuide/UserConstructionScript/
+tags:Render Targets
+tags:Blueprint
+checkpoint:HFP
+
+![](HFP_Header_Image.png)
+
+Having the ability to raise and lower various surfaces enables interesting and unique gameplay elements. In the following How To, we will go over how you can use Blueprints & Render Targets to create a Height Field Painter that raises the surface of an object when the player shoots at it. To accomplish this we will first create and setup a new Game Mode and HUD Blueprint that will allow the user to interact with the Height Field Painter. We will then create and setup the needed render targets and Materials. Finally we will create a Blueprint that will tie all of the individual pieces together, resulting in a Blueprint that can raise the surface of the object that it is applied to inside of Unreal Engine 4 (UE4).
+
+[DIR(output:"checkpoint" parent:"Engine/Rendering/RenderTargets/BlueprintRenderTargets/HowTo/HeightFieldPainter")]
 
 
 
diff --git a/Engine/Documentation/Source/Engine/Rendering/RenderTargets/BlueprintRenderTargets/Overview/BlueprintRenderTargetsOverview.INT.udn b/Engine/Documentation/Source/Engine/Rendering/RenderTargets/BlueprintRenderTargets/Overview/BlueprintRenderTargetsOverview.INT.udn
--- a/Engine/Documentation/Source/Engine/Rendering/RenderTargets/BlueprintRenderTargets/Overview/BlueprintRenderTargetsOverview.INT.udn
+++ b/Engine/Documentation/Source/Engine/Rendering/RenderTargets/BlueprintRenderTargets/Overview/BlueprintRenderTargetsOverview.INT.udn
@@ -1,29 +1,41 @@
-availability:docs
-title:Blueprint & Render Targets Overview
+availability:Docs
+title:Blueprints and Render Targets Overview
 description:Taking a high level look at how to use Blueprints and Render Targets togther in UE4.
-crumbs:
+crumbs:%ROOT%, Engine, Engine/Rendering, Engine/Rendering/RenderTargets, Engine/Rendering/RenderTargets/BlueprintRenderTargets
 version:4.13
-Type:topic
+Type:overview
 SkillLevel:Intermediate
-parent:Engine
-order:2
+parent:Engine/Rendering/RenderTargets/BlueprintRenderTargets
+order:1
 Navigation:topic
+tags: Render Targets
+tags: Blueprint
+topic-image:Engine/Rendering/RenderTargets/Render_Targets_Overview.png
 
 ![](BP_RT_Overview.png)
 
 ## Overview
-Using Blueprints to interact with render targets opens up a an entirely new way to create and use assets inside Unreal Engine 4 (UE4).  From creating fluid like surfaces to Textures that can be used with visual effects and Materials, there is almost no limit to the amount of creative freedom that this will offer your UE4 projects. In the following document, we will go over some of the things that you need to be aware of when using Blueprints and render targets.  
+Using Blueprints and render targets together opens up a whole new way to create and use assets inside of Unreal Engine 4 (UE4).  From creating fluid like surfaces, to creating Textures that can be used with visual effects and Materials there are many different ways Blueprints and render targets can be used. In the following document, we will go over some of the things that you need to be aware of when using Blueprints and render targets.
 
-## Know Limitations
-When using Blueprints and render targets, there are few limitations and restrictions that you need to be aware of in order to get the most out of this tool. In the following section we will cover what these limitations are and what can be done to accommodate them in your workflow.
+## Possible Uses 
+You can use Blueprints and render targets to create a wide variety of different rendering effects. The following list has a few suggestions of the types of items that can be created when using Blueprints and render targets. Keep in mind that the following are only suggestions on what you can achieve with Blueprints and render targets. 
 
-* You cannot draw to a render target that is being sampled as a Texture by the Material that is input into the Material input. If you want to accomplish this, you will need to either use alpha blending to modify a render target in-place, or ping-pong between two different render targets.
+* Create a Heightmap painter.
+* Create a Fluid surface simulation.
+* Create a runtime paint program.
+* Deform a surface when a certain object passes over.
 
-* Only the **Emissive Color** and **Opacity** outputs of the Material are valid when drawing to a render target. 
 
-* Lighting is not supported when using render targets.
+## Known Limitations
+When using Blueprints and render targets, there are some limitations and restrictions that you need to be aware of in order 
+to get the most out of this tool. In the following section, we will cover what these limitations are and what can be done to 
+accommodate them in your workflow.
+
+* You cannot draw to a render target that is being sampled as a Texture by the Material you specify.  Either use alpha blending to modify a render target in-place, or go back and forth between two different render targets.
+
+* Only the **Emissive Color** and **Opacity** outputs of the Material are valid when drawing to a render target. 
 
-* Material Expression nodes like **World Position** may return unexpected values.
+* Lighting is not supported when using render targets and Material Expression nodes like **World Position** may return unexpected values.
 
 * Emissive Color is clamped to be positive by default, but you can output negative values by enabling the material property **AllowNegativeEmissiveColor**.
 
diff --git a/Engine/Documentation/Source/Engine/Rendering/RenderTargets/BlueprintRenderTargets/Reference/BlueprintRenderTargetsReference.INT.udn b/Engine/Documentation/Source/Engine/Rendering/RenderTargets/BlueprintRenderTargets/Reference/BlueprintRenderTargetsReference.INT.udn
--- a/Engine/Documentation/Source/Engine/Rendering/RenderTargets/BlueprintRenderTargets/Reference/BlueprintRenderTargetsReference.INT.udn
+++ b/Engine/Documentation/Source/Engine/Rendering/RenderTargets/BlueprintRenderTargets/Reference/BlueprintRenderTargetsReference.INT.udn
@@ -1,88 +1,95 @@
 availability:docs
-title: Blueprint & Render Target Reference
-description: A listing of properties and descriptions for using Blueprints & Render Targets.
-crumbs:
+title: Blueprint and Render Targets Reference
+description: A listing of properties and descriptions for using Blueprint & Render Targets.
+crumbs:%ROOT%, Engine, Engine/Rendering, Engine/Rendering/RenderTargets, Engine/Rendering/RenderTargets/BlueprintRenderTargets
 version:4.13
-Type:Landing
-SkillLevel:Intermediate
-parent:Engine
+Type:reference
+parent:Engine/Rendering/RenderTargets/BlueprintRenderTargets
 order:2
 Navigation:Reference
+tags: Render Targets
+tags: Blueprint
+topic-image:Engine/Rendering/RenderTargets/Render_Targets_Reference.png
 
-## Overview
-This page contains a reference listing of properties you will find with the Render Target Blueprint nodes.
+[TOC (start:2 end:3)]
 
+This page contains a listing of properties you will find with the render target Blueprint nodes.
 
 ## Begin Draw Canvas to Render Target
 
 ![](Begin_Draw_Canvas_To_Render_Target.png)
 
-* Returns a Canvas object that can be used to draw to the specified render target. Be sure to call **EndDrawCanvasToRenderTarget** to complete the rendering!
+Returns a Canvas object that can be used to draw to the specified render target. Be sure to call **EndDrawCanvasToRenderTarget** to complete the rendering!
+
+[region:note]
+Please note that DrawMaterialtoRenderTarget is the simpler version and should be used when possible instead of Begin/EndDrawCanvasToRT.
+[/region]
 
 |Property| Description|
 |----------|-----------------|
-|Texture Render Target| Texture Render Target 2D reference.|
+|Texture Render Target| Texture render target 2D reference.|
 |Canvas| Reference to canvas.|
 |Size| Size of the render target.|
-|Context| Draw to Render Target Context Structure.|
+|Context| Draw to render target Context Structure.|
 
 ## End Draw Canvas to Render Target
 
 ![](End_Draw_Canvas_to_Render_Target.png)
 
-* Must be paired with a **BeginDrawCanvasToRenderTarget** to complete rendering to a Render Target.
+Must be paired with a **BeginDrawCanvasToRenderTarget** to complete rendering to a render target.
 
 |Property| Description|
 |----------|-----------------|
-|Context| Draw to Render Target Context Structure.|
+|Context| Draw to render target Context Structure.|
 
 
 ## Clear Render Target 2D
 
 ![](Clear_Render_Target_2D.png)
 
-* Clears the specified Render Target with the given Clear Color.
+Clears the specified render target with the given Clear Color.
 
 |Property| Description|
 |----------|-----------------|
-|Texture Render Target| Texture Render Target 2D reference.|
-|Clear Color| Clears the specified Render Target with the given Clear Color.|
+|Texture Render Target| Texture render target 2D reference.|
+|Clear Color| Clears the specified render target with the given Clear Color.|
 
 
 ## Create Render Target 2D
 
 ![](Create_Render_Target_2D.png)
 
-* Creates a new Render Target and initializes it to the specified dimensions.
+Creates a new render target and initializes it to the specified dimensions.
 
 |Property| Description|
 |----------|-----------------|
-|Width| Render Target's width.|
-|Height| Render Target's height.|
-|Return Value| Texture Render Target 2D reference.|
+|Width| Render targets width size.|
+|Height| Render targets height size.|
+|Return Value| Texture render target 2D reference.|
 
 ## Draw Material to Render Target
 
 ![](Draw_Material_To_Render_Target.png)
 
-* Renders a quad with the Material applied to the specified Render Target. This set the Render Target even if it is already set, which is an expensive operation. Use **BeginDrawCanvasToRenderTarget / EndDrawCanvasToRenderTarget** instead if rendering multiple primitives to the same Render Target.
+Renders a quad with the Material applied to the specified render target. This node sets the render target even if it is already set, which is an expensive operation. 
+If you are rendering multiple primitives to the same target area, use **BeginDrawCanvasToRenderTarget / EndDrawCanvasToRenderTarget**.
 
 |Property| Description|
 |----------|-----------------|
-|Texture Render Target| Texture Render Target 2D reference.|
-|Material| Material that the Render Target results will be saved to.|
+|Texture Render Target| Texture render target 2D reference.|
+|Material| Material that the render target results will be saved to.|
 
 ## Export Render Target
 
 ![](Export_Render_Target.png)
 
-* Exports a Render Target as an HDR image onto the disk.
+Exports a Render Target as an HDR image onto the disk.
 
 |Property| Description|
 |----------|-----------------|
-|Texture Render Target| Texture Render Target 2D reference.|
-|File Path| Path to where the Render Target will be saved.|
-|File Name| Name to give the saved Render Target.|
+|Texture Render Target| Texture render target 2D reference.|
+|File Path| Path to where the render target will be saved.|
+|File Name| Name to give the saved render target.|
 
 
 
diff --git a/Engine/Documentation/Source/Engine/Rendering/RenderTargets/RenderTargets.INT.udn b/Engine/Documentation/Source/Engine/Rendering/RenderTargets/RenderTargets.INT.udn
--- a/Engine/Documentation/Source/Engine/Rendering/RenderTargets/RenderTargets.INT.udn
+++ b/Engine/Documentation/Source/Engine/Rendering/RenderTargets/RenderTargets.INT.udn
@@ -1,40 +1,20 @@
-availability:docs
+availability:Docs
 title: Render Targets
-description:Using Render Targets in Unreal Engine 4
-crumbs:
+description:Using Render Targets in Unreal Engine 4.
+crumbs:%ROOT%, Engine, Engine/Rendering
 version:4.13
 Type:landing
-SkillLevel:Intermediate
-parent:Engine
-order:1
+parent:Engine/Rendering
+order:7
 Navigation:topic
+tags: Render Targets
+tags: Blueprint
 
 ![](BlueprintRenderTarget.png)
 
-A render target is a buffer where the video card draws pixels for a scene that is being rendered. render targets are used extensively throughout Unreal Engine 4 (UE4) and in the following set of documents we will take a look at how to use and manipulate render targets inside UE4.
-
-
-
-
-[REGION:topics third]
-[OBJECT:Topic]
-	[PARAM:image]
-		![Engine\Rendering\RenderTargets\BlueprintRenderTargets:description%](Engine/Rendering/RenderTargets/Render_Targets_Overview.png)
-	[/PARAM]
-	[PARAM:icon]
-		![](%ROOT%/blueprint_icon.png)(convert:false)
-	[/PARAM]
-	[PARAM:title]
-		%Engine\Rendering\RenderTargets\BlueprintRenderTargets/:description%
-	[/PARAM]
-	[PARAM:description]
-		%Engine\Rendering\RenderTargets\BlueprintRenderTargets:description%
-	[/PARAM]
-	[PARAM:path]
-		[RELATIVE:Engine\Rendering\RenderTargets\BlueprintRenderTargets/]
-	[/PARAM]
-[/OBJECT]
-[/REGION]
+Render targets are used quite extensively throughout Unreal Engine 4 (UE4) for a wide variety of different tasks. From storing the various buffers that are needed for a deferred renderer to helping display various complex effects like the ripples that happen when you throw an object into a pool of water, render targets can be leveraged to perform all kinds of interesting and useful effects.   In the following set of documents, we will take a look at how to use and manipulate render targets inside of UE4.
+
+[DIR(output:"topic" parent:"Engine/Rendering/RenderTargets")]
 
 
 
diff --git a/Engine/Documentation/Source/Engine/Rendering/Rendering.INT.udn b/Engine/Documentation/Source/Engine/Rendering/Rendering.INT.udn
--- a/Engine/Documentation/Source/Engine/Rendering/Rendering.INT.udn
+++ b/Engine/Documentation/Source/Engine/Rendering/Rendering.INT.udn
@@ -2,49 +2,10 @@ Availability:Public
 Title: Rendering and Graphics
 Crumbs: %ROOT%, Engine
 Description:Rendering subsystem including lighting and shadowing, materials and textures, particles and effects, and post processing.
-Related:Programming/Rendering
 Version: 4.9
 parent:Engine
 order:1
-
-[VAR:Topic]
-[OBJECT:Topic]
-	[PARAM:image]
-		![%Engine/Rendering:title%](Engine/Rendering/rendering_topic.png)
-	[/PARAM]
-	[PARAM:icon]
-		![](%ROOT%/light_icon.png)(convert:false)
-	[/PARAM]
-	[PARAM:title]
-		%Engine/Rendering:title%
-	[/PARAM]
-	[PARAM:description]
-		%Engine/Rendering:description%
-	[/PARAM]
-	[PARAM:path]
-		[RELATIVE:Engine/Rendering]
-	[/PARAM]
-[/OBJECT]
-[/VAR]
-[VAR:TopicCompact]
-[OBJECT:TopicCompact]
-	[PARAM:image]
-		![%Engine/Rendering:title%](Engine/Rendering/rendering_topic.png)
-	[/PARAM]
-	[PARAM:icon]
-		![](%ROOT%/light_icon.png)(convert:false)
-	[/PARAM]
-	[PARAM:title]
-		%Engine/Rendering:title%
-	[/PARAM]
-	[PARAM:description]
-		%Engine/Rendering:description%
-	[/PARAM]
-	[PARAM:path]
-		[RELATIVE:Engine/Rendering]
-	[/PARAM]
-[/OBJECT]
-[/VAR]
+topic-image:rendering_topic.png
 
 [REGION:banner]
 ![](rendering_banner.png)(convert:false)
@@ -85,24 +46,15 @@ order:1
 [/REGION]
 -->
 
-## Basics
-
-[REGION:topics third]
-%Engine/Rendering/Overview:Topic%
-%Engine/Rendering/LightingAndShadows:Topic%
-%Engine/Rendering/Materials:Topic%
-%Engine/Rendering/ParticleSystems:Topic%
-%Engine/Rendering/PostProcessEffects:Topic%
-%Engine/Actors/FogEffects:Topic%
-[/REGION]
+[DIR(output:"topic" parent:"Engine/Rendering" type:"landing" end:"1")]
 
 ## Performance & Debugging
 
-[REGION:topics third]
-%Engine/Performance:topic%
-%Engine/Performance/Scalability:topic%
-%Engine/Performance/Visualizers:topic%
-[/REGION]
+[DIR(output:"topic" parent:"Engine/Performance" tags:"performance and profiling" end:"1")]
+
+## Graphics Programming
+
+[DIR(output:"listbutton" parent:"Programming" tags:"rendering" org:"hierarchy")]
 
 
 
diff --git a/Engine/Documentation/Source/Engine/Rendering/SceneCapture/SceneCapture.INT.udn b/Engine/Documentation/Source/Engine/Rendering/SceneCapture/SceneCapture.INT.udn
new file mode 100644
--- /dev/null
+++ b/Engine/Documentation/Source/Engine/Rendering/SceneCapture/SceneCapture.INT.udn
@@ -0,0 +1,49 @@
+availability:docs
+title:Scene Capture
+description: Using Scene Capture Actors in Unreal Engine 4.
+crumbs:
+version:4.14
+Type:landing
+parent:Engine/Rendering
+order:1
+Navigation:landing
+tags: Scene  Capture
+tags: Rendering
+topic-image:Engine\Rendering\SceneCapture/Scene_Capture_Header.png
+
+![](Scene_Capture_Header.png) 
+
+
+Unreal Engine 4  (UE4) offers a set of Actors, called Scene Capture Actors, that can be used to capture a real time, fully dynamic image of the scene it is placed in. In the following pages we will go over what Scene Capture Actors are and how you can use them inside your UE4 projects.
+
+[DIR(output:"topic" parent:"Engine/Rendering/SceneCapture")]
+
+
+
+
+
+
+
+
+
+
+
+
+
+
+
+
+
+
+
+
+
+
+
+
+
+
+
+
+
+
diff --git a/Engine/Documentation/Source/Engine/Rendering/SceneCapture/SceneCapture2D/SceneCapture2D.INT.udn b/Engine/Documentation/Source/Engine/Rendering/SceneCapture/SceneCapture2D/SceneCapture2D.INT.udn
new file mode 100644
--- /dev/null
+++ b/Engine/Documentation/Source/Engine/Rendering/SceneCapture/SceneCapture2D/SceneCapture2D.INT.udn
@@ -0,0 +1,48 @@
+availability:Docs
+title:Scene Capture 2D
+description: Using the Scene Capture 2D Actor in Unreal Engine 4.
+crumbs:
+version:4.14
+Type:topic
+SkillLevel:Intermediate
+parent:Engine/Rendering/SceneCapture
+order:2
+Navigation:topic
+
+
+
+
+
+
+
+
+
+
+
+
+
+
+
+
+
+
+
+
+
+
+
+
+
+
+
+
+
+
+
+
+
+
+
+
+
+
diff --git a/Engine/Documentation/Source/Engine/Rendering/SceneCapture/SceneCaptureCube/SceneCaptureCube.INT.udn b/Engine/Documentation/Source/Engine/Rendering/SceneCapture/SceneCaptureCube/SceneCaptureCube.INT.udn
new file mode 100644
--- /dev/null
+++ b/Engine/Documentation/Source/Engine/Rendering/SceneCapture/SceneCaptureCube/SceneCaptureCube.INT.udn
@@ -0,0 +1,48 @@
+availability:Docs
+title:Scene Capture Cube
+description: Using the Scene Capture Cube Actor in Unreal Engine 4.
+crumbs:%ROOT%, Engine, Engine/Rendering, Engine/Rendering/SceneCapture
+version:4.14
+Type:topic
+SkillLevel:Intermediate
+parent:Engine/Rendering/SceneCapture
+order:2
+Navigation:topic
+
+
+
+
+
+
+
+
+
+
+
+
+
+
+
+
+
+
+
+
+
+
+
+
+
+
+
+
+
+
+
+
+
+
+
+
+
+
diff --git a/Engine/Documentation/Source/Engine/Rendering/SceneCapture/SceneCaptureMobile/SceneCaptureMobile.INT.udn b/Engine/Documentation/Source/Engine/Rendering/SceneCapture/SceneCaptureMobile/SceneCaptureMobile.INT.udn
new file mode 100644
--- /dev/null
+++ b/Engine/Documentation/Source/Engine/Rendering/SceneCapture/SceneCaptureMobile/SceneCaptureMobile.INT.udn
@@ -0,0 +1,57 @@
+availability:Docs
+title:Scene Capture for Mobile
+description: Things to be aware of when using Scene Captures for Mobile devices.
+crumbs:%ROOT%, Engine, Engine/Rendering, Engine/Rendering/SceneCapture
+version:4.14
+Type:topic
+SkillLevel:Intermediate
+parent:Engine/Rendering/SceneCapture
+order:2
+Navigation:topic
+topic-image:Scene_Capture_Mobile.png
+
+When using the Scene Capture Actor in project that are to be used on Mobile devices you will need to aware of the following information to make sure the Scene Capture Actors are as performant as possible when used on Mobile Devices.
+
+
+
+![](Scene_Capture_Mobile.png)
+
+
+[region:note]
+When rendering Scene Captures, the Scene Capture Source settings that output Inverse Opacity and Depth values are  supported on Mobile devices. 
+[/region]
+
+* The **SceneColor (HDR) in RGB, Inv Opacity in A** option can be used to render objects with translucency into a Texture which can then be alpha-blended over a scene or widget Blueprint. 
+
+* Similarly, the depth value can be used as a mask when using the resulting Texture.
+
+* Generating the Ppacity data has some cost, so use **SceneColor (HDR) in RGB, 0 in A** for improved performance if you do not need Opacity.
+
+* Scene captures now work correctly on devices that do not support floating point targets, such as Galaxy S6 prior to Android 6.
+
+
+
+
+
+
+
+
+
+
+
+
+
+
+
+
+
+
+
+
+
+
+
+
+
+
+
diff --git a/Engine/Documentation/Source/Engine/Replay/Replay.INT.udn b/Engine/Documentation/Source/Engine/Replay/Replay.INT.udn
--- a/Engine/Documentation/Source/Engine/Replay/Replay.INT.udn
+++ b/Engine/Documentation/Source/Engine/Replay/Replay.INT.udn
@@ -1,25 +1,40 @@
-Availability:Docs
+Availability:Public
 Title: Replay System
 Crumbs: %ROOT%, Engine
 Description:Overview of the Replay system for recording and playback of gameplay
-Version: 4.12
+Version: 4.13
 parent:Engine
-order:17
+order:18
 topic-image:replay_topic.png
+tags:Replays
 
 
 [REGION:banner]
 ![](replay_banner.png)(convert:false)
 [/REGION]
 
+[TOC (start:2 end:2)]
+
 ## Replay System Overview
-Unreal Engine 4 features a Replay system which can record gameplay for later viewing. At a high level, the Replay system works by reading data drawn from the built-in replication system, similar to how live gameplay operates in a networked environment. Any project that is set up to replicate data like a multiplayer game does, even if the project doesn't actually have a multiplayer mode, will be compatible with the Replay system without further modification. The way this works is by passing replicated data to a Replay Streamer, which handles the process of filtering and storing the data. When the Replay system attempts to show a saved replay, it will have all of the replicated information that was available during live play, as well as special data fields designated as relevant only to replays, and can reconstruct the events of the game from this data.
+Unreal Engine 4 features a Replay system which can record gameplay for later viewing. At a high level, the Replay system works by using a `DemoNetDriver` to read data drawn from the built-in replication system, similar to how a `NetDriver` operates in a live, networked gameplay environment. Any project that is set up to replicate data like a multiplayer game does, even if the project doesn't actually have a multiplayer mode, will be compatible with the Replay system without further modification. The way this works involves the `DemoNetDriver` passing network data to a Replay Streamer, which handles the process of filtering and storing the data. When viewing a replay, a `DemoNetDriver` will have access to all of the replicated information that was available during live play, as well as special data fields designated as relevant only to replays, and can reconstruct the events of the game from this data.
 
 ## Replay System Functionality
 The basic functions of the replay system are to start or stop recording a demo, or play back a previously-recorded demo. While in playback mode, the Replay system supports commands to pause, change the playback speed, or jump to a specific point in time. More advanced usage of the system can include adding text tags (usually meta tags or the names of players involved in the game) to a replay, or querying existing replays based on text tags or version information present. The Replay system can be used from C++ code, mainly via the `UGameInstance` and `UWorld` classes, or with console commands or command-line arguments. A unified C++/Blueprint API is under construction and will be released with a future engine version.
 
-## Replay Data Format
-In terms of data, a replay contains three types of game-state information, as well as some additional text data. At the start is baseline data describing the starting state of the game world. Checkpoints that act as snapshots of the net changes to the world (from the baseline) appear at regular, user-defined intervals. The space between checkpoints is then filled with incremental changes to individual objects in the world. Any moment in the game can be reconstructed by the engine quickly and accurately by initializing the world to the starting state, making the changes described in the last checkpoint before the chosen time, and then applying each incremental change after that checkpoint leading up to the desired point in time. The text data contained in a replay consists of a display name, which can be used when making a player-facing list, and user-defined text tags (HTTP Streamer only), which can be used when searching or filtering through lists of games.
+| Function Or Variable Name | Equivalent Console Command | Effect |
+| --- | --- | --- |
+| `UGameInstance::StartRecordingReplay` | "demorec (ReplayName)" | Begins recording a replay. If the (ReplayName) parameter is not supplied, the replay system can generate a name on its own. |
+| `UGameInstance::StopRecordingReplay` | "demostop" | Stops a replay currently being recorded. |
+| `UGameInstance::PlayReplay` | "demoplay (ReplayName)" | Plays the named replay. |
+| `UDemoNetDriver::GotoTimeInSeconds` | "demoscrub (ReplayTime)" | Jumps to the provided time (in seconds) in the replay. |
+| `UWorld::Pauser` (variable) | "demopause" | Setting this variable to the PlayerController held by the World's `DemoNetDriver->ServerConnection->OwningActor` variable will pause the replay. Setting it back to `NULL` will resume the replay. The "demopause" console command will automatically set or clear that variable. |
+| `WorldSettings->DemoPlayTimeDilation` (variable) | "demospeed (ReplaySpeed)" | Speeds up or slows down playback by setting the time-scaling factor. A value of 1.0 is normal speed. |
+| `AGameMode::bHandleDedicatedServerReplays` (variable) | N/A | Setting this variable to can be set to true to cause dedicated servers to record games automatically. |
+
+## DemoNetDriver and Streamers
+The `DemoNetDriver` uses **Streamers** to perform the actual extraction and recording of the information needed to create the replay. Depending on how the replay data is intended to be viewed, different Streamers can be attached to the `DemoNetDriver`. There are three different Streamers included with the engine, each suited to different use cases. The **NULL Streamer**, which is the default, records events from the host machine directly to disk, making it best suited for single-player games and games that keep replays locally, on the host player's own machine. The **Memory Streamer** runs on a client machine and stores data in memory, making it ideal for "instant replay" features in sports titles or "kill cams" in shooters. Finally, the **HTTP Streamer** is used to send replay data to a second machine over LAN or the internet, which makes it a great choice for dedicated-server games and for games that may need to be streamed live to a large number of viewers while still remaining responsive to players.
+
+## Backward Compatibility
+Backward compatibility is supported by the Replay system, starting in version 4.13. This means you can make modifications to a build, i.e. adding or removing replicated properties, and still load and view replays that were recorded with the original build. In most cases, this is handled automatically, as the Replay system will skip over old replicated fields that have been removed, or use defaults for new replicated fields that have been added. Note that structures with `NetSerialize` overridden will require these differences in data format to be handled manually. To this end, `FArchive` now provides the `EngineNetVer` and `GameNetVer` functions, which identify the engine and game version, respectively, enabling you to adjust replay data as it comes in so that it fits with your current build. The delegate `GetLocalNetworkVersionOverride` in `FNetworkVersion` can be bound to a function that returns a `uint32` version number of your choosing. In case exact matches are not required for compatibility, `FNetworkVersion` has a delegate called `IsNetworkCompatibleOverride` that can be overridden to perform the comparison between two version numbers and decide if they are compatible.
 
-## Streamers
-The Replay system's Streamers perform the actual extraction and recording of the information needed to create the replay. Depending on how the replay data is intended to be viewed, different Streamers can be used. There are three different Streamers, each suited to different use cases. The NULL Streamer, which is the default, records events from the host machine directly to disk, making it best suited for single-player games and games that keep replays locally, on the host player's own machine. The Memory Streamer runs on a client machine and stores data in memory, making it ideal for "instant replay" features in sports titles or "kill cams" in shooters. Finally, the HTTP Streamer is used to send replay data to a second machine over LAN or the internet, which makes it a great choice for dedicated-server games and for games that may need to be streamed live to a large number of viewers while still remaining responsive to players.
+For further details on the features of the `DemoNetDriver` and its Streamers, please see the [](Engine/Replay/Streamers) page.
diff --git a/Engine/Documentation/Source/Engine/Replay/Streamers/HTTPRESTAPI/HTTPRESTAPI.INT.udn b/Engine/Documentation/Source/Engine/Replay/Streamers/HTTPRESTAPI/HTTPRESTAPI.INT.udn
new file mode 100644
--- /dev/null
+++ b/Engine/Documentation/Source/Engine/Replay/Streamers/HTTPRESTAPI/HTTPRESTAPI.INT.udn
@@ -0,0 +1,287 @@
+Availability:Docs
+Title: HTTP Streamer REST API
+Crumbs: %ROOT%, Engine, Engine/Replay, Engine/Replay/Streamers
+Description:Overview of the REST API for the HTTP Streamer
+Version: 4.14
+parent:Engine/Replay/Streamers
+topic-image:replay_topic.png
+tags:Replays
+
+[TOC (start:2 end:3)]
+
+## HTTP REST API Implementation
+
+The following requests will be made by the HTTP Streamer and must be handled by your replay server. In most cases, simply recording the information given in the appropriate location, or providing that information back to a viewer upon request, will be sufficient. It is good to note at this point that replays are not generally stored as a single file. Each individual chunk of a replay generally occupies its own file. Event data, although associated with a replay, is stored separately, since events can be searched without knowing what replay they belong to. Even event groups should be stored as separate files with lists of all events (across all replays) belonging to that group.
+
+The HTTP Streamer's REST API provides functionality for streaming replays from the game server to the replay server (uploading), for streaming live or pre-recorded replays from the replay server to the game, or viewer, server (downloading), or for querying information about the replays available on the server (searching). To implement your own replay server, you will need to respond to the following HTTP requests, as well as any additional requests specific to your game.
+
+### Upload Requests
+
+---
+
+<!-- From StartStreaming(): -->
+* **(Server URL)**replay/**(Session Name Override)**?app=**(Replay Version AppString)**&version=**(Replay Network Version)**&cl=**(Replay CL)**&friendlyName=**(Platform Friendly Name)**
+| Element | Value | Intended Use / Notes |
+| --- | --- | --- |
+| Server URL | Replay server URL. Drawn from DefaultEngine.Ini, unmodified. | The final "/" in the URL must be included - it will not be added automatically. |
+| Session Name Override | (Optional, including the leading "/".) Any name for the replay. Usually includes a GUID. | The replay server is expected to recognize this name when a replay is being requested. It should be returned as the session name, although it can be ignored and the streamer will respect the returned name. |
+| Replay Version AppString | A user-defined string that describes the game being played. | This value is expected to be used for filtering when the HTTP Streamer requests a list of replays. |
+| Replay Network Version | A user-defined string that describes the version of the game being played. | This value is also used for filtering lists of replays, so it should be stored with the replay data. |
+| Replay CL | The changelist of the build being played. | This value is used for filtering lists of replays. |
+| Platform Friendly Name | The platform on which the game is being played. | This value might be displayed to users, and is expected to be returned when lists of replays are requested. However, it will not be sent to the game server with requests for lists of replays and therefore is not expected to be used as a filter. |
+Uses POST. "Content-Type" is "application/json".
+<br>JSON format:
+| Variable Name | Variable Type | Content |
+| --- | --- | --- |
+| Users | String Array | List of users present in the game. Should be stored as tags on the replay. |
+Expects a response of type `FNetworkReplayStartUploadingResponse`, which contains JSON data with the string "sessionId" defined. This value will be used as the **Session Name** in future communication from the HTTP Streamer. If the optional **(Session Name Override)** was provided, the returned "sessionId", and thus the ultimate <Session Name>, should be the same value. This is the game's way of demanding the session name.
+
+---
+
+<!-- From UploadHeader(): -->
+* **(Server URL)**replay/**(Session Name)**/file/replay.header?numChunks=**(Stream Chunk Index)**&time=**(Total Demo Time In MS)**
+| Element | Value | Intended Use / Notes |
+| --- | --- | --- |
+| Server URL | Replay server URL. Drawn from DefaultEngine.Ini, unmodified. | The final "/" in the URL must be included - it will not be added automatically. |
+| Session Name | Name of the replay being uploaded. | This name will match the value that the replay server provided in response to the initial upload request. |
+| Stream Chunk Index | The index of the current chunk of replay data being sent. | Stream chunks are not uniform in size, but they are strictly ordered by this value. |
+| Total Demo Time In MS | The total running time of the demo up to the end of the current stream chunk. | Tells how much replay data, in milliseconds, has been uploaded so far without having to interpret binary replay data. |
+Uses POST. "Content-Type" is "application/octet-stream".
+The content sent will be the binary data constituting the archive header, which should be recorded in a header file associated with **Session Name**.
+
+<!-- From FlushStream(): -->
+* **(Server URL)**replay/**(Session Name)**/file/stream.**(Stream Chunk Index)**?numChunks=**(Total Chunks)**&time=**(Total Demo Time In MS)**&mTime1=**(Start Time In MS)**&mTime2=**(End Time In MS)**&absSize=**(Total Uploaded Bytes)**
+| Element | Value | Intended Use / Notes |
+| --- | --- | --- |
+| Server URL | Replay server URL. Drawn from DefaultEngine.Ini, unmodified. | The final "/" in the URL must be included - it will not be added automatically. |
+| Session Name | Name of the replay being uploaded. | This name will match the value that the replay server provided in response to the initial upload request. |
+| Stream Chunk Index | The index of the current chunk of replay data being sent. | Stream chunks are not uniform in size, but they are strictly ordered by this value. |
+| Total Chunks | The total number of chunks sent so far. | This is always one greater than **Stream Chunk Index.** |
+| Total Demo Time In MS | The total running time of the demo up to the end of the current stream chunk. | Tells how much replay data, in seconds, has been uploaded so far without having to interpret binary replay data. |
+| Start Time In MS | The total running time as of the point where this stream chunk begins. | Time is in milliseconds. For streaming replays, can be used to identify which chunk a scrubbing viewer needs. |
+| End Time In MS | The total running time as of the point where this stream chunk ends. | Time is in milliseconds. For streaming replays, can be used to identify which chunk a scrubbing viewer needs. |
+| Total Uploaded Bytes | The total size of the demo up to the end of the current stream chunk. | Tells how much replay data, in bytes, has been uploaded so far. |
+Uses POST. "Content-Type" is "application/octet-stream".
+The content is binary replay data, which can be compressed. This data represents the chunk indicated by **Stream Chunk Index**. Since chunks can vary in data size as well as the length of time they cover, it is recommended that each chunk be stored in a separate file. The HTTP Streamer will attempt to flush its current stream data via this system every ten seconds. The interval between flushes can be adjusted by changing the console variable "httpReplay.ChunkUploadDelayInSeconds"
+[REGION:note]
+To enable compression, three functions must be overridden in a child class of `FHttpNetworkReplayStreamer`:
+`SupportsCompression` must return `true`.
+`CompressBuffer` must perform compression and return `true`, except in the case of an error.
+`DecompressBuffer` must perform decompression and return `true`, except in the case of an error.
+[/REGION:note]
+
+---
+
+<!-- From AddOrUpdateEvent(): -->
+* **(Server URL)**replay/**(Session Name)**/event?group=**(Group Name)**&time1=**(Event Time In MS)**&time2=**(Event Time In MS)**&meta=**(Meta Tag)**&incrementSize=false
+* **(Server URL)**replay/**(Session Name)**/event/**(Session Name**)`_`**(Event Name)**?group=**(Group Name)**&time1=**(Event Time In MS)**&time2=**(Event Time In MS)**&meta=**(Meta Tag)**&incrementSize=false
+| Element | Value | Intended Use / Notes |
+| --- | --- | --- |
+| Server URL | Replay server URL. Drawn from DefaultEngine.Ini, unmodified. | The final "/" in the URL must be included - it will not be added automatically. |
+| Session Name | Name of the replay being uploaded. | This name will match the value that the replay server provided in response to the initial upload request. |
+| Group Name | The name of the event group containing this event. | Events can be put into groups for filtering lists, such as grouping all penalty calls in a sports title. |
+| Event Name | The GUID of the event being added or updated. | (Optional.) Information about specific, named events may be requested by viewers. If not specified, the event can still be found by listing the events in its group. |
+| Event Time In MS | The time in the replay when this event takes place. | Time is in milliseconds. Although "time1" and "time2" can be used as a begin time and an end time respectively, they are always the same value in the default HTTP Streamer. |
+| Meta Tag | Meta tag associated with this event. | Replay lists can be filtered by meta tag, if the request specifies any. **Meta Tag** will be URL-encoded by the HTTP Streamer during transmission. |
+Uses POST. "Content-Type" is "application/octet-stream" and contains binary event data to be recorded. This data can be sent through DemoNetDriver and is be user-defined. The second version of this HTTP request can be used to modify or update an existing event, whereas the first version is used to create a new event.
+[REGION:note]
+Note that events are meant to be stored separately from replays, which is why event names use GUIDs. Each event knows what replay it belongs to, which is why events are not embedded in the replays themselves.
+[/REGION:note]
+
+---
+
+<!-- From StopUploading(): -->
+* **(Server URL)**replay/**(Session Name)**/stopUploading?numChunks=**(Total Num Chunks)**&time=**(Total Time In MS)**&absSize=**(Total Bytes Uploaded)**
+| Element | Value | Intended Use / Notes |
+| --- | --- | --- |
+| Server URL | Replay server URL. Drawn from DefaultEngine.Ini, unmodified. | The final "/" in the URL must be included - it will not be added automatically. |
+| Session Name | Name of the replay being uploaded. | This name will match the value that the replay server provided in response to the initial upload request. |
+| Total Num Chunks | The number of chunks in the finished replay. | |
+| Total Time In MS | The total running time, in milliseconds, of the replay. | |
+| Total Bytes Uploaded | The total size, in bytes, of the replay. | |
+Uses POST. "Content-Type" is "application/octet-stream", but no content is sent with this type of request. The request itself indicates that the replay is done uploading.
+
+---
+
+### Download Requests
+
+---
+
+<!-- From StartStreaming(): -->
+* **(Server URL)**replay/**(Session Name)**?shouldKeep=**(Should Keep)**
+| Element | Value | Intended Use / Notes |
+| --- | --- | --- |
+| Server URL | Replay server URL. Drawn from DefaultEngine.Ini, unmodified. | The final "/" in the URL must be included - it will not be added automatically. |
+| Session Name | Name of the replay being uploaded. | This name will match the value that the replay server provided in response to the initial upload request. |
+| Should Keep | "true" or "false" | Indicates whether or not the replay being sent should be saved. |
+Uses POST. "Content-Type" is "application/json" but is empty. This request is intended for debug purposes, and is not called in the existing code.
+
+---
+
+<!-- From StartStreaming(): -->
+* **(Server URL)**replay/**(Session Name)**/startDownloading?user=**(User Name)**
+| Element | Value | Intended Use / Notes |
+| --- | --- | --- |
+| Server URL | Replay server URL. Drawn from DefaultEngine.Ini, unmodified. | The final "/" in the URL must be included - it will not be added automatically. |
+| Session Name | Name of the replay being downloaded. | This name will match the value that the server provided when the replay was originally uploaded. |
+| User Name | Unique Net Id (GUID) | The preferred network GUID that is logged into the client machine. This value comes from the engine and is not intended to be user-supplied. |
+Uses POST. Does not send content. The purpose of this request is to start downloading a replay from the replay server by establishing a connection based on the user name. Immediately after this request goes through successfully, the HTTP Streamer will send a request to download the replay's header, and then another to enumerate all checkpoints in the replay. The expected response to this request is of type `FNetworkReplayStartDownloadingResponse`, which contains JSON data.
+<br>JSON format:
+| Variable Name | Variable Type | Content |
+| --- | --- | --- |
+| State | String | The word "Live", to indicate that the game being viewed is happening live. Otherwise, the replay is assumed to be pre-recorded. |
+| NumChunks | Integer | This is expected to be nonzero, and can be updated as streaming continues (used for live replays). |
+| Time | Integer | The total time (in MS) of the replay. Subject to change if **State** indicated that this is a live stream. |
+| ViewerId | String | A unique, name assigned by the replay server that the game server will use when requesting further chunks of replay data. Often referred to as **Viewer Name** in this document. |
+
+---
+
+<!-- From DownloadHeader(): -->
+* **(Server URL)**replay/**(Session Name)**/file/replay.header
+| Element | Value | Intended Use / Notes |
+| --- | --- | --- |
+| Server URL | Replay server URL. Drawn from DefaultEngine.Ini, unmodified. | The final "/" in the URL must be included - it will not be added automatically. |
+| Session Name | Name of the replay being downloaded. | This name will match the value that the replay server provided when the replay was originally uploaded. |
+Uses GET. Expects a response with the replay's header in the content portion, as binary data.
+
+---
+
+<!-- From EnumerateCheckpoints(): -->
+* **(Server URL)**replay/**(Session Name)**/event?group=checkpoint
+| Element | Value | Intended Use / Notes |
+| --- | --- | --- |
+| Server URL | Replay server URL. Drawn from DefaultEngine.Ini, unmodified. | The final "/" in the URL must be included - it will not be added automatically. |
+| Session Name | Name of the replay being downloaded. | This name will match the value that the replay server provided when the replay was originally uploaded. |
+Uses GET. Expects a response of type `FHttpResponsePtr` containing checkpoint data as a JSON-formatted string.
+<br>JSON format:
+| Variable Name | Variable Type | Content |
+| --- | --- | --- |
+| events | Array | Contains array entries which consist of all the other fields in this table. |
+| id | String | The name (a GUID) of the checkpoint. |
+| group | String | The group to which the event belongs, if any. |
+| meta | String | Meta Tag string associated with this event, if any. |
+| time1 | Integer | Start time (in MS) of the event. |
+| time2 | Integer | End time (in MS) of the event. Often the same as the start time. |
+[REGION:note]
+Note that this is technically a generic request for a group of events. While "checkpoint" is explicitly specified here, the replay server should be able to handle other values for "group" and report all events for any given group.
+[/REGION:note]
+
+---
+
+<!-- From ConditionallyDownloadNextChunk(): -->
+* **(Server URL)**replay/**(Session Name)**/file/stream.**(Stream Chunk Index)**
+| Element | Value | Intended Use / Notes |
+| --- | --- | --- |
+| Server URL | Replay server URL. Drawn from DefaultEngine.Ini, unmodified. | The final "/" in the URL must be included - it will not be added automatically. |
+| Session Name | Name of the replay being downloaded. | This name will match the value that the replay server provided when the replay was originally uploaded. |
+| Stream Chunk Index | The index of the current chunk of replay data being requested. | Stream chunks are not uniform in size, but they are strictly ordered by this value. |
+Uses GET. This is a request for a specific stream chunk, named by the viewer. This request will not be sent while there is a chunk already being downloaded. The expected response has a header with some information about the chunk, and its content is a binary chunk of replay data. The binary data can be compressed, if compression is supported. The following fields are expected in the header:
+| Header Field Name | Description |
+| --- | --- |
+| NumChunks | The total number of chunks in the stream. This can be useful with streamed games, as it lets the user know about new chunks that have been created while viewing. |
+| Time | The total length of the replay (in MS). As with NumChunks, this value can change during the viewing session if the replay is live. |
+| State | This variable has one recognized value, "Live". If this variable is missing or is a different value, the streamer will switch to non-live streaming mode. |
+| MTime1 | The start time of the stream, in MS. Can be non-zero for live replay streams. This variable will only be checked if a chunk of data sent, and either this is the first chunk the streamer has downloaded, or the stream is live. |
+| MTime2 | The end time of the stream, in MS. Can change as more chunks come in. This variable will always be checked if a chunk of data is sent, or if the replay is live. |
+[REGION:note]
+To enable compression, three functions must be overridden in a child class of `FHttpNetworkReplayStreamer`:
+`SupportsCompression` must return `true`.
+`CompressBuffer` must perform compression and return `true`, except in the case of an error.
+`DecompressBuffer` must perform decompression and return `true`, except in the case of an error.
+[/REGION:note]
+
+---
+
+<!-- From EnumerateEvents(): -->
+* **(Server URL)**replay/**(Session Name)**/event?group=**(Group Name)**"
+Uses GET. Can be called directly from DemoNetDriver, meaning your game project can do this at any time. This request asks the replay server to send a list of all events pertaining to a specific replay (identified by **Session Name**) and filtered by a certain event group. It is recommended that empty group names, as well as a lack of the "group" argument, be supported. Content received in response to this request is expected to be a JSON-formatted string corresponding to `FReplayEventList` and `FReplayEventListItem`.
+<br>JSON format:
+| Variable Name | Variable Type | Content |
+| --- | --- | --- |
+| events | Array | Contains elements which consist of all the other fields in this table. |
+| id | String | The name (a GUID) of the checkpoint. |
+| group | String | The group to which the event belongs, if any. |
+| meta | String | Meta Tag string associated with this event, if any. |
+| time1 | Integer | Start time (in MS) of the event. |
+| time2 | Integer | End time (in MS) of the event. Often the same as the start time. |
+
+---
+
+<!-- From RefreshViewer(): -->
+* **(Server URL)**replay/**(Session Name)**/viewer/**(Viewer Name)**
+* **(Server URL)**replay/**(Session Name)**/viewer/**(Viewer Name)**?final=**(Final Heartbeat)**
+| Element | Value | Intended Use / Notes |
+| --- | --- | --- |
+| Server URL | Replay server URL. Drawn from DefaultEngine.Ini, unmodified. | The final "/" in the URL must be included - it will not be added automatically. |
+| Final Heartbeat | "true" or "false" | By default, "final" will not be sent unless **Final Heartbeat** is true. If true, the viewer has stopped watching and no further data should be sent. |
+Uses POST. Does not send content.
+This request is a heartbeat, sent every ten seconds. This frequency is currently hard-coded, but will likely be made user-configurable without code changes in the future. If a viewer fails to send this request for a long period of time (perhaps one minute), then the viewer has most likely disconnected unexpectedly and it would be reasonable to disconnect them. If the "final=true" version is sent, the user is indicating that this is the final heartbeat, and the user has stopped watching the replay.
+
+---
+
+### Search Requests
+
+---
+
+<!-- From EnumerateStreams(): -->
+* **(Server URL)**replay?app=**(Replay Version AppString)**&cl=**(Replay CL)**&version=**(Replay Network Version)**&meta=**(Meta Tag)**&user=**(User Name)**&recent=**(In Recent Viewer)**"
+| Element | Value | Intended Use / Notes |
+| --- | --- | --- |
+| Server URL | Replay server URL. Drawn from DefaultEngine.Ini, unmodified. | The final "/" in the URL must be included - it will not be added automatically. |
+| Replay Version AppString | A user-defined string that describes the game being played. | This value was sent when the user began uploading the replay, and is now being used for filtering the list of replays sent to the potential viewer. |
+| Replay CL | The changelist of the build being played. | This value is used for filtering lists of replays. |
+| Replay Network Version | A user-defined string that describes the version of the game being played. | This value should reflect changes in your game's data formats as they affect replay. It is also used for filtering lists of replays. |
+| Meta | Tags applied to this replay. | User-applied meta tags can be used for filtering games that have certain features. This could be anything from "all Capture The Flag matches" to to "games that went into double-overtime". |
+| User | Unique Net Id (GUID) | This is used to filter games based on their inclusion of a specific player. |
+| Recent | "true" or "false" | If true, only recently-uploaded replays should be shown. What defines "recent" is up to the replay server. |
+Uses GET. Expects a response with content in the form of a JSON-formatted string, containing a list of replays that the replay server can send. These replays are expected to be pre-filtered by the user's criteria in the URL. The replay server should be able to handle some of the filter fields being empty or missing.
+<br>JSON format:
+| Variable Name | Variable Type | Content |
+| --- | --- | --- |
+| replays | Array | Contains elements which consist of all the other fields in this table. |
+| AppName | String | The **Replay Version AppString** that was sent when the replay was uploaded. Should match the parameter you sent with the request. |
+| SessionName | String | The name (GUID) of the replay. |
+| FriendlyName | String | The **(Platform Friendly Name)** that was provided when the replay was uploaded. |
+| Timestamp | DateTime | The time and date associated with the replay (i.e. when this game happened). |
+| SizeInBytes | Integer | Total size, in bytes, of the replay. |
+| DemoTimeInMs | Integer | Running length, in MS, of the replay. |
+| NumViewers | Integer | Number of users currently watching this replay. |
+| bIsLive | Boolean | True if the game in the replay is still being played. |
+| Changelist | Integer | The changelist (CL) of the engine running the game when this replay was recorded. Corresponds to the **Replay CL** argument sent when uploading the replay. Should match the parameter you sent with the request. |
+| shouldKeep | Boolean | True if the debug "shouldKeep" request was made. |
+
+---
+
+<!-- From SearchEvents(): -->
+* **(ServerURL)**event?group=**(Event Group)**
+| Element | Value | Intended Use / Notes |
+| --- | --- | --- |
+| Server URL | Replay server URL. Drawn from DefaultEngine.Ini, unmodified. | The final "/" in the URL must be included - it will not be added automatically. |
+| Group Name | The name of the event group containing this event. | Events can be put into groups for filtering lists, such as grouping all penalty calls in a sports title. |
+Uses GET. Expects a response with content as a JSON-formatted string containing a list of replays from the server. These replays are expected to be pre-filtered to include only those that contain at least one event in the named **Event Group**.
+<br>JSON format:
+| Variable Name | Variable Type | Content |
+| --- | --- | --- |
+| replays | Array | Contains elements which consist of all the other fields in this table. |
+| AppName | String | The **Replay Version AppString** that was sent when the replay was uploaded. |
+| SessionName | String | The name (GUID) of the replay. |
+| FriendlyName | String | The **(Platform Friendly Name)** that was provided when the replay was uploaded. |
+| Timestamp | DateTime | The time and date associated with the replay (i.e. when this game happened). |
+| SizeInBytes | Integer | Total size, in bytes, of the replay. |
+| DemoTimeInMs | Integer | Running length, in MS, of the replay. |
+| NumViewers | Integer | Number of users currently watching this replay. |
+| bIsLive | Boolean | True if the game in the replay is still being played. |
+| Changelist | Integer | The changelist (CL) of the engine running the game when this replay was recorded. Corresponds to the **Replay CL** argument sent when uploading the replay. |
+| shouldKeep | Boolean | True if the debug "shouldKeep" request was made. |
+
+---
+
+<!-- From RequestEventData(): -->
+* **(Server URL)**event/**(Event Name)**
+| Element | Value | Intended Use / Notes |
+| --- | --- | --- |
+| Server URL | Replay server URL. Drawn from DefaultEngine.Ini, unmodified. | The final "/" in the URL must be included - it will not be added automatically. |
+| Event Name | The GUID of the event requested. | Since events are identified by GUID, the event can be from any replay on file. Make sure that the event you request is from a list of events pertaining to the replay you are currently watching. |
+Uses GET. Expects a response with Content as a binary stream from the replay server. This is called directly from the DemoNetDriver, so it can be done from the game project at any time, but will usually be called while viewing a replay with an **Event Name** retrieved from a list provided by a previous request.
+
+---
diff --git a/Engine/Documentation/Source/Engine/Replay/Streamers/Streamers.INT.udn b/Engine/Documentation/Source/Engine/Replay/Streamers/Streamers.INT.udn
new file mode 100644
--- /dev/null
+++ b/Engine/Documentation/Source/Engine/Replay/Streamers/Streamers.INT.udn
@@ -0,0 +1,61 @@
+Availability:Docs
+Title: Replay System
+Crumbs: %ROOT%, Engine, Engine/Replay
+Description:Overview of the Replay system for recording and playback of gameplay
+Version: 4.14
+parent:Engine/Replay
+order:17
+topic-image:replay_topic.png
+tags:Replays
+
+[TOC (start:2 end:3)]
+
+## DemoNetDriver Features
+* DemoNetDrivers use the NULL Streamer by default, but can override this by receiving the URL option "ReplayStreamerOverride" set to the name of a different Streamer Factory's module, such as "InMemoryNetworkReplayStreaming" or "HttpNetworkReplayStreaming". The default value, "NullNetworkReplayStreaming", can be changed by setting the "DefaultFactoryName" variable in the "NetworkReplayStreaming" section of your project's `DefaultEngine.Ini` file. This can also be accomplished by calling `InitBase` and providing the appropriate URL argument as a parameter.
+* DemoNetDrivers can amortize the time cost of recording replay data by setting the "demo.CheckpointSaveMaxMSPerFrame" CVar to a positive value. `Actor`s that are not recorded into the replay before the per-frame time limit expires will be queued to record on the next frame. The advantage to this feature is that it caps the amount of time that checkpoint recording and take and helps to keep your game free of hitches. The drawback is that slight visual errors can appear during playback due to checkpoints containing data on Actors taken from different frames. This feature will only be activated in the case that your game is taking longer than your specified time limit to record checkpoints, which means that it will primarily apply to lower-end machines, or to more performance-intensive games.
+* If `bPrioritizeActors` is set to true, Actors being saved into a replay will be pre-sorted by priority for recording order, based on the virtual function `GetReplayPriority`. This is useful when combined with amortized recording via `MaxDesiredRecordTimeMS`.
+* Checkpoint recording frequency can be adjusted by changing the CVar `demo.CVarCheckpointUploadDelayInSeconds`. The default is 30 seconds. Increasing the length of time between checkpoints will make scrubbing backward or skipping around in replays slower, but will decrease the size of replays.
+* The variable `bPauseRecording` can be set to true while recording a demo to suspend recording temporarily. Setting it back to false will resume the recording.
+* The **Game Mode** will use a different **Player Controller** class (designated as `ReplaySpectatorPlayerControllerClass`) when viewing replays.
+* Using `SetViewerOverride`, a `DemoNetDriver` can change how Actors' network relevancy, culling, and prioritization are determined by creating an alternate **Player Controller** that will be used for recording purposes. This is especially useful for games with large maps, like Paragon, where the player isn't always aware of things happening far away during live play (both for efficiency and cheat-prevention reasons), but will expect to see everything when viewing a replay.
+* DemoNetDrivers can operate in parallel with Slate. To accomplish this, both "tick.DoAsyncEndOfFrameTasks" and "demo.ClientRecordAsyncEndOfFrame" CVars must be non-zero.
+
+[REGION:note]
+Replay-generated Actors will make function calls just like live-gameplay Actors do. This causes them to behave just like live actors with minimal replay data, but it also means that function calls that affect shared objects such as the `GameInstance`, `GameState`, or `GameMode`, will still be usable by replay Actors, and can affect the state of the game in unintended ways. This is especially true in the case of the Memory Streamer, which can view a replay while live gameplay is still running. To guard against Actors affecting things that they shouldn't, it is recommended that these operations first be checked to see if the Actor is part of a live or replay level, and behave accordingly. This is most likely a game-specific issue and must be handled by each project on a case-by-case basis; for example, a given game may wish to update the player's health bar or full-screen damage overlay during a replay, but not alter the player's score.
+[/REGION:note]
+
+## Replay Data Format
+In terms of data, a replay contains three types of game-state information, as well as some additional text data. At the start is baseline data describing the starting state of the game world. Checkpoints that act as snapshots of the net changes to the world (from the baseline) appear at regular, user-defined intervals. The space between checkpoints is then filled with incremental changes to individual objects in the world. Any moment in the game can be reconstructed by the engine quickly and accurately by initializing the world to the starting state, making the changes described in the last checkpoint before the chosen time, and then applying each incremental change after that checkpoint leading up to the desired point in time. The text data contained in a replay consists of a display name, which can be used when making a player-facing list, and user-defined text tags (HTTP Streamer only), which can be used when searching or filtering through lists of games.
+
+## The NULL Streamer
+The NULL Streamer is the default way to record a replay. This streamer writes replay data to disk, e.g. a local hard drive. This is good for making local recordings, especially of single-player games. These recordings can be very helpful for a variety of uses, such as producing gameplay trailers or in-game cutscenes, or enabling your users to view and share speedrun or tutorial videos within your game. Replay data files are saved to the Saved/Demos/(Replay Name) folder in your project.
+
+## The Memory Streamer
+The Memory Streamer keeps a user-configurable length of replay data (e.g. 10 seconds, 3 minutes, etc.) in memory on the local machine. This type of stream is best suited to instant replays of recent, dramatic moments, such as scores in a sports game, deaths in a shooter, or the final moments of a boss battle in an action game.
+
+### Memory Streamer Usage Details
+The Memory Streamer is special in that it is meant to record, play back, and resume gameplay during a single session. The live game is able to continue, invisibly and silently, while the player is watching a replay, so that the game can be resumed seamlessly the moment the replay ends. At load time, the engine collects levels into three groups: Static Levels, Dynamic Source Levels, and Dynamic Duplicated Levels. These groups determine how the level will interact with live gameplay and the replay system, as follows:
+
+| Level collection | Levels added to this collection | Behavior |
+| -------- | ----- | --- |
+| Static Levels | Levels that are not the persistent level and are marked with the `IsStatic` boolean. | Should be unaffected by gameplay, and will be shown during both live play and replays. |
+| Dynamic Source Levels | The persistent level and any sublevels with the `IsStatic` boolean variable set to false. | Affected by live gameplay. Hidden during replays, but with gameplay still running as normal. |
+| Dynamic Duplicate Levels | Copies made from Dynamic Source Levels at load time. Do not exist on dedicated servers or in editor mode. | Hidden during live gameplay. Replays take place in these levels, then they are emptied out. |
+
+The "Is Static" setting can be found by activating the Levels view from the Window menu, and then clicking the Level Details button with your sublevel highlighted.
+![](Sublevel.png)
+From there, the "Is Static" option will be available in the Level Details window. This window will be empty for the persistent level, so make sure to select a sublevel.
+![](IsStatic.png)
+
+[REGION:note]
+These level collections are a part of the level-streaming process and are not specifically related to the Replay system.
+[/REGION:note]
+
+With this system in mind, we can create one `DemoNetDriver` for the Dynamic Source Levels, and another for the Dynamic Duplicate Levels. This enables us to record live gameplay in the Dynamic Source Levels and then play that data back in the Dynamic Duplicate Levels. By hiding the Dynamic Source Levels and showing the Dynamic Duplicated Levels during replay, the game can continue playing and receiving network updates unaffected by the replay. The third group, Static Levels, can be active and visible at any time; they should contain things like static world geometry or ambient background sounds, particles, and animations that are not affected by live gameplay, and therefore do not need to be involved in the replay process. When the replay is finished, the contents of the Dynamic Duplicate Levels will be destroyed will undergo garbage collection, and the Dynamic Source Levels will be made visible/audible again. Since the Dynamic Source Levels were never destroyed or suspended, only hidden, the game will have naturally progressed while the replay was being viewed, and can be shown immediately and without hitches. Additionally, this system enables increased efficiency by giving the developers the ability to mark levels as static in order to exclude them from replay recording and playback, saving both memory and time.
+
+## The HTTP Streamer
+The HTTP Streamer sends replay data to another server, which could be on a LAN or elsewhere on the internet. This is useful for live-streaming matches, or for keeping recordings of matches that can be viewed at any time. This streamer is especially useful for dedicated-server games, where only the server knows everything that's happening everywhere in the game at all times, and where offloading the work of processing replay data will increase the number of simultaneous games that can be hosted on a single server. It can also help to serve as a moderation or cheat detection tool, since the data can be captured from computers completely controlled by the party running the game.
+
+### HTTP Streamer Usage Details
+
+The HTTP Streamer will communicate with a custom-written replay server through the [](Engine/Replay/Streamers/HTTPRESTAPI), using GET and POST methods to send data as binary or JSON-formatted strings. In order to set up your own replay server, you must first establish your URL. Your project's DefaultEngine.Ini file contains the replay server's URL under the `[HttpNetworkReplayStreaming]` section, in the variable `ServerURL`. Your `ServerURL` should be in the format "http://replay.yourgame.com/". The final "/" is important, as the HTTP Streamer will not assume that it should modify the format of the URL.
diff --git a/Engine/Documentation/Source/Engine/Sequencer/HowTo/SequenceRecorder/SequenceRecorder.INT.udn b/Engine/Documentation/Source/Engine/Sequencer/HowTo/SequenceRecorder/SequenceRecorder.INT.udn
--- a/Engine/Documentation/Source/Engine/Sequencer/HowTo/SequenceRecorder/SequenceRecorder.INT.udn
+++ b/Engine/Documentation/Source/Engine/Sequencer/HowTo/SequenceRecorder/SequenceRecorder.INT.udn
@@ -27,10 +27,6 @@ For this how-to, we are using the **Blueprint Third Person Template** project wi
 	![](RecorderStep1.png)
 
 	The Sequence Recorder window will automatically open. There are some options under **Sequence Recording** which will determine how (and where) the new Level Sequence asset will be saved. You can choose to record Actors that are spawned (such as particle effects, other characters, etc.), determine the naming convention and save location, and change other settings like the duration to record or delay before starting to record. 
-
-	[REGION:note]
-	See the [](Engine\Sequencer\Features\SequenceRecorder) page for more information about each option.
-	[/REGION]
 	
 1.	From the Main Editor Toolbar, click the **Play** button to start a Play in Editor session. 
 
diff --git a/Engine/Documentation/Source/Engine/Sequencer/Sequencer.INT.udn b/Engine/Documentation/Source/Engine/Sequencer/Sequencer.INT.udn
--- a/Engine/Documentation/Source/Engine/Sequencer/Sequencer.INT.udn
+++ b/Engine/Documentation/Source/Engine/Sequencer/Sequencer.INT.udn
@@ -4,7 +4,7 @@ Crumbs: %ROOT%, Engine
 Description:Sequencer is Unreal Engine 4's mult-track editor used for creating and previewing cinematic sequences in real time. 
 Navigation:topic
 parent:Engine
-order:3
+order:4
 social-image:sequencer_social.png
 related: Engine/Matinee
 
diff --git a/Engine/Documentation/Source/Engine/Sequencer/Workflow/CustomRenderPass/CustomRenderPass.INT.udn b/Engine/Documentation/Source/Engine/Sequencer/Workflow/CustomRenderPass/CustomRenderPass.INT.udn
--- a/Engine/Documentation/Source/Engine/Sequencer/Workflow/CustomRenderPass/CustomRenderPass.INT.udn
+++ b/Engine/Documentation/Source/Engine/Sequencer/Workflow/CustomRenderPass/CustomRenderPass.INT.udn
@@ -1,4 +1,4 @@
-Availability:Docs
+Availability:Public
 Title: Exporting Custom Render Passes
 Crumbs: %ROOT%, Engine
 Description:Describes how you can export your Sequence in Custom Render Passes.
@@ -11,7 +11,7 @@ Related: Engine/Sequencer/HowTo
 
 This page describes how you can render out custom passes from the **Render Movie Settings** window inside **Sequencer**. 
 Using a Custom Render Pass will output your sequence into a series of image files based on the Render Pass you specify. 
-By default the resulting files will be exported into the .png file format, however you can capture frames in HDR to export them in the .exr file format as well. 
+The image files will be exported in the .png file format by default, however you can capture frames in HDR to export them in the .exr file format as well. 
 
 [TOC(start:2 end:2)]
 
@@ -27,20 +27,19 @@ To render out custom passes for your cinematic refer to the following steps.
 
 	![](RenderMovie2.png)
 
-1. Select the desired Render Pass(es) from the **Include Render Passes** drop-down.
+1. Select the desired Render Pass(es) from the **Include Render Passes** drop-down. You can include more than one pass as part of the export process. 
 
 	![](RenderMovie3.png)
 
-	You can include more than one pass as part of the export process. 
+	[REGION:warning]
+	If you do not specify which render pass(es) to include, all available passes will be included as part of the export. 
+	Depending upon the number of frames in your sequence, this may result in hundreds of files being generated for each available render pass.
+	[/REGION]
 
+1.	(Optional) Click the **Minus** button next to a render pass to remove it from the export.
+	
 	![](RenderMovie4.png)
 
-	Click the **Minus** button next to each render pass to remove it. 
-
-	[REGION:note]
-	If you do not specify which render pass(es) to include, all available passes will be included as part of the export. Depending upon the number of frames in your sequence, this may result in hundreds of files being generated for each available render pass.
-	[/REGION]
-
 1. Specify the **Output Directory** and **Filename Format**, then click **Capture Movie**. 
 
 	![](RenderMovie5.png)
@@ -48,8 +47,13 @@ To render out custom passes for your cinematic refer to the following steps.
 	If you use the default **Filename Format**, your files will automatically be named based on the render pass it reflects.
 
 Once the render movie process completes, in your output directory there will be an image file for each frame of your sequence that corresponds to the selected render pass(es). 
-If you selected more than one render pass, you will have a separate set of image files for each pass. 
+If you selected more than one render pass, you will have a separate set of image files for each pass. Below is an example of an render pass for Scene Depth next to the original shot from the full scene. 
 
+[REGION:imagetable]
+| ![](FullScene.png)(w:480) |![](DepthScene.png)(w:580)|
+| ----------- |------|
+| **Full Scene** | **Exported Scene Depth** | 
+[/REGION]
 
 ## HDR Exporting
 
@@ -59,7 +63,7 @@ The choice of Capture Gamut will largely depend on the project that will make us
 
 | Gamut Type | Description |
 | ---- | ----------- |
-| [**Rec.709 / sRGB**](https://en.wikipedia.org/wiki/Rec._709#Primary_chromaticities) | This is the most widely used color gamut for images that will be displayed on monitors, TVs and Projectors. |
+| [**Rec.709 / sRGB**](https://en.wikipedia.org/wiki/Rec._709#Primary_chromaticities) | This is the most widely used color gamut for images that will be displayed on monitors, TVs, and Projectors. |
 | [**P3 D65**](https://en.wikipedia.org/wiki/DCI-P3#System_colorimetry) | This is used widely for cinematic projection. |
 | [**Rec.2020**](https://en.wikipedia.org/wiki/Rec._2020#System_colorimetry), [**ACES**](http://www.oscars.org/science-technology/aces/aces-documentation), [**ACEScg**](http://duikerresearch.com/2015/09/acescg-a-common-color-encoding-for-visual-effects-applications/) | These gamuts are being used more with the advent of UHD TVs, HDR TVs and other advances in display technologies. ACEScg in particular now plays a prominent role in rendering and compositing for visual effects and is probably the most future-proof of gamuts.  |
 
diff --git a/Engine/Documentation/Source/Engine/Sequencer/Workflow/ImportExport/ImportExportFBX.INT.udn b/Engine/Documentation/Source/Engine/Sequencer/Workflow/ImportExport/ImportExportFBX.INT.udn
--- a/Engine/Documentation/Source/Engine/Sequencer/Workflow/ImportExport/ImportExportFBX.INT.udn
+++ b/Engine/Documentation/Source/Engine/Sequencer/Workflow/ImportExport/ImportExportFBX.INT.udn
@@ -1,4 +1,4 @@
-Availability:Docs
+Availability:Public
 Title: Exporting and Importing FBX files
 Crumbs: %ROOT%, Engine
 Description:Describes how you can export FBX files from Sequencer and import FBX files to Sequencer. 
@@ -32,7 +32,7 @@ Below, we export a scene from the **Sequencer** Showcase (pictured below) availa
 
 First, determine which Level Sequence to export. 
 
-1.	In the clip below, we pick the shot we want to alter and open it up. We then go into the **General Options** and choose **Export**. 
+1.	In the clip below, we pick the shot we want to alter and open it up. We then go into the **General Options** and choose **Export**. We have a Master Track that includes **shot0040_001** so we open that particular Level Sequence to export it. 
 
 	[OBJECT:EmbeddedVideo]
 	[PARAMLITERAL:width]
@@ -45,8 +45,6 @@ First, determine which Level Sequence to export.
 	SM77fHV_wu8
 	[/PARAMLITERAL]
 	[/OBJECT]
-
-	We have a Master Track that includes **shot0040_001** so we open that particular Level Sequence to export it. 
 	
 ### External Editing
 
@@ -66,7 +64,7 @@ Once we have an exported FBX from Sequencer, we open our 3rd party application t
 	[/PARAMLITERAL]
 	[/OBJECT]
 
-1.	We then import our FBX into Maya and set up our scene so that we can begin making edits. 
+1.	We then import our FBX into Maya and set up our scene so that we can begin making edits. Once our scene is imported, we set our Perspective under **Panel** to use the camera **shotCamA** in our scene and hide any elements obstructing the view (such as the godray mesh) with it selected and **Ctrl + H**. Lastly under **View** we **Select Camera** to display the properties and keyframes for our **shotCamA** which we can then begin to edit.
 
 	[OBJECT:EmbeddedVideo]
 	[PARAMLITERAL:width]
@@ -78,12 +76,9 @@ Once we have an exported FBX from Sequencer, we open our 3rd party application t
 	[PARAMLITERAL:videoid]
 	0Nc4S-T1t1I
 	[/PARAMLITERAL]
-	[/OBJECT]
+	[/OBJECT] 
 
-	Once our scene is imported, we set our Perspective under **Panel** to use the camera **shotCamA** in our scene and hide any elements obstructing the view (such as the godray mesh) with it selected and **Ctrl + H**. 
-	Lastly under **View** we **Select Camera** to display the properties and keyframes for our **shotCamA** which we can then begin to edit. 
-
-1.	Next we make edits to the movement of our camera. 
+1.	Next we make edits to the movement of our camera. Above we select the Translation and Rotation properties of our **shotCamA** and **Right-click** and **Break Connections**. We then pick a new starting point for the camera to be in to start our scene and press **S** to keyframe the new position. We scrub ahead slightly and move the camera to the ending position and press **S** again to keyframe the end position so that our camera moves along a new path.
 
 	[OBJECT:EmbeddedVideo]
 	[PARAMLITERAL:width]
@@ -95,13 +90,9 @@ Once we have an exported FBX from Sequencer, we open our 3rd party application t
 	[PARAMLITERAL:videoid]
 	3asklRbBagA
 	[/PARAMLITERAL]
-	[/OBJECT]
+	[/OBJECT] 
 
-	Above we select the Translation and Rotation properties of our **shotCamA** and **Right-click** and **Break Connections**. 
-	We then pick a new starting point for the camera to be in to start our scene and press **S** to keyframe the new position. 
-	We scrub ahead slightly and move the camera to the ending position and press **S** again to keyframe the end position so that our camera moves along a new path. 
-
-1.	With our changes in place, we then export our scene using **Export Selection**.
+1.	With our changes in place, we then export our scene using **Export Selection**. During the export process from Maya, we first change the export type to **FBX export**.Then, under **Advanced Options** and **Axis Conversion**, we change the **Up Axis** to follow **Z** to allow for a compatible import into Unreal Engine. Future versions of the engine may address this so that you do not have to convert the axis, but you will need to do this for now to avoid any importing problems.
 
 	[OBJECT:EmbeddedVideo]
 	[PARAMLITERAL:width]
@@ -113,18 +104,14 @@ Once we have an exported FBX from Sequencer, we open our 3rd party application t
 	[PARAMLITERAL:videoid]
 	sj-gWprQFNk
 	[/PARAMLITERAL]
-	[/OBJECT]
-
-	During the export process from Maya, we first change the export type to **FBX export**.
-	Then, under **Advanced Options** and **Axis Conversion**, we change the **Up Axis** to follow **Z** to allow for a compatible import into Unreal Engine. 
-	Future versions of the engine may address this so that you do not have to convert the axis, but you will need to do this for now to avoid any importing problems. 
+	[/OBJECT] 
 
 
 ### Importing to Sequencer
 
 Now that we have an edited version in FBX format we can import it and the changes into Unreal Engine, updating our scene.
 
-1.	Inside our **shot0040_001** that we wanted to edit, we **right-click** on our **shotCamA** and select **Import**. 
+1.	Inside our **shot0040_001** that we wanted to edit, we **right-click** on our **shotCamA** and select **Import**. This will take the transform data inside the FBX and apply it to the selected object in Unreal Engine. This can be useful if you want to apply the same transform data to multiple objects inside your scene. 
 
 	[OBJECT:EmbeddedVideo]
 	[PARAMLITERAL:width]
@@ -138,12 +125,13 @@ Now that we have an edited version in FBX format we can import it and the change
 	[/PARAMLITERAL]
 	[/OBJECT]
 
-	This will take the transform data inside the FBX and apply it to the selected object in Unreal Engine. 
-	This can be useful if you want to apply the same transform data to multiple objects inside your scene. You can also import the FBX from the **General Options** window instead of using the right-click method. 
+	[REGION:tip]
+	You can also import the FBX from the **General Options** window instead of using the right-click method. 
 
 	![](ImportOption.png)
 
 	When choosing this method, Unreal Engine will try to match the names of your objects. For example, we have **shotCamA** in both our FBX file and our Sequence so it will apply the data from the FBX file to that object inside Unreal Engine. 
+	[/REGION]
 
 In our example we only applied Transform changes to our camera. However, we could have keyframed new values for our **Focal Length** to produce a different looking shot as well.
 You can experiment with this workflow and use what works best in your given situation. There may be times when you want to export to a 3D package like Maya or 3dsMax to fine-tune aspects of your scene and then import your changes back into your scene in Unreal Engine. 
diff --git a/Engine/Documentation/Source/Engine/Sequencer/Workflow/SequenceRecorder/SequenceRecorder.INT.udn b/Engine/Documentation/Source/Engine/Sequencer/Workflow/SequenceRecorder/SequenceRecorder.INT.udn
--- a/Engine/Documentation/Source/Engine/Sequencer/Workflow/SequenceRecorder/SequenceRecorder.INT.udn
+++ b/Engine/Documentation/Source/Engine/Sequencer/Workflow/SequenceRecorder/SequenceRecorder.INT.udn
@@ -1,4 +1,4 @@
-Availability:Docs
+Availability:Public
 Title: Sequence Recorder
 Crumbs: %ROOT%, Engine
 Description:Overview of the Sequence Recorder tool and recording options to capture gameplay. 
@@ -62,7 +62,7 @@ Please see each of the tables below for a list of properties along with a descri
 | ---- | ----------- |
 | **Create Level Sequence** | Whether to create a Level Sequence when recording or not. Actors and animations will be inserted into this sequence if enabled. |
 | **Immersive Mode** | Whether to maximize the viewport when recording. |
-| **Sequence Length** | Defines the maximum length of the sequence to record. |
+| **Sequence Length** | Defines the maximum length of the sequence to record (setting to 0 recording will continue until Stop is invoked). |
 | **Recording Delay** | Delay that will be used before starting the recording. |
 | **Sequence Name** | The name of the sequence to record to. This name will also be used to auto-generate any assets created by this recording. |
 | **Sequence Recording Base Path** | Defines the location within the project to save the recording. Sub-assets will be created in subdirectories within the base path specified. |
@@ -99,6 +99,15 @@ Please see each of the tables below for a list of properties along with a descri
 
 This section contains some workflow tips & tricks you can use when working with the Sequence Recorder. 
 
+### Continuous Recording
+
+In the **Sequence Recorder** window, you can set the **Sequence Length** to **0** to perform continuous recording. 
+
+![](NoLength.png)
+
+When set to 0, recording will not stop unless you manually choose to stop recording by hitting the **Esc** key. 
+
+
 ### Right-click Recording
 
 If you have an active **Level Sequence** open, rather than clicking the **Add** button inside the Sequence Recorder, you can **Right-click** on an Actor in your level and initiate recording of the Actor from the context menu (or press **Alt + R**). 
@@ -182,3 +191,4 @@ The **Immersive Mode** option inside the Sequence Recorder window will automatic
 Once you stop recording, the viewport will return back to its normal size. 
 
 
+
diff --git a/Engine/Documentation/Source/Engine/Sequencer/Workflow/Workflow.INT.udn b/Engine/Documentation/Source/Engine/Sequencer/Workflow/Workflow.INT.udn
--- a/Engine/Documentation/Source/Engine/Sequencer/Workflow/Workflow.INT.udn
+++ b/Engine/Documentation/Source/Engine/Sequencer/Workflow/Workflow.INT.udn
@@ -23,7 +23,7 @@ If you are new to Sequencer it is recommended that you first check out the [**Ov
 	Features & Tools
 	[/PARAM]
 	[PARAM:links]
-		* [](Engine\Sequencer\Workflow\SequenceRecorder "%Engine\Sequencer\Workflow\SequenceRecorder:description%")
+		<!--* [](Engine\Sequencer\Workflow\SequenceRecorder "%Engine\Sequencer\Workflow\SequenceRecorder:description%")-->
 		* [](Engine\Sequencer\Workflow\EditorPrefs "%Engine\Sequencer\Workflow\EditorPrefs:description%")
 	[/PARAM]
 [/OBJECT]
@@ -38,7 +38,7 @@ If you are new to Sequencer it is recommended that you first check out the [**Ov
 	[PARAM:links]
 		* [](Engine\Sequencer\Workflow\CheatSheet "%Engine\Sequencer\Workflow\CheatSheet:description%")
 		* [](Engine\Sequencer\Workflow\ImportExport "%Engine\Sequencer\Workflow\ImportExport:description%")
-		* [](Engine\Sequencer\Workflow\CustomRenderPass "%Engine\Sequencer\Workflow\CustomRenderPass:description%")
+		<!--* [](Engine\Sequencer\Workflow\CustomRenderPass "%Engine\Sequencer\Workflow\CustomRenderPass:description%")-->
 		* [](Engine\Sequencer\Workflow\ToolUsage "%Engine\Sequencer\Workflow\ToolUsage:description%")
 	[/PARAM]
 [/OBJECT]
diff --git a/Engine/Documentation/Source/Engine/UI/LevelEditor/Details/LevelEditorDetails.INT.udn b/Engine/Documentation/Source/Engine/UI/LevelEditor/Details/LevelEditorDetails.INT.udn
--- a/Engine/Documentation/Source/Engine/UI/LevelEditor/Details/LevelEditorDetails.INT.udn
+++ b/Engine/Documentation/Source/Engine/UI/LevelEditor/Details/LevelEditorDetails.INT.udn
@@ -203,6 +203,28 @@ Other types use customized widgets to provide a more intuitive experience when m
 	[RELATIVE:Engine/UI/LevelEditor/Details/Properties/Array]
 	[/PARAM]
 [/OBJECT]
+[OBJECT:TopicCompact]
+	[PARAM:image]
+	![](Engine/UI/LevelEditor/Details/Properties/Map/map_topic.png)(convert:false)
+	[/PARAM]
+	[PARAM:title]
+	%Engine/UI/LevelEditor/Details/Properties/Map:title%
+	[/PARAM]
+	[PARAM:path]
+	[RELATIVE:Engine/UI/LevelEditor/Details/Properties/Map]
+	[/PARAM]
+[/OBJECT]
+[OBJECT:TopicCompact]
+	[PARAM:image]
+	![](Engine/UI/LevelEditor/Details/Properties/Set/set_topic.png)(convert:false)
+	[/PARAM]
+	[PARAM:title]
+	%Engine/UI/LevelEditor/Details/Properties/Set:title%
+	[/PARAM]
+	[PARAM:path]
+	[RELATIVE:Engine/UI/LevelEditor/Details/Properties/Set]
+	[/PARAM]
+[/OBJECT]
 [/REGION]
 
 ## Categories
diff --git a/Engine/Documentation/Source/Engine/UI/LevelEditor/Details/Properties/Map/MapProperty.INT.udn b/Engine/Documentation/Source/Engine/UI/LevelEditor/Details/Properties/Map/MapProperty.INT.udn
new file mode 100644
--- /dev/null
+++ b/Engine/Documentation/Source/Engine/UI/LevelEditor/Details/Properties/Map/MapProperty.INT.udn
@@ -0,0 +1,25 @@
+Availability:Public
+Title: Map Control
+Crumbs: %ROOT%, Engine, Engine/UI, Engine/UI/LevelEditor, Engine/UI/LevelEditor/Details
+Description:Dynamic lists of two related properties.
+Version:4.14
+
+Maps are primarily defined by two types - a key type and a value type - which are stored as associated pairs in the map. 
+
+![](MapExampleEmpty.png)
+
+Much like working with [Arrays](Engine/UI/LevelEditor/Details/Properties/Array) in the editor, you can add an element by clicking on the plus sign icon next to the Map name. 
+
+![](AddElements.png)
+
+Map keys and values are entered like other properties in the **Details** panel: colors are selected with the **Color Picker**, numbers can be entered or selected with the slider, etc.
+
+![](TMapFull.png)
+
+If you enter a duplicate key, the editor will warn you that this is not an allowed action and prevent the duplicate key from being saved.
+
+![](DuplicateKeyWarning.png)
+
+To delete any entry in the Map, use the drop-down arrow next to that element, then select **Delete**.
+
+![](TMapDelete.png)
\ No newline at end of file
diff --git a/Engine/Documentation/Source/Engine/UI/LevelEditor/Details/Properties/Set/SetProperty.INT.udn b/Engine/Documentation/Source/Engine/UI/LevelEditor/Details/Properties/Set/SetProperty.INT.udn
new file mode 100644
--- /dev/null
+++ b/Engine/Documentation/Source/Engine/UI/LevelEditor/Details/Properties/Set/SetProperty.INT.udn
@@ -0,0 +1,25 @@
+Availability:Public
+Title: Set Control
+Crumbs: %ROOT%, Engine, Engine/UI, Engine/UI/LevelEditor, Engine/UI/LevelEditor/Details
+Description:Dynamic unordered lists of other property types.
+Version:4.14
+
+Sets are a fast container class to store (usually) unique elements in a context where the order is irrelevant. 
+
+![](FruitSetEmpty.png)
+
+Much like working with [Arrays](Engine/UI/LevelEditor/Details/Properties/Array) in the editor, you can add an element by clicking on the plus sign icon next to the Set name. 
+
+![](SetAddElement.png)
+
+Set elements are entered like other properties in the **Details** panel: colors are selected with the **Color Picker**, numbers can be entered or selected with the slider, etc.
+
+![](SetFull.png)
+
+If you enter a duplicate element, the editor will warn you that this is not an allowed action and prevent the duplicate element from being saved.
+
+![](DuplicateSetWarning.png)
+
+To delete any entry in the Set, use the drop-down arrow next to that element, then select **Delete**.
+
+![](SetDelete.png)
\ No newline at end of file
diff --git a/Engine/Documentation/Source/Engine/UI/LevelEditor/QuickSettings/ViewOptions.INT.udn b/Engine/Documentation/Source/Engine/UI/LevelEditor/QuickSettings/ViewOptions.INT.udn
--- a/Engine/Documentation/Source/Engine/UI/LevelEditor/QuickSettings/ViewOptions.INT.udn
+++ b/Engine/Documentation/Source/Engine/UI/LevelEditor/QuickSettings/ViewOptions.INT.udn
@@ -26,7 +26,7 @@ The **Settings** are a set of properties for controlling selection, editing, and
 | **Enable Socket Snapping** | If checked, allows Actors to be snapped to sockets. |
 | **Enable Vertex Snapping** | If checked, Actors will snap to the nearest vertex on another Actor in the direction of movement. |
 | **Enable Planar Snapping** | If checked, actors will snap to the nearest location on the constraint plane. This only works correctly in Perspective views. |
-[REGION:tablesection]Viewport[/REGION] ||
+| [REGION:tablesection]Viewport[/REGION] ||
 | **Hide Viewport UI** | Toggles hidden Viewport UI mode, and hides all overlaid Viewport UI widgets. |
 | [REGION:tablesection]Previewing[/REGION] ||
 | **Draw Brush Polys** | If checked, semi-translucent polygons are rendered for the faces of a CSG brush when selected. [INCLUDE:Engine/UI/LevelEditor/QuickSettings#BrushPolys] |
diff --git a/Engine/Documentation/Source/Engine/UI/MeshPreviewScene/MeshPreviewScene.INT.udn b/Engine/Documentation/Source/Engine/UI/MeshPreviewScene/MeshPreviewScene.INT.udn
new file mode 100644
--- /dev/null
+++ b/Engine/Documentation/Source/Engine/UI/MeshPreviewScene/MeshPreviewScene.INT.udn
@@ -0,0 +1,197 @@
+Availability:docs
+Title: Mesh Preview Scenes
+Crumbs: %ROOT%, Engine
+Description: How to setup and use Mesh Preview Scene Profiles in Unreal Engine 4.
+Version: 4.13
+
+
+
+In the editor various editor windows used for your Static and Skeletal Meshes you can now setup preview scenes to be used when inspecting your assets. This allows you to quickly assess what your assets would look like in multiple different 
+environments and lighting scenarios without having to set up these scenes within the game viewport. 
+
+![](6.png)(w:950)
+
+The **Preview Scene Settings** panel enables you to adjust the following in your preview scene:
+
+* Directional Lights: Color, Intensity, and Rotation
+* Sky Lights: HDRI Environment Map, Intensity, and Rotation
+* Post Process Volumes: All Settings
+
+## Using Scene Profiles
+
+The **Preview Scene Settings** enable you to set up any number of environments and lighting scenarios to emulate your game levels. These scenes can be saved as any number of separate 
+profiles that can be shared between the various editor windows. 
+
+
+### Adding Profiles
+
+1. When you open your Static or Skeletal Mesh from the Content Browser, locate the panel on the right for **Preview Scene Settings** in the editor window. 
+	
+	[REGION:imagetable]
+	| [INCLUDE:#staticmesheditor] | [INCLUDE:#skeletalmesheditor] |
+	| --- | --- |
+	|Static Mesh Editor| Skeletal Mesh Editor |
+	| _Click for Full Image_ | _Click for Full Image_ |
+	[/REGION]
+
+2. Adjust the various settings until you have something that emulates the environment lighting and Post Processing settings you would like to use.
+	
+	[OBJECT:EmbeddedVideo]
+		[PARAMLITERAL:width]
+		640
+		[/PARAMLITERAL]
+		[PARAMLITERAL:height]
+		360
+		[/PARAMLITERAL]
+		[PARAMLITERAL:videoid]
+		VxhRgJuxxwI
+		[/PARAMLITERAL]
+	[/OBJECT]
+
+3. In the text box for **Profile Name**, give it a name that will be useful, like "Nighttime Environment" or "Daytime Environment".
+
+	[REGION:lightbox]
+	[![](ProfileName.png)(w:350)](ProfileName.png) 
+	[/REGION]
+
+4. Click on the button for **Add Profile**. 
+
+	[REGION:lightbox]
+	[![](AddingProfile.png)(w:350)](AddingProfile.png) 
+	[/REGION]
+
+5. You have now created and saved your profile and can select it from the drop-down next to **Scene Profile Settings**. 
+
+	[REGION:lightbox]
+	[![](ProfileSelection.png)(w:350)](ProfileSelection.png) 
+	[/REGION]
+
+
+#### Removing Profiles
+
+1. Select a profile from the drop-down menu next to **Scene Profile Settings**.
+
+	[REGION:lightbox]
+	[![](RemoveProfileSelection.png)(w:350)](ProfileSelection.png) 
+	[/REGION]
+
+2. Click on the button for **Remove Profile** to permanently delete the currently selected profile.
+
+	[REGION:lightbox]
+	[![](RemoveProfileButton.png)(w:350)](ProfileSelection.png) 
+	[/REGION]
+
+3. You will now see that profile has been removed leaving only the default profile and any other previously set up profiles. 
+
+	[REGION:lightbox]
+	[![](RemoveProfileSelection1.png)(w:350)](ProfileSelection.png) 
+	[/REGION]
+
+
+### Rotating the Environment
+
+In this section, you will adjust the environment rotation which will rotate the skybox and the Directional Light around the mesh. This setting can be used to simply rotate the environment
+around the mesh manually or to set a speed for the environment to rotate continually.
+
+
+1. Enable environment rotation by placing a checkbox next to **Rotate Sky and Directional Lighting**.
+
+	[REGION:lightbox]
+	[![](RotatingEnvironment.png)(w:350)](RotatingEnvironment.png) 
+	[/REGION]
+
+2. Next, set a value for **Rotation Speed** for the environment to start rotating.
+
+	![](RotationSpeed.png)(w:350)
+	[OBJECT:EmbeddedVideo]
+		[PARAMLITERAL:width]
+		640
+		[/PARAMLITERAL]
+		[PARAMLITERAL:height]
+		360
+		[/PARAMLITERAL]
+		[PARAMLITERAL:videoid]
+		3GeYcT1isgE
+		[/PARAMLITERAL]
+	[/OBJECT]
+
+	[REGION:note]
+	You can manually rotate the environment by holding down the **K** key and left-click drag and the Directional Light by holding down the **L** key and left-click drag to rotate.
+	[/REGION]
+
+
+### Toggle Environment and Floor Visibility
+
+In this section, you will learn how to hide and show the Environment and Floor mesh for your preview scene. 
+
+1. Click the check box next to **Show Environment** and **Show Floor** to hide/show these features. 
+
+	[OBJECT:EmbeddedVideo]
+		[PARAMLITERAL:width]
+		640
+		[/PARAMLITERAL]
+		[PARAMLITERAL:height]
+		360
+		[/PARAMLITERAL]
+		[PARAMLITERAL:videoid]
+		b9r1-mx1xfM
+		[/PARAMLITERAL]
+	[/OBJECT]
+
+
+### Assigning an Environment Cube Map
+
+In this section, you will learn how to assign a skybox for the environment for your preview scene.
+
+1. Click the selection drop-down next to **Environment Cube Map** to select from the available HDRI cube maps. 
+
+	[OBJECT:EmbeddedVideo]
+		[PARAMLITERAL:width]
+		640
+		[/PARAMLITERAL]
+		[PARAMLITERAL:height]
+		360
+		[/PARAMLITERAL]
+		[PARAMLITERAL:videoid]
+		j8r1-Ws-33U
+		[/PARAMLITERAL]
+	[/OBJECT]
+
+
+[REGION:tip]
+Use the View Options for **Show Engine Content** to see additional cube maps that are included in the **Engine Content** folder.
+
+[REGION:lightbox]
+[![](ViewOptions.png)(w:600)](ViewOptions.png)
+[/REGION]
+[/REGION]
+
+
+### Post Process Settings
+
+Like the Post Process Volume settings available in your game viewport, you have full control over all available Post Process settings for your preview scenes under the advanced rollout. 
+You can toggle the post processing on and off for the Preview Scene by using the checkbox for **Post Processing Enabled**. 
+
+[REGION:lightbox]
+[![](ScenePostProcess.png)(w:350)](ScenePostProcess.png)
+[/REGION]
+
+For a full breakdown of the various settings you can check out the [Post Process Effects](https://docs.unrealengine.com/latest/INT/Engine/Rendering/PostProcessEffects/) page. 
+
+
+
+
+
+<!---
+[EXCERPT:StaticMeshEditor]
+[REGION:lightbox]
+[![](StaticMeshEditor.png)(w:450)](StaticMeshEditor.png)
+[/REGION] 
+[/EXCERPT:StaticMeshEditor]
+
+[EXCERPT:SkeletalMeshEditor]
+[REGION:lightbox]
+[![](SkeletalMeshEditor.png)(w:450)](SkeletalMeshEditor.png)
+[/REGION] 
+[/EXCERPT:SkeletalMeshEditor]
+-->
