Copying //UE4/Dev-Documentation to Samples-Main (//UE4/Samples-Main)

diff --git a/Engine/Documentation/Source/Engine/UMG/HowTo/UMGHowTo.INT.udn b/Engine/Documentation/Source/Engine/UMG/HowTo/UMGHowTo.INT.udn
--- a/Engine/Documentation/Source/Engine/UMG/HowTo/UMGHowTo.INT.udn
+++ b/Engine/Documentation/Source/Engine/UMG/HowTo/UMGHowTo.INT.udn
@@ -18,6 +18,7 @@ Refer to the table below for more information and click each link to view its co
 
 ## Intermediate
 [REGION:wrapper transparent]
+[](FANCY:Engine\UMG\HowTo\WidgetTemplates "%Engine\UMG\HowTo\WidgetTemplates:description%")
 [](FANCY:Engine\UMG\HowTo\Create3DWidgets "%Engine\UMG\HowTo\Create3DWidgets:description%")
 [PUBLISH:docs]
 [](FANCY:Engine\UMG\HowTo\DragAndDrop "%Engine\UMG\HowTo\DragAndDrop:description%")
diff --git a/Engine/Documentation/Source/Engine/UMG/HowTo/WidgetTemplates/WidgetTemplates.INT.udn b/Engine/Documentation/Source/Engine/UMG/HowTo/WidgetTemplates/WidgetTemplates.INT.udn
new file mode 100644
--- /dev/null
+++ b/Engine/Documentation/Source/Engine/UMG/HowTo/WidgetTemplates/WidgetTemplates.INT.udn
@@ -0,0 +1,122 @@
+Availability: Public
+Title:Creating Widget Templates
+Crumbs: %ROOT%, Engine, Engine/UMG, Engine/UMG/HowTo
+Description:An example of how to create, display and remove UI elements from the screen during gameplay. 
+Related: Engine/UMG
+Related: Engine/UMG/UserGuide
+Version: 4.13
+SkillLevel: Intermediate
+
+Every **Widget Blueprint** that you create with UMG is considered a **User Widget** that can be reused and placed within other Widget Blueprints. 
+This allows you to essentially create templates for pieces of your UI that can be used within other Widget Blueprints. 
+Not only will the visual aspects of the widget be carried over to the Widget Blueprint you place it in, but scripted functionality will also be carried over. 
+
+With some Blueprint Scripting, you could create variables for how your UI widget functions or looks that can be overwritten on a per instance basis. 
+For example, in this How-to we will create a Button Widget Blueprint that we can use within other Widget Blueprints. 
+We will create and stylize the Button but set it up in such a way that we can override the style when we place it in other Widget Blueprints. 
+Lastly, we will customize what happens when the Button is clicked on a per instance basis (this is useful if you have a series of Buttons in a menu that you want to look and act the same but produce different results when pressed). 
+
+## Steps
+
+[REGION:note]
+For this how-to, we are using the **Blueprint Third Person Template** with **Starter Content** enabled. 
+[/REGION]
+
+1.	**Right-click** in the **Content Browser** and create two **Widget Blueprints**, one called **CustomButton** and another called **HUD**. The **CustomButton** is a UserWidget that we create and use within our main **HUD** Widget Blueprint. 
+
+	![](UserWidget01.png)
+
+1.	Inside the **CustomButton** Widget Blueprint, delete the **Canvas Panel**, add a **Button** and **Right-click** on it and **Wrap With** a **Size Box**. 
+
+	![](UserWidget02.png)
+
+1.	For the **Size Box**, in the **Details** panel, change the **Width Override** to **300** and **Height Override** to **100**. Optionally, you can also change the graph layout to **Desired on Screen** to get an idea of the actual size of the button. 
+
+	![](UserWidget03.png)
+
+1.	On the **Graph** tab, hold **Ctrl** and drag **Button_0** from the **My Blueprint** panel, then drag off it and add a **Set Style** node.Select the option under the **Variables > Apperance** category, as the other deals with the color multiplier for the button background. 
+
+	![](UserWidget04.png)
+
+1.	Connect the **Event Construct** node to the **Set Widget Style** node, then **Right-click** on **Widget Style** and **Promote to Variable**. 
+
+	![](UserWidget05.png)
+
+1.	Name the new variable **ButtonStyle** and **Compile**, then set the **Normal** style **Image** to any texture. 
+
+	![](UserWidget06.png)
+
+1.	**Right-click** on the **Normal** section and select **Copy**, then **Right-click** and **Paste** to the **Hovered** and **Pressed** states. 
+
+	![](UserWidget07.png)
+
+1.	Expand **Hovered** and change the **Tint** color to any color (such as yellow). 
+
+	![](UserWidget08.png)
+
+1.	Also for **ButtonStyle**, check the **Editable** and **Expose on Spawn** checkboxes. This will allow us to modify the values of this variable from other Widget Blueprints when we use this Widget Blueprint elsewhere. 
+
+	![](UserWidget08b.png)
+
+1.	Select the **Button_0** variable, then click the **+** sign for **OnClicked**. 
+
+	![](UserWidget09.png)
+
+1.	In the **MyBlueprint** panel, click the **+ Event Dispatcher** button to create a new **Event Dispatcher** and call it **ButtonClicked**. 
+
+	![](UserWidget10.png)
+
+1.	Drag **ButtonClicked** into the graph and select **Call**, then hook it up to the **OnClicked** event. What this allows us to do is create unique scripted functionality whenever the button is clicked. If we had multiple instances of this button and only used the **OnClicked** event, each instance would respond and execute the same functionality. However, by creating an **Event Dispatcher**, we can then implement events on a per instance basis and have only the button that was actually clicked on fire off additional script. 
+
+	![](UserWidget11.png)
+
+1.	Open the **HUD** Widget Blueprint, then add a **Vertical Box** to the **Canvas Panel** and resize similar to below. 
+
+	![](UserWidget12.png)
+
+1.	From the **Palette** under **User Created**, add three **Custom Buttons** to the **Vertical Box**. 
+
+	![](UserWidget13.png)
+
+1.	On the **Graph**, select **CustomButton** then click **+** for the Event **Button Clicked**. This is the **Event Dispatcher** that we created inside our **CustomButton** Widget Blueprint and will respond whenever this particular button is clicked. 
+
+	![](UserWidget15.png)
+
+1.	Click the **+** for **Button Clicked** for the other buttons and connect **Print String** nodes up to each, printing different text. In our example, we are printing to the screen each time one of the buttons are clicked, however, depending on the button that is clicked, different text will be printed to the screen. In a real-world example, you could have each button open up a different menu, change different game settings, or make modifications to a player character.
+
+	![](UserWidget16.png)
+
+1.	From the Level Editor Main Toolbar, click **Blueprints** and **Open Level Blueprint**. 
+
+	![](UserWidget17.png)
+
+1.	**Right-click** and add an **Event BeginPlay** then connect it to a **Create Widget** node, with the **Class** set to **HUD**. 
+
+	![](UserWidget18.png)
+
+1.	Create and connect an **Add to Viewport** node, then **Get Player Controller** to **Set Show Mouse Cursor** (set to **True**). 
+
+	![](UserWidget19.png)
+
+1.	**Compile** and close the **Level Blueprint**, then click **Play** to play in the editor. 
+	
+	
+## End Result
+
+When playing in the editor, each Button is automatically styled based on the defined **ButtonStyle** variable. When clicking on each of the buttons they execute a different portion of the script, as the **Event Dispatcher** is called and the corresponding **ButtonClicked** Event is called. Optionally, since our **ButtonStyle** variable is exposed and editable on spawn, we can override the style of the button in whatever other Widget Blueprints we place it in. This is useful if you have some kind of animation or other complex scripted logic that you didn't want to re-create each time, by creating it as a UserWidget, you can reuse it in any other Widget Blueprints without additional work. 
+
+[OBJECT:EmbeddedVideo]
+[PARAMLITERAL:width]
+640
+[/PARAMLITERAL]
+[PARAMLITERAL:height]
+360
+[/PARAMLITERAL]
+[PARAMLITERAL:videoid]
+j3Oni-3AM_I
+[/PARAMLITERAL]
+[/OBJECT]
+
+
+
+
diff --git a/Engine/Documentation/Source/Engine/UMG/QuickStart/3/UMGQS_3.INT.udn b/Engine/Documentation/Source/Engine/UMG/QuickStart/3/UMGQS_3.INT.udn
--- a/Engine/Documentation/Source/Engine/UMG/QuickStart/3/UMGQS_3.INT.udn
+++ b/Engine/Documentation/Source/Engine/UMG/QuickStart/3/UMGQS_3.INT.udn
@@ -2,7 +2,7 @@ Availability: Public
 Title:3. Creating a Main Menu
 Crumbs:%ROOT%, Engine, Engine/UMG, Engine/UMG/QuickStart
 Description:In this step we will create the visual layout for our Main Menu.
-version: 4.9
+version: 4.13
 
 
 [VAR:Steps]
@@ -71,7 +71,7 @@ First bit of business is to create the layout of our Main Menu which we will do
 
 	Here we are adjusting the size of all of our buttons as well as defining a color for the button to take when mousing over it. 
 
-1.	Select the **Image** in the **Hierarchy** then in the **Details** under **Appearance** and **Brush**, select a **Texture** or **Material** to use. 
+1.	Select the **Image** in the **Hierarchy** then in the **Details** under **Appearance** and **Brush**, select a **Texture**, **Sprite**, or **Material** to use. 
 
 	![](Hierarchy8.png)
 
@@ -81,11 +81,11 @@ First bit of business is to create the layout of our Main Menu which we will do
 	Drag-and-drop the Example Background into **Unreal Engine** to import it (click **Yes** at any confirmation messages). 
 	[/REGION]
 
-1.	In the **Designer** window, resize the image so that it takes up the entire layout. 
+1.	In the **Designer** window, resize the image so that it takes up the entire layout.
 
 	![](Hierarchy9.png)
 
-1.	In the **Details** panel for the **Image**, click the **Anchor** button and choose the fill screen option. 
+1.	In the **Details** panel for the **Image**, click the **Anchor** button and choose the fill screen option. You can also use a Scalebox to hold the image, which will ensure that the image scales and resizes appropriately to match the aspect ratio.
 
 	![](Hierarchy10.png)
 
diff --git a/Engine/Documentation/Source/Engine/UMG/UserGuide/Fonts/Fonts.INT.udn b/Engine/Documentation/Source/Engine/UMG/UserGuide/Fonts/Fonts.INT.udn
--- a/Engine/Documentation/Source/Engine/UMG/UserGuide/Fonts/Fonts.INT.udn
+++ b/Engine/Documentation/Source/Engine/UMG/UserGuide/Fonts/Fonts.INT.udn
@@ -5,7 +5,7 @@ Description:Details the process of importing fonts, working with the Font Editor
 Related: Engine/UMG
 Related: Engine/UMG/UserGuide/WidgetBlueprints
 Related: Engine/UMG/UserGuide/WidgetTypeReference
-Version: 4.9
+Version: 4.14
 
 [TOC(start:2 end:2)]
 
@@ -97,25 +97,25 @@ As seen in the example image above, the Japanese text falls within the character
 
 Once you have created a [Widget Blueprint](Engine/UMG/UserGuide/WidgetBlueprints) and placed a Widget that contains Text (such as a Text Box or Text Widget), when clicking on it and accessing the **Details** panel, under **Appearance**, you will see a **Font** option. 
 
-![](UMGFonts1.png)
+![](FontWithNewOptions.png)
 
 By default the Engine uses **Roboto**, however if you click the drop-down menu, any Composite Font assets created can be selected and used instead. 
 
-![](UMGFonts2.png)
+![](NewFontWithNewOptions.png)
 
 You can also choose to create a Composite Font from this menu and specify where the new asset should be saved (it will be blank by default and you will have to fill it out). 
 
 Once you select your Composite Font, the second drop-down menu will allow you to select a font to use from the **Default Font Family**. 
 
-![](UMGFonts3.png)
+![](FontStyleNewOptions.png)
 
 You can also specify the size of the font in the input box. 
 
-![](UMGFonts4.png)
+![](FontSizeNewOptions.png)
 
 Currently UMG only supports **Runtime** cached font assets. Also, if you have assigned fonts using the old method, none of your existing file based font settings will be lost, however going forward, you will need to create a Composite Font asset in order to use custom fonts with UMG. 
 
-
+You can style your **Font** by [setting colors, Materials, and outline properties](Engine/UMG/UserGuide/Fonts/Materials), as well as by using shadows. 
 
 
 
diff --git a/Engine/Documentation/Source/Engine/UMG/UserGuide/Fonts/Materials/FontMaterialsAndOutlines.INT.udn b/Engine/Documentation/Source/Engine/UMG/UserGuide/Fonts/Materials/FontMaterialsAndOutlines.INT.udn
new file mode 100644
--- /dev/null
+++ b/Engine/Documentation/Source/Engine/UMG/UserGuide/Fonts/Materials/FontMaterialsAndOutlines.INT.udn
@@ -0,0 +1,88 @@
+Availability:Public
+Title: Font Materials and Outlines
+Description:Examples of how you can style your UMG fonts using colors, Materials, and outlines.
+Crumbs:
+Related: Engine/UMG
+Related: Engine/UMG/UserGuide/Fonts
+Related: Engine/UMG/UserGuide/Styling
+Version: 4.14
+Parent: Engine/UMG/UserGuide/Fonts
+order: 1
+tags: User Interface
+
+[TOC]
+
+In addition to being able to set a **Color and Opacity** for your **Font** in UMG, you can also use Materials and font outlines for additional **Font** styling.
+
+## Font Colors
+
+You can set the **Vertex Color** for your **Font** by setting its **Color and Opacity**. 
+
+![](GreenFontNoMaterial.png)
+
+Without a **Font Material** specified, this will apply a solid color to your text.
+
+![](GreenFontNoMaterialDisplay.png)
+ 
+## Font Materials
+
+You can specify a **Font Material** for your **Font** in the **Details** panel.
+
+![](ColorAndOpacity.png)
+
+If your **Font Material** doesn't have a **Vertex Color** node, the effect is like a simple multiplication. 
+
+[REGION:imagetable]
+| [INCLUDE:#basecolorgraph] | ![](BaseColorFinalFontMat.png)(w:200)|![](BaseColorFinalDisplay.png)(w:550) |
+|---- |----- |----- |
+| Font Material Setup |  Font Material Preview | Final Font |
+[/REGION]
+
+However, if you place a **Vertex Color** node in your **Font Material**, you can use its outputs to do math within your shader. 
+
+[REGION:imagetable]
+|[INCLUDE:#vertexcolorgraph] | ![](VertexColorFinalMat.png)(w:200)|![](VertexColorFont.png)(w:550) |
+|---- |----- |----- |
+| Font Material Setup |  Font Material Preview | Final Font |
+[/REGION]
+
+Materials used as font materials must be in the **User Interface** domain.
+
+![](UI.png)
+
+<!--
+[EXCERPT:BaseColorGraph]
+[REGION:lightbox]
+[![](BaseColorMatGraph.png)(w:400)](BaseColorMatGraph.png)
+[/REGION]
+[/EXCERPT]
+
+[EXCERPT:VertexColorGraph]
+[REGION:lightbox]
+[![](VertexColorMatGraph.png)(w:400)](VertexColorMatGraph.png)
+[/REGION]
+[/EXCERPT]
+-->
+## Font Outlines
+
+You can specify a different **Outline Color** as well as a different Material to use for a font outline. 
+
+![](FontOutlineSettings.png)
+
+The size of the outline is specified in Slate units, but when the font's scale is 1.0, 1 Slate unit is equivalent to 1 pixel. 
+
+![](FontOutline5Display.png)
+
+One interesting note is that you can specify whether or not to use **Separate Fill Alpha**. 
+
+![](SeparateFill.png)
+
+When this is enabled, the outline is translucent
+where the filled area will be, allowing you to adjust the alphas of the font and the font outline independently. When it is disabled, the font is overlaid
+on the outline, so the alphas are additive and the outline color and Material is visible through a font with an alpha less than 1.
+
+[REGION:imagetable]
+|![](SeparateAlphaOff.png)(w:550) | ![](SeparateAlphaOn.png)(w:550)|
+|---|----|
+| Font Color Alpha = 0.5, **Separate Fill Alpha** disabled | Font Color Alpha = 0.5, **Separate Fill Alpha** enable |
+[/REGION]
\ No newline at end of file
diff --git a/Engine/Documentation/Source/Engine/UMG/UserGuide/WidgetTypeReference/WidgetTypeReference.INT.udn b/Engine/Documentation/Source/Engine/UMG/UserGuide/WidgetTypeReference/WidgetTypeReference.INT.udn
--- a/Engine/Documentation/Source/Engine/UMG/UserGuide/WidgetTypeReference/WidgetTypeReference.INT.udn
+++ b/Engine/Documentation/Source/Engine/UMG/UserGuide/WidgetTypeReference/WidgetTypeReference.INT.udn
@@ -2,7 +2,7 @@ Availability:Public
 Title: Widget Type Reference
 Crumbs:%ROOT%, Engine, Engine/UMG, Engine/UMG/UserGuide
 Description:Reference for the types of Widgets available in Unreal Motion Graphics.
-version: 4.9
+version: 4.13
 
 [TOC(start:2 end:2)]
 
@@ -26,12 +26,12 @@ The most frequently used Widgets will more than likely fall under this category.
 | **Border** | A border is a container widget that can contain one child widget, providing an opportunity to surround it with a border image and adjustable padding. |
 | **Button** | The button is a single-child, click-able primitive widget which enables basic interaction. You can place any other widget inside a button to make a more complex and interesting click-able element in your UI.  |
 | **Check Box** | The check box widget allows you to display a toggled state of 'unchecked', 'checked', and 'indeterminable'. You can use the check box for a classic check box, as a toggle button, or as radio buttons.  |
-| **Image** | The image widget allows you to display a Slate Brush, Texture, or Material in the UI. |
 | [Named Slot](Engine\UMG\UserGuide\WidgetTypeReference\NamedSlot) | [INCLUDE:Engine\UMG\UserGuide\WidgetTypeReference\NamedSlot#desc] |
 | **Progress Bar** | The progress bar widget is a simple bar that fills up and can be restyled to fit any number of uses such as Experience, Health, Points, etc. |
 | **Slider** | A simple widget that shows a sliding bar with a handle that allows you to control the value between 0-1. |
 | **Text** | The basic method for displaying text on screen and could be used for text descriptions of options or other UI elements. |
 | **Text Box** | Allows the user to type in custom text. Only permits a single line of text to be entered. |
+| **Image** | The image widget allows you to display a Slate Brush, Texture, Sprite, or Material in the UI. [REGION:note]Sprites that are part of the same texture atlas can be batched together in Slate, provided they all share the same layer when rendered. This means that platforms where the draw call count budget is tight will render UMG and Slate Widgets significantly more efficiently if Paper2D Sprites are used as Brush inputs.[/REGION] |
 
 
 ## Input Widgets
@@ -65,7 +65,7 @@ The **Panel** category contains useful Widgets for controlling layouts and placi
 
 | Option | Description |
 | -------- | ----------- |
-| **Canvas Panel** | The canvas panel is a designer friendly panel that allows widgets to be laid out at arbitrary locations, anchored and z-ordered with other children of the canvas. The canvas panel is a great widget for manual layout, but bad when you want to procedurally just generate widgets and place them in a container (unless you want absolute layout). For more information on Anchors, see the [Anchors](Engine\UMG\UserGuide\Anchors\) page. |
+| **Canvas Panel** | The canvas panel is a designer friendly panel that allows widgets to be laid out at arbitrary locations, anchored, and z-ordered with other children of the canvas. Note that, although z-order can be altered manually, widgets are rendered in list order, and it is considered preferable to order them correctly in the list rather than to rely on z-ordering. The canvas panel is a great widget for manual layout, but not as useful when you want to generate widgets procedurally and place them in a container (unless you want an absolute layout). For more information on Anchors, see the [Anchors](Engine\UMG\UserGuide\Anchors\) page. |
 | **Grid Panel** | This is a panel that evenly divides up available space between all of its children widgets. |
 | **Horizontal Box** | Allows child widgets to be laid out in a flow horizontally. |
 | **Overlay** | Allows widgets to be stacked on top of each other and uses a simple flow layout for content on each layer.  |
diff --git a/Engine/Documentation/Source/Gameplay/Framework/GameMode/GameMode.INT.udn b/Engine/Documentation/Source/Gameplay/Framework/GameMode/GameMode.INT.udn
--- a/Engine/Documentation/Source/Gameplay/Framework/GameMode/GameMode.INT.udn
+++ b/Engine/Documentation/Source/Gameplay/Framework/GameMode/GameMode.INT.udn
@@ -7,6 +7,7 @@ type:reference
 parent:Gameplay/Framework
 order:5
 tags:Gameplay Framework
+Related:Gameplay/Framework/Controller/PlayerController
 
 [TOC (start:2 end:2)]
 
@@ -28,36 +29,32 @@ When rule-related events in the game happen and need to be tracked and shared wi
 
 ## Game Modes
 
-While certain fundamentals, like the number of players required to play, or the method by which those players join the game, are common to many types of games, limitless rule variations are possible depending on the specific game you are developing. Regardless of what those rules are, Game Modes are designed to define and implement them. There are currently two commonly-used base classes for Game Modes. Engine version 4.14 introduces `AGameModeBase`, which is the base class for all Game Modes, and is a simplified and streamlined version of the classic `AGameMode`. `AGameMode`, the Game Mode base class before version 4.14, still exists and functions as it did before, but is now a child of `AGameModeBase`. `AGameMode` is more suited to standard game types like multiplayer shooters due to its implementation of the concept of match state. `AGameModeBase` is the new default game mode included in new code projects due to its simplicity and efficiency.
+While certain fundamentals, like the number of players required to play, or the method by which those players join the game, are common to many types of games, limitless rule variations are possible depending on the specific game you are developing. Regardless of what those rules are, Game Modes are designed to define and implement them. There are currently two commonly-used base classes for Game Modes. 
+
+Engine version 4.14 introduces `AGameModeBase`, which is the base class for all Game Modes and is a simplified and streamlined version of the classic `AGameMode`. `AGameMode`, the Game Mode base class before version 4.14, still exists and functions as it did before, but is now a child of `AGameModeBase`. `AGameMode` is more suited to standard game types like multiplayer shooters due to its implementation of the concept of match state. `AGameModeBase` is the new default game mode included in new code projects due to its simplicity and efficiency.
 
 ### AGameModeBase
 All Game Modes are subclasses of `AGameModeBase`, which contains considerable base functionality that can be overridden. Some of the common functions include:
 
-$ InitGame : The `InitGame` event is called before any other scripts (including `PreInitializeComponents`), and is used by `AGameModeBase` to initialize parameters and spawn its helper classes.
-
-	[REGION:note]
-	This is called before any Actor runs `PreInitializeComponents`, including the Game Mode instance itself.
-	[/REGION:note]
-
-$ PreLogin : Accepts or rejects a player who is attempting to join the server. Causes the `Login` function to fail if it sets `ErrorMessage` to a non-empty string. `PreLogin` is called before `Login`, and a significant amount of time may pass before Login is called, especially if the joining player needs to download game content.
-
-$ PostLogin : Called after a successful login. This is the first place it is safe to call replicated functions on the `PlayerController`. `OnPostLogin` can be implemented in blueprint to add extra logic.
+| Function/Event | Purpose |
+| --- | --- |
+| `InitGame` | The `InitGame` event is called before any other scripts (including `PreInitializeComponents`), and is used by `AGameModeBase` to initialize parameters and spawn its helper classes.	[REGION:note]This is called before any Actor runs `PreInitializeComponents`, including the Game Mode instance itself.[/REGION:note] |
+| `PreLogin` | Accepts or rejects a player who is attempting to join the server. Causes the `Login` function to fail if it sets `ErrorMessage` to a non-empty string. `PreLogin` is called before `Login`, and a significant amount of time may pass before Login is called, especially if the joining player needs to download game content. |
+| `PostLogin` | Called after a successful login. This is the first place it is safe to call replicated functions on the `PlayerController`. `OnPostLogin` can be implemented in Blueprint to add extra logic. |
+| `HandleStartingNewPlayer` | Called after `PostLogin` or after seamless travel, this can be overridden in Blueprint to change what happens to a new player. By default, it will create a pawn for the player. |
+| `RestartPlayer` | This is called to start spawning a player's Pawn. There are also `RestartPlayerAtPlayerStart` and `RestartPlayerAtTransform` functions available if you want to dictate where the Pawn will be spawned. `OnRestartPlayer` can be implemented in Blueprint to add logic after this function is finished. |
+| `SpawnDefaultPawnAtTransform` | This is what actually spawns the player's Pawn, and can be overridden in Blueprints. |
+| `Logout` | Called when a player leaves the game or is destroyed. `OnLogout` can be implemented to do Blueprint logic. |
 
-$ HandleStartingNewPlayer : Called after `PostLogin` or after seamless travel, this can be overridden in blueprint to change what happens to a new player. By default, it will create a pawn for the player.
+A subclass of the `AGameModeBase` class may be created for each match format, mission type, or special zone that the game offers. A game may have any number of Game Modes, and thus subclasses of the `AGameModeBase` class; however, only one Game Mode may be in use at any given time. A Game Mode Actor is instantiated each time a level is initialized for play via the `UGameEngine::LoadMap()` function. 
 
-$ RestartPlayer : This is called to start spawning a player's Pawn. There are also `RestartPlayerAtPlayerStart` and `RestartPlayerAtTransform` functions available if you want to dictate where the Pawn will be spawned. `OnRestartPlayer` can be implemented in blueprint to add logic after this function is finished.
-
-$ SpawnDefaultPawnAtTransform : This is what actually spawns the player's Pawn, and can be overridden in blueprints.
-
-$ Logout : Called when a player leaves the game or is destroyed. `OnLogout` can be implemented to do blueprint logic.
-
-A subclass of the `AGameModeBase` class may be created for each match format, mission type, or special zone that the game offers. A game may have any number of game modes, and thus subclasses of the `AGameModeBase` class; however, only one Game Mode may be in use at any given time. A Game Mode Actor is instantiated each time a level is initialized for play via the `UGameEngine::LoadMap()` function. The Game Mode is not replicated to any remote clients that join in a multiplayer game; it exists only on the server, so local clients can see the stock Game Mode class (or Blueprint) that was used, but cannot access the actual instance and check its variables to see what has changed as the game progresses. If players do need updated information relating to the current Game Mode, that information is easily kept in sync by being stored on an `AGameStateBase` Actor, one of which will be created along with the Game Mode and then replicated out to all remote clients.
+The Game Mode is not replicated to any remote clients that join in a multiplayer game; it exists only on the server, so local clients can see the stock Game Mode class (or Blueprint) that was used, but cannot access the actual instance and check its variables to see what has changed as the game progresses. If players do need updated information relating to the current Game Mode, that information is easily kept in sync by being stored on an `AGameStateBase` Actor, one of which will be created along with the Game Mode and then replicated out to all remote clients.
 
 ### AGameMode
 
 `AGameMode` is a subclass of `AGameModeBase` that has some extra functionality to support multiplayer matches and legacy behavior. All newly created projects use `AGameModeBase` by default, but you can switch to inheriting from `AGameMode` if you need this extra behavior. If you inherit from `AGameMode`, you should also inherit your game state from `AGameState`, which also supports the match state machine.
 
-AGameMode contains a state machine that tracks the state of the match, or the general gameplay flow. To query the current state, you can use `GetMatchState`, or wrappers like `HasMatchStarted`, `IsMatchInProgress`, and `HasMatchEnded`. Here are the possible match states:
+`AGameMode` contains a state machine that tracks the state of the match or the general gameplay flow. To query the current state, you can use `GetMatchState`, or wrappers like `HasMatchStarted`, `IsMatchInProgress`, and `HasMatchEnded`. Here are the possible match states:
 
 * `EnteringMap` is the initial state. Actors are not yet ticking and the world has not been fully initialized. It will transition to the next state when things are fully loaded.
 * `WaitingToStart` is the next state, and `HandleMatchIsWaitingToStart` is called when entering. Actors are ticking, but players have not yet spawned in. It transitions to the next state if `ReadyToStartMatch` returns _true_, or if `StartMatch` is called.
@@ -66,7 +63,7 @@ AGameMode contains a state machine that tracks the state of the match, or the ge
 * `LeavingMap` is the last state in a normal flow, and calls `HandleLeavingMap` upon entry. The match stays in this state while transferring to a new map, where it will transition back to `EnteringMap`.
 * `Aborted` is the failure state, and can be started by calling `AbortMatch`. This is set when there is an unrecoverable error.
 
-The match state will almost always be `InProgress`, since this is the state where `BeginPlay` is called and actors begin ticking. However, individual games can override the behavior of these states to build a multiplayer game with more complicated rules, such as permitting players to fly around the level freely while waiting for other players to join in a multiplayer shooter.
+The match state will almost always be `InProgress` since this is the state where `BeginPlay` is called and actors begin ticking. However, individual games can override the behavior of these states to build a multiplayer game with more complicated rules, such as permitting players to fly around the level freely while waiting for other players to join in a multiplayer shooter.
 
 ### Game Mode Blueprints
 
@@ -84,15 +81,15 @@ In addition, Blueprints of Game Modes are very useful because they enable adjust
 
 ### Setting the Game Mode
 
-There are several methods to set the Game Mode for a level, ordered here from lowest priority to highest priority: 
+There are several methods to set the Game Mode for a level, ordered here from lowest priority to highest priority:
 
-* Setting the `GlobalDefaultGameMode` entry in the `/Script/EngineSettings.GameMapsSettings` section of the `DefaultEngine.ini` file will set the default game mode for all maps in the project. 
+* Setting the `GlobalDefaultGameMode` entry in the `/Script/EngineSettings.GameMapsSettings` section of the `DefaultEngine.ini` file will set the default game mode for all maps in the project.
 
 		[/Script/EngineSettings.GameMapsSettings]
 		GlobalDefaultGameMode="/Script/MyGame.MyGameGameMode"
 		GlobalDefaultServerGameMode="/Script/MyGame.MyGameGameMode"
 
-* To override the project settings for an individual map, set the **GameMode Override** in the **World Settings** tab in the editor. 
+* To override the project settings for an individual map, set the **GameMode Override** in the **World Settings** tab in the editor.
 		
 	![](WorldSettings_GameMode.png)
 
@@ -100,7 +97,7 @@ There are several methods to set the Game Mode for a level, ordered here from lo
 
 		UE4Editor.exe /Game/Maps/MyMap?game=MyGameMode -game
 
-* Finally, map prefixes (and aliases for the URL method) can be set in the `/Script/Engine.WorldSettings/` section of the `DefaultGame.ini` file. These prefixes set the default game mode for all maps that have a given prefix. 
+* Finally, map prefixes (and aliases for the URL method) can be set in the `/Script/Engine.WorldSettings/` section of the `DefaultEngine.ini` file. These prefixes set the default game mode for all maps that have a given prefix.
 
 		[/Script/EngineSettings.GameMapsSettings]
 		+GameModeMapPrefixes=(Name="DM",GameMode="/Script/UnrealTournament.UTDMGameMode")
@@ -113,14 +110,16 @@ For an example of setting up a Game Mode, refer to the [Setting Up a Game Mode](
 
 ## Game State
 
-The **Game State** is responsible for enabling the clients to monitor the state of the game. Conceptually, the Game State should manage information that is meant to be known to all connected clients and is specific to the Game Mode, but is not specific to any individual player. It can keep track of game-wide properties such as the list of connected players, team score in Capture The Flag, missions that been completed in an open world game, and so on. It is not the best place to keep track of player-specific things like how many points one specific player has scored for the team in a Capture The Flag match, because that can be handled more cleanly by **Player State**. In general, the GameState should track properties that change during gameplay and are relevant and visible to everyone. While the Game mode exists only on the server, the Game State exists on the server and is replicated to all clients, keeping all connected machines up to date as the game progresses.
-
-`AGameStateBase` is the base implementation, and some of it's default functionality includes:
+The **Game State** is responsible for enabling the clients to monitor the state of the game. Conceptually, the Game State should manage information that is meant to be known to all connected clients and is specific to the Game Mode but is not specific to any individual player. It can keep track of game-wide properties such as the list of connected players, team score in Capture The Flag, missions that been completed in an open world game, and so on.
 
-$ GetServerWorldTimeSeconds : This is the server's version of World::GetTimeSeconds and will be synchronized on both the client and server, so it is a reliable time to use for replication
+Game State is not the best place to keep track of player-specific things like how many points one specific player has scored for the team in a Capture The Flag match because that can be handled more cleanly by **Player State**. In general, the GameState should track properties that change during gameplay and are relevant and visible to everyone. While the Game mode exists only on the server, the Game State exists on the server and is replicated to all clients, keeping all connected machines up to date as the game progresses.
 
-$ PlayerArray : This is the array of all APlayerState objects, which is useful when doing something to all players in a game
+`AGameStateBase` is the base implementation, and some of its default functionality includes:
 
-$ HasBegunPlay : Returns true if the BeginPlay function has been called on actors in the game
+| Function or Variable | Use |
+| --- | --- |
+| `GetServerWorldTimeSeconds` | This is the server's version of the `UWorld` function `GetTimeSeconds` and will be synchronized on both the client and server, so it is a reliable time to use for replication. |
+| `PlayerArray` | This is the array of all `APlayerState` objects, which is useful when doing something to all players in a game. |
+| `HasBegunPlay` | Returns true if the `BeginPlay` function has been called on actors in the game. |
 
 `AGameStateBase` is very commonly extended in C++ or Blueprints to contain additional variables and functions that are needed to keep players informed of what's going on in the game. The specific modifications made are generally based on a paired Game Mode for which the Game State is made. The Game Mode itself can also override its default Game State type to be any C++ class or Blueprint derived from `AGameStateBase`.
diff --git a/Engine/Documentation/Source/Gameplay/Gameplay.INT.udn b/Engine/Documentation/Source/Gameplay/Gameplay.INT.udn
--- a/Engine/Documentation/Source/Gameplay/Gameplay.INT.udn
+++ b/Engine/Documentation/Source/Gameplay/Gameplay.INT.udn
@@ -5,6 +5,7 @@ Availability:Public
 version: 4.9
 parent:%ROOT%
 order:6
+type:landing
 
 [VAR:Params]
 [PARAM:image]
diff --git a/Engine/Documentation/Source/Gameplay/HowTo/UseRaycasts/GHT_2_Overview.INT.udn b/Engine/Documentation/Source/Gameplay/HowTo/UseRaycasts/GHT_2_Overview.INT.udn
--- a/Engine/Documentation/Source/Gameplay/HowTo/UseRaycasts/GHT_2_Overview.INT.udn
+++ b/Engine/Documentation/Source/Gameplay/HowTo/UseRaycasts/GHT_2_Overview.INT.udn
@@ -21,6 +21,7 @@ To learn more about Raycasts using Blueprints , follow the link below:
 * [](Gameplay\HowTo\UseRaycasts\Blueprints)
 [PUBLISH:Docs]
 * [](Gameplay\HowTo\UseRaycasts\Code)
+* [](Gameplay/HowTo/UseRaycasts/UVCoordsFromTrace)
 [/PUBLISH:Docs]
 [/REGION]
 
@@ -41,4 +42,3 @@ To learn more about Raycasts using Blueprints , follow the link below:
 
 
 
-
diff --git a/Engine/Documentation/Source/Gameplay/HowTo/UseRaycasts/UVCoordsFromTrace/UVCoordsFromTrace.INT.udn b/Engine/Documentation/Source/Gameplay/HowTo/UseRaycasts/UVCoordsFromTrace/UVCoordsFromTrace.INT.udn
new file mode 100644
--- /dev/null
+++ b/Engine/Documentation/Source/Gameplay/HowTo/UseRaycasts/UVCoordsFromTrace/UVCoordsFromTrace.INT.udn
@@ -0,0 +1,57 @@
+Availability:Docs
+Title:Getting UVs From Hit Results
+Crumbs: %ROOT%, Engine
+Description:This how-to steps you through enabling the Support UV From Hit Results feature of Unreal Engine 4.13.
+parent:Gameplay/HowTo/UseRaycasts
+order:1
+Version: 4.14
+SkillLevel: Beginner
+type:howto
+topic-image:trace_topic.png
+tags:UV
+tags:Physics
+tags:Trace
+
+Obtaining UV data from hit (collision) results has several potential use cases. While some users might want to use the collected UVs with Materials for rendering purposes, others might want to use 
+the UV data with the new **Blueprint Render to Texture** feature. See the [Unreal Live Training Stream covering](Videos/PLZlv_N0_O1ga0aV9jVqJgog0VWz1cLL5f/1Z-V1Buk6z8) **Blueprint Render to Texture** for details on how to use UV's from Traces.
+
+If you want to get UV data from hit results, you'll need to enable **Support UV From Hit Results**. Enabling this 
+feature in Unreal Engine 4 (UE4) is straightforward, requiring only one Blueprint node to utilize it. This guide shows you how to enable the feature, providing you a sample 
+graph that utilizes the **Find Collision UV** node.
+
+[REGION:note]
+Using this feature will increase CPU memory usage, as Unreal Engine 4 needs to keep an additional copy of positions and UV coordinates in main memory.
+[/REGION]
+
+[REGION:note]
+As of 4.14 this only works on **Static Mesh Components**, **Procedural Mesh Components**, and **BSP**. It will **not** work on **Skeletal Mesh Components** because you trace against the **Physics Asset** which doesn't have UV coordinates.
+[/REGION]
+
+## Steps
+
+1. Access your **Project Settings** from the **Edit Menu**.
+
+	![](accessProj.png)
+
+1. Enable the **Support UV From Hit Results** feature in the **Physics Section** of your **Project Settings**.
+
+	![](ProjectSettings.png)
+
+1. Restart the Editor.
+
+	[REGION:note]
+	You can inspect the Blueprint node using this feature prior to restarting the editor; however, the node will only return `0.0, 0.0` when you inspect it.
+	If you want the node to return the correct UV data, you'll have to restart the editor.
+	[/REGION]
+
+## Result
+
+With the feature enabled, you can input a **Hit Result Struct** into a **Find Collision UV** node inside of your Blueprints.
+
+![](traceExample.png)(convert:false)
+
+In the above example, we trace against UV Channel 0, but you can specify which channel you desire UV Coordinates from. This can give dramatically different results given how different UV channels can be from one to another.
+
+Also, make sure that **Trace Complex** is enabled. Otherwise, you're not going to hit a mesh in order to collect UV data from it.
+
+And finally, you should always check to see if you hit something before requesting UV Coordinates from the trace. In general, this will reduce processing overhead.
diff --git a/Engine/Documentation/Source/Gameplay/Localization/Formatting/Formatting.INT.udn b/Engine/Documentation/Source/Gameplay/Localization/Formatting/Formatting.INT.udn
--- a/Engine/Documentation/Source/Gameplay/Localization/Formatting/Formatting.INT.udn
+++ b/Engine/Documentation/Source/Gameplay/Localization/Formatting/Formatting.INT.udn
@@ -1,25 +1,35 @@
-﻿Availability:Docs
+﻿Availability:Public
 Title:Formatting User-Facing Text
 Crumbs:%ROOT%, Gameplay, Gameplay/Localization
 Description:Using the Format function to present localized text to your users
 Parent:Gameplay/Localization
 Order:3
 version: 4.13
+Related:Programming/UnrealArchitecture/StringHandling/FText
 
 FText is used to create user-facing text, which is often localized. This means that user interface text that changes dynamically over the course of gameplay can also be subject to grammar rules from a number of different languages. Using `FText::Format` and the corresponding 
 **Format Text** Blueprint node enables you to construct reusable format patterns to make the localization process easier and less prone to edge-case errors. Both functions support ordered arguments, named arguments, and specialized arguments based on plural forms, gender forms, 
 and Hangul post-positions.
 
-FText::Format is also available in Blueprints. The Format text can be a literal or connected to a pin. When used as a literal, it will auto-generate arguments for formatting. So by entering the literal "There are {number} days until {day}.", a Number pin and a Day pin will be 
-automatically added to the node. These pins can be filled in as a literal or from a variable and the output Text will replace the arguments with whatever values are hooked to the pins.
+In C++, you can pre-compile your `FText` format pattern into an `FTextFormat` to reuse the compiled format pattern. This is useful if you're reusing a format pattern frequently, and it also automatically recompiles if the source format pattern changes due to a culture change.
 
-When the Format pin is connected, the Format can not be deciphered because it is variable. In this case, you can add pins yourself and name them after the possible arguments. In this way, you can set up the Format Text node to be ready for multiple possibilities. Whenever an argument 
+You can also format FText in Blueprints using the **Format Text** node. The **Format** text can be a literal or connected to a pin. When used as a literal, it will auto-generate arguments for formatting. So by entering the literal "There are {number} days until {day}.", a **Number** pin and a **Day** pin will be 
+automatically added to the node. These pins can be filled in as a literal or from a variable and the output **Text** will replace the arguments with whatever values are hooked to the pins.
+
+When the **Format** pin is connected to a variable, the **Format** can not be deciphered because it can possibly change over the course of gameplay. In this case, you can add pins yourself and name them after the possible arguments. In this way, you can set up the **Format Text** node to be ready for multiple possibilities. Whenever an argument 
 matches an argument in the Format variable, it will utilize the provided information.
 
-Pin names can be changed, added, and removed from the Details Panel.
+Pin names can be changed, added, and removed from the **Details** panel.
+
+[REGION:imagetable]
+|![](FormatNode.png)(h:400) | ![](FormatNodeDetails.png)(h:400)|
+|--- |---- |
+|**Format Text** node | **Details** panel | 
+[/REGION]
+
 
 [REGION:tip]
-For users using a source build of Unreal Engine, `TextTest.cpp` can show you examples of how to perform valid culture aware formatting, comparison, and sorting.
+For those using a source build of Unreal Engine, `TextTest.cpp` can show you examples of how to perform valid culture aware formatting, comparison, and sorting.
 [/REGION]
 
 ## Ordered Arguments
@@ -57,7 +67,7 @@ The following C++ and Blueprint snippets show a named format argument to display
 
 ## Specialized Arguments
 
-To allow you to pass in the numeric/gender values needed for plural/gender form support, all of the FText::Format(...) family of functions now take their values as FFormatArgumentValue rather than FText. This can be implicitly constructed from any numeric type, ETextGender, or FText.
+To allow you to pass in the numeric/gender values needed for plural/gender form support, all of the `FText::Format(...)` family of functions now take their values as FFormatArgumentValue rather than FText. This can be implicitly constructed from any numeric type, ETextGender, or FText. 
 
 The ability to set these value types in Blueprints has been exposed using wildcard pins on the **Format Text** node.
 
diff --git a/Engine/Documentation/Source/Gameplay/Localization/Localization.INT.udn b/Engine/Documentation/Source/Gameplay/Localization/Localization.INT.udn
--- a/Engine/Documentation/Source/Gameplay/Localization/Localization.INT.udn
+++ b/Engine/Documentation/Source/Gameplay/Localization/Localization.INT.udn
@@ -1,17 +1,19 @@
 Availability:Public
 Title: Localization
-Crumbs:%ROOT%, Gameplay
 Description:Information about how to localize your project.
 Related: Programming/UnrealArchitecture/StringHandling/FText
-version: 4.9
+version: 4.13
 
 
 ## Localizing Text
 [DIR(output:"listbutton" parent:"Gameplay/Localization")]
 
 
+## Localizing Audio
 
-[REGION:tip]
-Asset localization will be coming in the future just before proper audio localization support. 
-The implementation will revolve around swapping out the entire uasset package with locale specific versions. 
-[/REGION:tip]
+Spoken text is now gathered from [**Dialogue Wave**](Engine/Audio/Overview#dialoguevoiceanddialoguewave) assets and included in your PO files for translation along with your other asset based text.
+
+Several new [commandlets](Gameplay/Localization/Commandlets) have been added to facilitate the recording and importing of audio for your **Dialogue Wave** assets.
+* ExportDialogueScript - Exports a per-culture CSV file containing each line of dialogue and its associated filename.
+* ImportDialogueScript - Imports a copy-edited per-culture CSV file and updates any translations with the recorded dialogue.
+* ImportLocalizedDialogue - Imports per-culture recorded WAV files and generates localized dialogue and sound wave assets.
diff --git a/Engine/Documentation/Source/Gameplay/Networking/Networking.INT.udn b/Engine/Documentation/Source/Gameplay/Networking/Networking.INT.udn
--- a/Engine/Documentation/Source/Gameplay/Networking/Networking.INT.udn
+++ b/Engine/Documentation/Source/Gameplay/Networking/Networking.INT.udn
@@ -29,6 +29,7 @@ version: 4.9
 			* [](Gameplay/Networking/Actors/Roles "%Gameplay/Networking/Actors/Roles:Description%")
 			* [](Gameplay/Networking/Actors/OwningConnections "%Gameplay/Networking/Actors/OwningConnections:Description%")
 			* [](Gameplay/Networking/Actors/ReplicationFlow "%Gameplay/Networking/Actors/ReplicationFlow:Description%")
+		* [](Gameplay/Networking/OnlineBeacons "%Gameplay/Networking/OnlineBeacons:Description%")
 		* [](Gameplay/Networking/Travelling "%Gameplay/Networking/Travelling:Description%")
 		* [](Gameplay/Networking/CharacterMovementComponent "%Gameplay/Networking/CharacterMovementComponent:Description%")
 		* [](Gameplay/Networking/Actors/ReplicationPerformance "%Gameplay/Networking/Actors/ReplicationPerformance:Description%")
diff --git a/Engine/Documentation/Source/Gameplay/Networking/OnlineBeacons/OnlineBeacons.INT.udn b/Engine/Documentation/Source/Gameplay/Networking/OnlineBeacons/OnlineBeacons.INT.udn
new file mode 100644
--- /dev/null
+++ b/Engine/Documentation/Source/Gameplay/Networking/OnlineBeacons/OnlineBeacons.INT.udn
@@ -0,0 +1,38 @@
+Availability: Public
+Title: Online Beacons
+Crumbs: %ROOT%, Engine
+Description:Overview of the Online Beacon system
+Version: 4.13
+parent:Engine
+Related: Gameplay/Networking/Actors/RPCs
+Related: Gameplay/Networking/Actors/Properties
+
+[TOC (start:2 end:3)]
+
+**Online Beacons** are a special type of **Actor** that provide a lightweight way to contact a server and interact with it (via RPCs) without committing to a normal game connection. While the built-in classes can be used as they are in some cases, they are intended to be extended into custom classes that perform project-specific interactions, logic, and information requests.
+
+## Online Beacon Base Classes
+Common actions performed by beacon classes are to request quality-of-service information, reserve a slot in a game that the client wishes to join, receive a list of the names of players in a game, find out the score and running time of a game in progress, and so on. The following classes are provided with the engine, and lay the foundation for the Online Beacon system:
+
+### AOnlineBeacon
+This is the base class for `AOnlineBeaconClient` and `AOnlineBeaconHost`. It derives from `AActor` directly.
+
+### AOnlineBeaconHost
+This class uses its own `UNetDriver` to listen for incoming Online Beacon connections from remote client machines. When it receives one, it looks through its list of registered `AOnlineBeaconHostObject` instances to find the one that matches the incoming client and hands the connection off to it. This class will not usually need to be derived, since it only manages the initial connection between the client and the registered `AOnlineBeaconHostObject`.
+
+### AOnlineBeaconClient
+Children of this class connect to hosts and make the actual RPCs. One of these will be spawned on the client machine, and one will be spawned on the server by the appropriate `AOnlineBeaconHostObject` that is registered with the server's `AOnlineBeaconHost`.  By default, the output of the `GetBeaconType` function (which will be the class name) will be used to match instances of this class to the registered instance of the appropriate host object class. Note that this is different from how Actor spawning usually works, where the server spawns the Actor and then replicates it to the client. However, once the connection between the client's and server's copy of the object has been made, replication of the object behaves as normal, with either side able to make RPCs to the other, and the server's version of the object being authoritative with regard to property replication.
+The base class implements `OnConnected` and `OnFailure` functions that can be overridden by child classes to perform RPCs when connected, or to handle failures to connect. This class is the workhorse of the Online Beacon system, and will do the client-side work that your Beacon requires. Since a second instance, synced with the original, will be spawned on the server in the event of a successful connection, this class can also do server-side work, coordinating and communicating through client and server RPCs, or server-to-client replicated properties.
+
+### AOnlineBeaconHostObject
+This class should also be overridden so that it can be paired with an overridden `AOnlineBeaconClient` class. Pairing is accomplished by matching the return value of the client's `GetBeaconType` with the value stored in the `BeaconTypeName` member variable. When the server's `AOnlineBeaconHost` detects a paired `AOnlineBeaconHostObject` for an incoming `AOnlineBeaconClient`, it will instruct the `AOnlineBeaconHostObject` to spawn a local copy of the `AOnlineBeaconClient` via the virtual `SpawnBeaconActor` function. By default, this function uses the `ClientBeaconActorClass` member variable to determine the class of actor to spawn, which should be set to the paired `AOnlineBeaconClient` class. It also calls `SetBeaconOwner` on the server's copy of the spawned object, so that the server-side instance of the client object can communicate with the host object. Most of this setup is established in the base class, and does not need to be overridden.
+
+## Online Beacon Example
+There are a variety of premade Beacon class pairs already in the engine, but the Test Beacon is a good, simple example of how to build your own. `ATestBeaconClient` is derived from `AOnlineBeaconClient`, and `ATestBeaconHost` is derived from `AOnlineBeaconHostObject`. The Test Beacon simply makes repeated function calls back and forth between the client and server. Here is how the classes work together:
+1. In `ATestBeaconHost`'s constructor, `ClientBeaconActorClass` is set to the class of `ATestBeaconClient`, and `BeaconTypeName` is set to the name (obtained by calling the `GetName` function) of `ClientBeaconActorClass`. This ensures that the `AOnlineBeaconHost` will pair the two Test Beacon classes correctly.
+1. When the client's `ATestBeaconClient` successfully contacts the server, the server's `AOnlineBeaconHost` instance will direct the registered `ATestBeaconHost` to spawn its own `ATestBeaconClient` and associate it with the `ATestBeaconClient` on the client. This is important because it allows client and server RPCs to pass between the two `ATestBeaconClient` instances.
+1. The server's `ATestBeaconHost` will then call `OnClientConnected`, a client RPC, on its `ATestBeaconClient` instance. Since this function is a client RPC, the command to call the function will be passed to the client machine rather than calling the function on the server.
+1. On the client's `ATestBeaconClient` instance, the RPC will run the virtual `OnConnected` function, which has been overridden to call `ClientPing`.
+1. `ClientPing` will output some text to the debug log and then call `ServerPong`, a server RPC, which will then be sent across the network connection and executed on the server's `ATestBeaconClient` instance.
+1. `ServerPong` will output debug text and then call `ClientPing`, again triggering a networked RPC call to the client.
+1. This process will continue indefinitely. While this is usually not desirable, especially since these RPCs aren't really doing anything, it could be used to keep track of a server's ping time by recording the time between two subsequent runs of the same RPC. There could also be a delay, so that this ping check is not done as frequently as the computer and network will allow, or the client or server could simply stop pinging back after a certain number of pings. Other use cases might involve only a single RPC call back to the client, rather than a continuing loop, such as asking a question along the lines of "What map is being used for this server's game?" or "Is this a Capture The Flag server?" depending on what information might be relevant to your specific game or project.
diff --git a/Engine/Documentation/Source/Platforms/Android/Android.INT.udn b/Engine/Documentation/Source/Platforms/Android/Android.INT.udn
--- a/Engine/Documentation/Source/Platforms/Android/Android.INT.udn
+++ b/Engine/Documentation/Source/Platforms/Android/Android.INT.udn
@@ -20,7 +20,7 @@ type:landing
 	[PARAM:description]
 	[/PARAM]
 	[PARAM:links]
-		[DIR(output:"topiccompact" parent:"Platforms/Android" type:"quickstart")]
+		[DIR(output:"topiccompact" parent:"Platforms/Android" type:"quick start")]
 	[/PARAM]
 [/OBJECT]
 
@@ -34,7 +34,7 @@ type:landing
 	[PARAM:description]
 	[/PARAM]
 	[PARAM:links]
-		[DIR(end:"1" output:"listbutton" type:"howto" parent:"Platforms/Android")]
+		[DIR(end:"1" output:"listbutton" type:"how-to" parent:"Platforms/Android")]
 	[/PARAM]
 [/OBJECT]
 
diff --git a/Engine/Documentation/Source/Platforms/Android/OpenGLES31MobileRenderer/OpenGLES31MobileRenderer.INT.udn b/Engine/Documentation/Source/Platforms/Android/OpenGLES31MobileRenderer/OpenGLES31MobileRenderer.INT.udn
--- a/Engine/Documentation/Source/Platforms/Android/OpenGLES31MobileRenderer/OpenGLES31MobileRenderer.INT.udn
+++ b/Engine/Documentation/Source/Platforms/Android/OpenGLES31MobileRenderer/OpenGLES31MobileRenderer.INT.udn
@@ -1,7 +1,7 @@
 Title: Android OpenGL ES 3.1 Mobile Renderer
 Description:Setting up your UE4 Android project to work with the OpenGL ES 3.1 mobile render
 Crumbs:%ROOT%, Platforms, Platforms/Android, Platforms/Android/OpenGLES31MobileRenderer/
-Availability:Docs
+Availability:Public
 version:4.12
 Parent:Platforms/Android
 tags:Mobile
diff --git a/Engine/Documentation/Source/Platforms/GearVR/GearVR.INT.udn b/Engine/Documentation/Source/Platforms/GearVR/GearVR.INT.udn
--- a/Engine/Documentation/Source/Platforms/GearVR/GearVR.INT.udn
+++ b/Engine/Documentation/Source/Platforms/GearVR/GearVR.INT.udn
@@ -11,10 +11,11 @@ type:landing
 parent:Platforms/VR
 tags:VR
 tags:GearVR
+tags:Platform
 
 ## Starting Out
 
-[DIR(output:"listbutton" type:"quickstart" parent:"Platforms/GearVR")]
+[DIR(output:"listbutton" type:"quick start" parent:"Platforms/GearVR")]
 
 ## GearVR Topics
 
diff --git a/Engine/Documentation/Source/Platforms/GoogleVR/GoogleVR.INT.udn b/Engine/Documentation/Source/Platforms/GoogleVR/GoogleVR.INT.udn
--- a/Engine/Documentation/Source/Platforms/GoogleVR/GoogleVR.INT.udn
+++ b/Engine/Documentation/Source/Platforms/GoogleVR/GoogleVR.INT.udn
@@ -9,6 +9,7 @@ type:landing
 parent:Platforms/VR
 tags:VR
 tags:Google VR
+tags:Platform
 prereq: Platforms/Android/GettingStarted/
 social-image:googlevr_social.png
 
@@ -16,7 +17,7 @@ social-image:googlevr_social.png
 
 The GoogleVR Quick Start will walk you through getting your PC and Android-based smartphone to work with UE4 and Google VR.
 
-[DIR(output:"listbutton" type:"quickstart" parent:"Platforms/GoogleVR")]
+[DIR(output:"listbutton" type:"quick start" parent:"Platforms/GoogleVR")]
 
 ## Google VR Daydream
 
diff --git a/Engine/Documentation/Source/Platforms/GoogleVR/QuickStart/1/GoogleVRQuickStart_1.INT.udn b/Engine/Documentation/Source/Platforms/GoogleVR/QuickStart/1/GoogleVRQuickStart_1.INT.udn
--- a/Engine/Documentation/Source/Platforms/GoogleVR/QuickStart/1/GoogleVRQuickStart_1.INT.udn
+++ b/Engine/Documentation/Source/Platforms/GoogleVR/QuickStart/1/GoogleVRQuickStart_1.INT.udn
@@ -7,7 +7,7 @@ Skilllevel: Intermediate
 Version: 4.12
 Parent:Platforms/GoogleVR/QuickStart
 order:1
-type:step
+type:multi-step
 tags:VR
 tags:Google VR
 
diff --git a/Engine/Documentation/Source/Platforms/GoogleVR/QuickStart/2/GoogleVRQuickStart_2.INT.udn b/Engine/Documentation/Source/Platforms/GoogleVR/QuickStart/2/GoogleVRQuickStart_2.INT.udn
--- a/Engine/Documentation/Source/Platforms/GoogleVR/QuickStart/2/GoogleVRQuickStart_2.INT.udn
+++ b/Engine/Documentation/Source/Platforms/GoogleVR/QuickStart/2/GoogleVRQuickStart_2.INT.udn
@@ -7,7 +7,7 @@ Skilllevel: Intermediate
 Version: 4.12
 Parent:Platforms/GoogleVR/QuickStart
 order:2
-type:step
+type:multi-step
 tags:VR
 tags:Google VR
 
@@ -92,7 +92,7 @@ tags:Google VR
 	[/PARAMLITERAL]
 	[/OBJECT]
 
-1. In the **Platforms** section under the **Android** section press the **Configure Now** button in the **APKPackaging** and **Google Play Services** sections to ensure that our project can .
+1. In the **Platforms** section under the **Android** section press the **Configure Now** button in the **APKPackaging** and **Google Play Services** sections to ensure that our project can package properly.
 
 	![](GVRQS_Config_Now_00.png)
 
diff --git a/Engine/Documentation/Source/Platforms/GoogleVR/QuickStart/3/GoogleVRQuickStart_3.INT.udn b/Engine/Documentation/Source/Platforms/GoogleVR/QuickStart/3/GoogleVRQuickStart_3.INT.udn
--- a/Engine/Documentation/Source/Platforms/GoogleVR/QuickStart/3/GoogleVRQuickStart_3.INT.udn
+++ b/Engine/Documentation/Source/Platforms/GoogleVR/QuickStart/3/GoogleVRQuickStart_3.INT.udn
@@ -7,7 +7,7 @@ Skilllevel: Intermediate
 Version: 4.12
 Parent:Platforms/GoogleVR/QuickStart
 order:3
-type:step
+type:multi-step
 tags:VR
 tags:Google VR
 
diff --git a/Engine/Documentation/Source/Platforms/Linux/BeginnerLinuxDeveloper/SettingUpAnIDE/1/SettingUpAnIDE_1.INT.udn b/Engine/Documentation/Source/Platforms/Linux/BeginnerLinuxDeveloper/SettingUpAnIDE/1/SettingUpAnIDE_1.INT.udn
new file mode 100644
--- /dev/null
+++ b/Engine/Documentation/Source/Platforms/Linux/BeginnerLinuxDeveloper/SettingUpAnIDE/1/SettingUpAnIDE_1.INT.udn
@@ -0,0 +1,174 @@
+Availability: Docs
+Title: 1 - Configuring your IDE to Build UE4
+Crumbs:
+Description:Learn how to configure your IDE to build UE4 in Debug and Release mode.
+Version: 4.14
+type:how to
+parent:Platforms/Linux/BeginnerLinuxDeveloper/SettingUpAnIDE
+order: 1
+tags:Linux
+SkillLevel:Beginner
+
+[VAR:Steps]
+[OBJECT:Navigation]
+	[PARAM:previous]
+	[/PARAM]
+	[PARAM:current]
+	[/PARAM]
+	[PARAM:home]
+		[Home](Platforms/Linux/BeginnerLinuxDeveloper/SettingUpAnIDE)
+	[/PARAM]
+	[PARAM:next]
+		[Next Step](Platforms/Linux/BeginnerLinuxDeveloper/SettingUpAnIDE/2 "%Platforms/Linux/BeginnerLinuxDeveloper/SettingUpAnIDE/2:title%")
+	[/PARAM]
+[/OBJECT]
+[/VAR]
+
+%Steps%
+
+While some developers prefer to use text editors and command line debuggers to modify their game's source code, most developers 
+prefer to use an Integrated Development Environment (IDE) to edit, debug, and compile their code. If you're a game developer looking to setup 
+an IDE to work with Unreal Engine 4 (UE4) in a Linux environment, this guide is for you. 
+
+Although we recognize that chosing an IDE is a matter of individual preference, we've configured UE4 to interface with three IDEs in Linux; 
+specifically, Qt Creator, KDevelop, and Codelite. We don't endorse one IDE over another; however, as an example for learning purposes, we've chosen to 
+show you how to configure Qt Creator to build and run UE4. By the end of this guide, you'll know how to configure your IDE to build and run UE4. Finally, 
+you'll also learn how to setup the IDE to build and run new C++ (CPP) projects that you create from the Unreal Editor.
+
+[REGION:note]
+Before we begin, please make sure that you've already setup your workflow using our [](Platforms/Linux/BeginnerLinuxDeveloper/SettingUpAnUnrealWorkflow). Also, please make sure that [Qt Creator](https://wiki.qt.io/Install_Qt_5_on_Ubuntu) is
+ installed on your machine.
+[/REGION]
+
+1.	Launch Qt Creator, opening the **Qt Welcome** menu.  
+
+	[REGION:lightbox]
+		[![](SetupIDE_Step1.png)(w:720)](SetupIDE_Step1.png)
+	[/REGION]
+	
+	[REGION:caption]
+		Click for full image.
+	[/REGION]
+
+1.	Click the **Open Project** button to open the Qt Project file browser.  
+
+	[REGION:lightbox]
+		[![](SetupIDE_Step2.png)(w:720)](SetupIDE_Step2.png)
+	[/REGION]
+	
+	[REGION:caption]
+		Click for full image.
+	[/REGION]
+
+1.	Navigate to the directory where you built UE4, select the UE4 Qt Project (**UE4.pro**) file, and click the **Open** button.
+
+	[REGION:lightbox]
+		[![](SetupIDE_Step3.png)(w:720)](SetupIDE_Step3.png)
+	[/REGION]
+	
+	[REGION:caption]
+		Click for full image.
+	[/REGION]
+
+1.	At this point, Qt Creator hasn't been configured to build UE4 in Debug and Release Mode. Click the **Configure Project** button located in the **Configure Project** menu.
+
+	[REGION:lightbox]
+		[![](SetupIDE_Step4.png)(w:720)](SetupIDE_Step4.png)
+	[/REGION]
+	
+	[REGION:caption]
+		Click for full image.
+	[/REGION]
+
+1.	Now, click the **Projects** button to enter into Projects mode, opening the **Build Settings** menu.
+
+	[REGION:lightbox]
+		[![](SetupIDE_Step5.png)(w:720)](SetupIDE_Step5.png)
+	[/REGION]
+	
+	[REGION:caption]
+		Click for full image.
+	[/REGION]
+
+	[REGION:note]
+	When you open the **Build Settings** menu, you'll notice that Qt Creator defaults its **Edit build configuration: ** to **Debug** mode.
+	[region: lightbox]
+	[![](SetupIDE_Step5_Note.png)(w:347 h:249)](SetupIDE_Step5_Note.png)
+	[/region]
+	[REGION:caption]
+	Click for full image.
+	[/REGION]
+	[/REGION]
+
+1.	Under the **Build Steps** section, clear the **qmake** build step by clicking on the **X** button next to the **Details** drop down menu.
+
+	[REGION:lightbox]
+		[![](SetupIDE_Step6Alt.png)(w:720)](SetupIDE_Step6Alt.png)
+	[/REGION]
+	
+	[REGION:caption]
+		Click for full image.
+	[/REGION]
+
+1.	Now, clear the **Shadow build** checkbox to cancel the selection. 
+
+	[REGION:lightbox]
+		[![](SetupIDE_Step7.png)(w:720)](SetupIDE_Step7.png)
+	[/REGION]
+	
+	[REGION:caption]
+		Click for full image.
+	[/REGION]
+
+	[REGION:note]
+	Shadow building allows you to build projects in a separate build directory with Qt Creator's qmake utility. Because UE4 is distributed with a custom Make file that makes use of 
+	[](Programming/UnrealBuildSystem), you don't need to use Qt Creator's qmake utility.
+	[/REGION]
+
+1.	Now, click the **Details** drop down menu to open the input form for the **Make** build step.
+
+	![](SetupIDE_Step8.png)
+
+1.	Enter `UE4Editor` into the **Make arguments:** text field.
+
+	![](SetupIDE_Step9.png)
+
+1.	Under the **Clean Steps** section, click the **Details** drop down menu to open the input form for the **Make** clean step.
+
+	![](SetupIDE_Step10.png)
+
+1.	Replace `clean` by entering `UE4Editor ARGS="-clean"` into the text field.
+
+	![](SetupIDE_Step11.png)
+
+1.	Now, click the **Edit build configuration:** drop down menu and select **Release** to open the **Build Settings** for Release mode.
+
+	[REGION:lightbox]
+		[![](SetupIDE_Step12.png)(w:720)](SetupIDE_Step12.png)
+	[/REGION]
+	
+	[REGION:caption]
+		Click for full image.
+	[/REGION]
+
+1.	Repeat Steps 6 through 11 to edit the build configuration settings for Release mode.
+
+	![](SetupIDE_Step13.png)
+
+	[REGION:note]
+	Read through our [](Programming/Development/BuildConfigurations) to learn more about the different configurations you can set your build environment to.
+	[/REGION]
+
+## End Result
+
+[REGION:lightbox]
+	[![](QtBuildSettings_UE4.png)(w:720)](QtBuildSettings_UE4.png)
+[/REGION]
+
+[REGION:caption]
+	Click for full image.
+[/REGION]
+
+At this point, you've configured UE4 to build in both Debug and Release mode. Now, you're ready to configure Qt Creator to run UE4.
+
+%Steps%
\ No newline at end of file
diff --git a/Engine/Documentation/Source/Platforms/Linux/BeginnerLinuxDeveloper/SettingUpAnIDE/2/SettingUpAnIDE_2.INT.udn b/Engine/Documentation/Source/Platforms/Linux/BeginnerLinuxDeveloper/SettingUpAnIDE/2/SettingUpAnIDE_2.INT.udn
new file mode 100644
--- /dev/null
+++ b/Engine/Documentation/Source/Platforms/Linux/BeginnerLinuxDeveloper/SettingUpAnIDE/2/SettingUpAnIDE_2.INT.udn
@@ -0,0 +1,94 @@
+Availability: Docs
+Title: 2 - Configuring your IDE to Run UE4
+Crumbs:
+Description:Learn how to configure your IDE to run UE4 in Debug and Release mode.
+Version: 4.14
+type:how to
+parent:Platforms/Linux/BeginnerLinuxDeveloper/SettingUpAnIDE
+order: 2
+tags:Linux
+SkillLevel:Beginner
+
+[VAR:Steps]
+[OBJECT:Navigation]
+	[PARAM:previous]
+		[Previous Step](Platforms/Linux/BeginnerLinuxDeveloper/SettingUpAnIDE/1 "%Platforms/Linux/BeginnerLinuxDeveloper/SettingUpAnIDE/1:title%")
+	[/PARAM]
+	[PARAM:current]
+	[/PARAM]
+	[PARAM:home]
+		[Home](Platforms/Linux/BeginnerLinuxDeveloper/SettingUpAnIDE)
+	[/PARAM]
+	[PARAM:next]
+		[Next Step](Platforms/Linux/BeginnerLinuxDeveloper/SettingUpAnIDE/3 "%Platforms/Linux/BeginnerLinuxDeveloper/SettingUpAnIDE/3:title%")
+	[/PARAM]
+[/OBJECT]
+[/VAR]
+
+%Steps%
+
+At this point, you've configured UE4 to build in both Debug and Release mode. Now, you're ready to configure Qt Creator to run UE4.
+
+1.	Click the **Run** button to open the **Run Settings** menu.
+
+	[REGION:lightbox]
+		[![](SetupIDE_Step10.png)(w:720)](SetupIDE_Step10.png)
+	[/REGION]
+	
+	[REGION:caption]
+		Click for full image.
+	[/REGION]
+
+1.	Under the **Run** section, click the **Browse** button that is located at the end of the **Executable:** text field.
+
+	[REGION:lightbox]
+		[![](SetupIDE_Step11.png)(w:720)](SetupIDE_Step11.png)
+	[/REGION]
+	
+	[REGION:caption]
+		Click for full image.
+	[/REGION]
+
+1.	Navigate to the directory where the Editor's binary path is located, select **UE4Editor**, and click the **Open** button.
+
+	[REGION:lightbox]
+		[![](SetupIDE_Step12.png)(w:720)](SetupIDE_Step12.png)
+	[/REGION]
+	
+	[REGION:caption]
+		Click for full image.
+	[/REGION]
+
+1.	Now, click the **Browse** button that is located at the end of the **Working directory:** text field.
+
+	[REGION:lightbox]
+		[![](SetupIDE_Step13.png)(w:720)](SetupIDE_Step13.png)
+	[/REGION]
+	
+	[REGION:caption]
+		Click for full image.
+	[/REGION]
+
+1.	Navigate to the directory where the Editor's binary path is located, select `[UE4ROOT]\Engine\Binaries\Linux`, and click the **Choose** button.
+
+	[REGION:lightbox]
+		[![](SetupIDE_Step14.png)(w:720)](SetupIDE_Step14.png)
+	[/REGION]
+	
+	[REGION:caption]
+		Click for full image.
+	[/REGION]
+
+## End Result
+
+[REGION:lightbox]
+	[![](QtRunSettings_UE4.png)(w:720)](QtRunSettings_UE4.png)
+[/REGION]
+
+[REGION:caption]
+	Click for full image.
+[/REGION]
+
+You're now ready to build and run UE4 with Qt Creator in both Debug and Release mode. 
+
+%Steps%
\ No newline at end of file
diff --git a/Engine/Documentation/Source/Platforms/Linux/BeginnerLinuxDeveloper/SettingUpAnIDE/3/SettingUpAnIDE_3.INT.udn b/Engine/Documentation/Source/Platforms/Linux/BeginnerLinuxDeveloper/SettingUpAnIDE/3/SettingUpAnIDE_3.INT.udn
new file mode 100644
--- /dev/null
+++ b/Engine/Documentation/Source/Platforms/Linux/BeginnerLinuxDeveloper/SettingUpAnIDE/3/SettingUpAnIDE_3.INT.udn
@@ -0,0 +1,85 @@
+Availability: Docs
+Title: 3 - Building UE4 from your IDE
+Crumbs:
+Description:Learn how to build Unreal Engine from your IDE.
+Version: 4.14
+type:how to
+parent:Platforms/Linux/BeginnerLinuxDeveloper/SettingUpAnIDE
+order: 3
+tags:Linux
+SkillLevel:Beginner
+
+[VAR:Steps]
+[OBJECT:Navigation]
+	[PARAM:previous]
+		[Previous Step](Platforms/Linux/BeginnerLinuxDeveloper/SettingUpAnIDE/2 "%Platforms/Linux/BeginnerLinuxDeveloper/SettingUpAnIDE/2:title%")
+	[/PARAM]
+	[PARAM:current]
+	[/PARAM]
+	[PARAM:home]
+		[Home](Platforms/Linux/BeginnerLinuxDeveloper/SettingUpAnIDE)
+	[/PARAM]
+	[PARAM:next]
+		[Next Step](Platforms/Linux/BeginnerLinuxDeveloper/SettingUpAnIDE/4 "%Platforms/Linux/BeginnerLinuxDeveloper/SettingUpAnIDE/4:title%")
+	[/PARAM]
+[/OBJECT]
+[/VAR]
+
+%Steps%
+
+Now that you've configured Qt Creator to build and run UE4, you're ready to run UE4 from the IDE. Before you run UE4 from the IDE, you'll want to open some output windows so that you can see 
+Unreal Build System at work.
+
+1.	Click the **Edit** button to enter into Edit mode, opening Qt Creator's editor.
+
+	[REGION:lightbox]
+		[![](SetupIDE_Step15.png)(w:720)](SetupIDE_Step15.png)
+	[/REGION]
+	
+	[REGION:caption]
+		Click for full image.
+	[/REGION]
+
+	[REGION:note]
+	While in Edit mode, you can view, debug, and edit the CPP source code powering UE4.
+	[/REGION]
+
+1.	By default, Qt Creator hides the main toolbar. Hover your mouse cursor over the top of the screen, touching the top edge to reveal the toolbar.
+
+	![](SetupIDE_Step15A.png)
+
+1.	Now, open the **Window** menu, hovering over the **Output Panes** sub-menu, clicking **Compile Output**. Alternatively, you can press the **Alt+4** keyboard combination to open the **Compile Output** window.
+
+	![](SetupIDE_Step15B.png)
+
+	[REGION:note]
+	Opening the **Compile Output** window allows you to watch for any errors while building UE4 from source.
+	[/REGION]
+
+## Building UE4
+	
+1.	Hover your mouse cursor over the top of the screen, touching the top edge to reveal the toolbar.
+
+	![](SetupIDE_Step15A.png)
+
+1.	To build UE4, open the **Build** menu and click the **Build Project "UE4"** button. Alternatively, you can press the **Ctrl+B** keyboard combination to build UE4.
+
+	![](SetupIDE_Step15C.png)
+
+## End Result
+
+[REGION:lightbox]
+	[![](SetupIDE_Step16.png)(w:720)](SetupIDE_Step16.png)
+[/REGION]
+
+[REGION:caption]
+	Click for full image.
+[/REGION]
+
+At this point, you should see the compiler flushing messages to the **Compile Output** window inside of the IDE. 
+Depending on your system's specifications, it may take anywhere from ten minutes to over an hour to compile the engine. 
+If you want to shorten the time it takes to compile the engine from source, we recommend compiling the source code on a 
+machine that has at least eight (8) gigabytes of RAM with a multi-core processor having at least eight (8) cores (including 
+hyperthreading).
+
+%Steps%
\ No newline at end of file
diff --git a/Engine/Documentation/Source/Platforms/Linux/BeginnerLinuxDeveloper/SettingUpAnIDE/4/SettingUpAnIDE_4.INT.udn b/Engine/Documentation/Source/Platforms/Linux/BeginnerLinuxDeveloper/SettingUpAnIDE/4/SettingUpAnIDE_4.INT.udn
new file mode 100644
--- /dev/null
+++ b/Engine/Documentation/Source/Platforms/Linux/BeginnerLinuxDeveloper/SettingUpAnIDE/4/SettingUpAnIDE_4.INT.udn
@@ -0,0 +1,54 @@
+Availability: Docs
+Title: 4 - Running UE4 from your IDE
+Crumbs:
+Description:Learn how to run Unreal Engine from your IDE.
+Version: 4.14
+type:how to
+parent:Platforms/Linux/BeginnerLinuxDeveloper/SettingUpAnIDE
+order: 4
+tags:Linux
+SkillLevel:Beginner
+
+[VAR:Steps]
+[OBJECT:Navigation]
+	[PARAM:previous]
+		[Previous Step](Platforms/Linux/BeginnerLinuxDeveloper/SettingUpAnIDE/3 "%Platforms/Linux/BeginnerLinuxDeveloper/SettingUpAnIDE/3:title%")
+	[/PARAM]
+	[PARAM:current]
+	[/PARAM]
+	[PARAM:home]
+		[Home](Platforms/Linux/BeginnerLinuxDeveloper/SettingUpAnIDE)
+	[/PARAM]
+	[PARAM:next]
+		[Next Step](Platforms/Linux/BeginnerLinuxDeveloper/SettingUpAnIDE/5 "%Platforms/Linux/BeginnerLinuxDeveloper/SettingUpAnIDE/5:title%")
+	[/PARAM]
+[/OBJECT]
+[/VAR]
+
+%Steps%
+
+Now that you've compiled the engine, you're ready to run the editor. During this step, you're going to use the IDE to run UE4.
+
+1.	To run UE4 from the IDE, click the **Run** button. Alternatively, you can press the **Ctrl+R** keyboard combination to run UE4 from the IDE.
+
+	[REGION:lightbox]
+		[![](SetupIDE_Step17.png)(w:720)](SetupIDE_Step17.png)
+	[/REGION]
+	
+	[REGION:caption]
+		Click for full image.
+	[/REGION]
+
+## End Result
+
+[REGION:lightbox]
+	[![](SetupIDE_Step18.png)(w:720)](SetupIDE_Step18.png)
+[/REGION]
+
+[REGION:caption]
+	Click for full image.
+[/REGION]
+
+Congratulations! You should now see UE4 running from your IDE.
+
+%Steps%
\ No newline at end of file
diff --git a/Engine/Documentation/Source/Platforms/Linux/BeginnerLinuxDeveloper/SettingUpAnIDE/5/SettingUpAnIDE_5.INT.udn b/Engine/Documentation/Source/Platforms/Linux/BeginnerLinuxDeveloper/SettingUpAnIDE/5/SettingUpAnIDE_5.INT.udn
new file mode 100644
--- /dev/null
+++ b/Engine/Documentation/Source/Platforms/Linux/BeginnerLinuxDeveloper/SettingUpAnIDE/5/SettingUpAnIDE_5.INT.udn
@@ -0,0 +1,93 @@
+Availability: Docs
+Title: 5 - Creating a Blueprint Project from your IDE
+Crumbs:
+Description:Learn how to create and configure a Blueprint project from your IDE.
+Version: 4.14
+type:how to
+parent:Platforms/Linux/BeginnerLinuxDeveloper/SettingUpAnIDE
+order: 5
+tags:Linux
+SkillLevel:Beginner 
+
+[VAR:Steps]
+[OBJECT:Navigation]
+	[PARAM:previous]
+		[Previous Step](Platforms/Linux/BeginnerLinuxDeveloper/SettingUpAnIDE/4 "%Platforms/Linux/BeginnerLinuxDeveloper/SettingUpAnIDE/4:title%")
+	[/PARAM]
+	[PARAM:current]
+	[/PARAM]
+	[PARAM:home]
+		[Home](Platforms/Linux/BeginnerLinuxDeveloper/SettingUpAnIDE)
+	[/PARAM]
+	[PARAM:next]
+		[Next Step](Platforms/Linux/BeginnerLinuxDeveloper/SettingUpAnIDE/6 "%Platforms/Linux/BeginnerLinuxDeveloper/SettingUpAnIDE/6:title%")
+	[/PARAM]
+[/OBJECT]
+[/VAR]
+
+%Steps%
+
+You can create either Blueprint or CPP projects from the [](Engine/Basics/Projects/Browser) in UE4. In this section, we'll show you how to create a new Blueprint project from UE4. 
+For the purposes of this guide, let's suppose that you want to 
+create a new Blueprint game with a flying mechanic. Fortunately, UE4 is shipped with the perfect template for you! 
+
+1.	While in the **New Project** menu, make sure that the **Blueprint** menu is open.
+
+	[REGION:lightbox]
+		[![](SetupIDE_Step19.png)(w:720)](SetupIDE_Step19.png)
+	[/REGION]
+	
+	[REGION:caption]
+		Click for full image.
+	[/REGION]
+
+1.	Select the **Flying** Blueprint project template and name the project `FlyingBP_Prj`.
+
+	[REGION:lightbox]
+		[![](SetupIDE_Step20.png)(w:720)](SetupIDE_Step20.png)
+	[/REGION]
+	
+	[REGION:caption]
+		Click for full image.
+	[/REGION]
+
+1.	Finally, click the **Create Project** button to create your new Blueprint project using the Flying Blueprint project template.
+
+	[REGION:lightbox]
+		[![](SetupIDE_Step21.png)(w:720)](SetupIDE_Step21.png)
+	[/REGION]
+	
+	[REGION:caption]
+		Click for full image.
+	[/REGION]
+
+	[REGION:note]
+	Depending your system specifications, it may take a couple of minutes for UE4 to create your new Blueprint project.
+	[/REGION]
+	
+## End Result
+
+[REGION:lightbox]
+	[![](SetupIDE_Step22.png)(w:720)](SetupIDE_Step22.png)
+[/REGION]
+
+[REGION:caption]
+	Click for full image.
+[/REGION]
+
+UE4 will launch your new Blueprint project without any additional steps. 
+
+[REGION:lightbox]
+	[![](SetupIDE_Step23.png)(w:720)](SetupIDE_Step23.png)
+[/REGION]
+
+[REGION:caption]
+	Click for full image.
+[/REGION]
+
+Congratulations! You've just created a new Blueprint project from UE4 with the help of a Blueprint project template. 
+Creating any type of Blueprint project, with or without a project template, is now within your reach. There are times; 
+however, when you'll want to create and modify the source code driving 
+your game. In the next section, we'll show you how to create a new CPP project from UE4.
+
+%Steps%
\ No newline at end of file
diff --git a/Engine/Documentation/Source/Platforms/Linux/BeginnerLinuxDeveloper/SettingUpAnIDE/6/SettingUpAnIDE_6.INT.udn b/Engine/Documentation/Source/Platforms/Linux/BeginnerLinuxDeveloper/SettingUpAnIDE/6/SettingUpAnIDE_6.INT.udn
new file mode 100644
--- /dev/null
+++ b/Engine/Documentation/Source/Platforms/Linux/BeginnerLinuxDeveloper/SettingUpAnIDE/6/SettingUpAnIDE_6.INT.udn
@@ -0,0 +1,156 @@
+Availability: Docs
+Title: 6 - Creating a CPP Project from your IDE
+Crumbs:
+Description:Learn how to create and configure a CPP project form your IDE.
+Version: 4.14
+type:how to
+parent:Platforms/Linux/BeginnerLinuxDeveloper/SettingUpAnIDE
+order: 6
+tags:Linux
+SkillLevel:Beginner 
+Related: GettingStarted
+Related: Engine/Content/QuickStart
+Related: Programming/QuickStart
+Related: Engine/QuickStart
+Related: Support/Builds
+
+[VAR:Steps]
+[OBJECT:Navigation]
+	[PARAM:previous]
+		[Previous Step](Platforms/Linux/BeginnerLinuxDeveloper/SettingUpAnIDE/5 "%Platforms/Linux/BeginnerLinuxDeveloper/SettingUpAnIDE/5:title%")
+	[/PARAM]
+	[PARAM:current]
+	[/PARAM]
+	[PARAM:home]
+		[Home](Platforms/Linux/BeginnerLinuxDeveloper/SettingUpAnIDE)
+	[/PARAM]
+	[PARAM:next]
+	[/PARAM]
+[/OBJECT]
+[/VAR]
+
+%Steps%
+
+By now, you've learned how to configure your IDE to run UE4 in a Linux environment. Also, you've learned how to create a new Blueprint project using a Project template. In this section, we'll show you 
+how to create a new CPP project from UE4. For the purposes of this demonstration, let's suppose that you want to create a new CPP game with a driving mechanic.
+
+1.	While in the **New Project** menu, make sure that the **C++** menu is open.
+
+	[REGION:lightbox]
+		[![](SetupIDE_Step24.png)(w:720)](SetupIDE_Step24.png)
+	[/REGION]
+	
+	[REGION:caption]
+		Click for full image.
+	[/REGION]
+
+1.	Select the **Vehicle** CPP project template and name the project `DrivingCPP_Prj`.
+
+	[REGION:lightbox]
+		[![](SetupIDE_Step25.png)(w:720)](SetupIDE_Step25.png)
+	[/REGION]
+	
+	[REGION:caption]
+		Click for full image.
+	[/REGION]
+
+1.	Finally, click the **Create Project** button to create your new CPP project using the Vehicle CPP project template.
+
+	[REGION:lightbox]
+		[![](SetupIDE_Step26.png)(w:720)](SetupIDE_Step26.png)
+	[/REGION]
+	
+	[REGION:caption]
+		Click for full image.
+	[/REGION]
+
+	[REGION:note]
+	Depending on your system specifications, it may take a couple of minutes for UE4 to create your new CPP project. 
+	After UE4 creates your new CPP project, an instance of Qt Creator for your new CPP project will be launched, closing UE4Editor in the process.
+	[/REGION]
+
+1.	At this point, you should see Qt Creator's **Configure Project** menu in the newly launched instance of Qt Creator. Use what you've learned thus far to configure your new CPP project to have the following Build Settings for Debug and Release Mode:
+
+	![](SetupIDE_Step27.png)
+
+	[REGION:note]
+	At this time, project settings aren't inherited from the original UE4 Editor project settings. So, when you create a new UE4 CPP project using Qt Creator, you'll need to configure 
+	the project settings using the same steps that you learned earlier in this guide.
+	[/REGION]
+
+1.	Now, configure your new CPP project to have the following Run Settings:
+
+	![](SetupIDE_Step28.png)
+
+1.	Switch to **Edit** mode in your IDE by clicking on the **Edit** button in the left panel of your IDE. Alternatively, you can press the **Ctrl+2** keyboard combination to open Qt Creator's **Edit Mode**.
+
+	[REGION:lightbox]
+		[![](SetupIDE_Step29.png)(w:720)](SetupIDE_Step29.png)
+	[/REGION]
+	
+	[REGION:caption]
+		Click for full image.
+	[/REGION]
+
+1.	Now, click the **Run** button. Alternatively, you can press the **Ctrl+R** keyboard combination to run UE4 from the IDE.
+
+	[REGION:lightbox]
+		[![](SetupIDE_Step30.png)(w:720)](SetupIDE_Step30.png)
+	[/REGION]
+	
+	[REGION:caption]
+		Click for full image.
+	[/REGION]
+	
+1.	After the **Missing DriveCPP_Prj Modules** dialogue box appears, click the **Yes** button.
+
+	[REGION:lightbox]
+		[![](SetupIDE_Step31.png)(w:720)](SetupIDE_Step31.png)
+	[/REGION]
+	
+	[REGION:caption]
+		Click for full image.
+	[/REGION]
+
+1.	Depending on your system specifications, it may take several minutes for the IDE to build and launch your new CPP project.
+
+	[REGION:lightbox]
+		[![](SetupIDE_Step32.png)(w:720)](SetupIDE_Step32.png)
+	[/REGION]
+	
+	[REGION:caption]
+		Click for full image.
+	[/REGION]
+
+## End Result
+
+[REGION:lightbox]
+	[![](SetupIDE_Step33.png)(w:720)](SetupIDE_Step33.png)
+[/REGION]
+
+[REGION:caption]
+	Click for full image.
+[/REGION]
+
+Congratulations! You should now see your new CPP project running from inside of Qt Creator. 
+Now that you've created and launched your new CPP project, you've learned everything you need to work on UE4 projects from an IDE in your Linux environment. 
+As you worked your way through this tutorial, you learned:
+
+&#x2713; How to configure your IDE to build UE4 in debug and release mode.  
+&#x2713; How to configure your IDE to run UE4 in debug and release mode.  
+&#x2713; How to build UE4 from your IDE.  
+&#x2713; How to run UE4 from your IDE.  
+&#x2713; How to create and configure a new Blueprint project with UE4 from your IDE.  
+&#x2713; How to create and configure a new CPP project with UE4 from your IDE.  
+
+## Getting Started with Unreal Engine
+
+Learn how to use Unreal Engine by referring to our [](GettingStarted) documentation!
+
+If you're looking to quickly get started with UE4, refer to the following tutorials:
+*	[](Engine/Content/QuickStart)
+*	[](Programming/QuickStart)
+*	[](Engine/QuickStart)
+
+
+%Steps%
\ No newline at end of file
diff --git a/Engine/Documentation/Source/Platforms/Linux/BeginnerLinuxDeveloper/SettingUpAnIDE/SettingUpAnIDE.INT.udn b/Engine/Documentation/Source/Platforms/Linux/BeginnerLinuxDeveloper/SettingUpAnIDE/SettingUpAnIDE.INT.udn
new file mode 100644
--- /dev/null
+++ b/Engine/Documentation/Source/Platforms/Linux/BeginnerLinuxDeveloper/SettingUpAnIDE/SettingUpAnIDE.INT.udn
@@ -0,0 +1,71 @@
+Availability: Docs
+Title: How to Setup an IDE for UE4
+Crumbs:
+Description:Learn how to setup your IDE to build and run UE4 on Linux.
+Platform:Linux
+Version: 4.14
+type:how to
+parent:Platforms/Linux/BeginnerLinuxDeveloper
+order: 3
+tags:Linux
+SkillLevel:Beginner
+prereq:Platforms/Linux/BeginnerLinuxDeveloper/SettingUpAnUnrealWorkflow
+
+[VAR:Params]
+[PARAM:image]
+	![%Platforms/Linux/BeginnerLinuxDeveloper/SettingUpAnIDE:title%](IDE_HowToTopic.png)
+[/PARAM]
+[PARAM:icon]
+	![](%ROOT%/tutorial_icon.png)(convert:false)
+[/PARAM]
+[PARAM:title]
+	%Platforms/Linux/BeginnerLinuxDeveloper/SettingUpAnIDE:title%
+[/PARAM]
+[PARAM:path]
+	[RELATIVE:Platforms/Linux/BeginnerLinuxDeveloper/SettingUpAnIDE]
+[/PARAM]
+[/VAR]
+
+[VAR:Topic]
+[OBJECT:Topic]
+%params%
+[/OBJECT]
+[/VAR]
+
+[REGION:imagetable]
+![](CompiledEditor_LinuxIDE.png)(w:1280)  
+At the end of this tutorial, you'll have an **Integrated Development Environment** (IDE) configured to build and run Unreal Engine 4 (UE4).
+[/REGION]
+
+## Goal
+
+The focus of this guide is to show you how to setup an IDE to build and run UE4. 
+
+## Objectives
+
+After working through this tutorial, you'll know the following:
+
+* How to configure your IDE to build UE4 in debug and release mode.
+* How to configure your IDE to run UE4 in debug and release mode.
+* How to build UE4 from your IDE.
+* How to run UE4 from your IDE.
+* How to create and configure a new Blueprint project with UE4 from your IDE.
+* How to create and configure a new CPP project with UE4 from your IDE.
+
+[OBJECT:TopicList]
+	[PARAM:icon]
+		![](%ROOT%/build_icon.png)(convert:false)
+	[/PARAM]
+	[PARAM:title]
+		Steps
+	[/PARAM]
+	[PARAM:description]
+	[/PARAM]
+	[PARAM:links]
+	[DIR(output:"checkpoint" parent:"Platforms/Linux/BeginnerLinuxDeveloper/SettingUpAnIDE" org:"hierarchy")]
+	[/PARAM]
+[/OBJECT]
+
+[REGION:call_to_action]
+	[Click to Start](Platforms/Linux/BeginnerLinuxDeveloper/SettingUpAnIDE/1 "%Platforms/Linux/BeginnerLinuxDeveloper/SettingUpAnIDE/1:title%")
+[/REGION]
\ No newline at end of file
diff --git a/Engine/Documentation/Source/Platforms/Linux/BeginnerLinuxDeveloper/SettingUpAnUnrealWorkflow/1/SettingUpAnUnrealWorkflow_1.INT.udn b/Engine/Documentation/Source/Platforms/Linux/BeginnerLinuxDeveloper/SettingUpAnUnrealWorkflow/1/SettingUpAnUnrealWorkflow_1.INT.udn
--- a/Engine/Documentation/Source/Platforms/Linux/BeginnerLinuxDeveloper/SettingUpAnUnrealWorkflow/1/SettingUpAnUnrealWorkflow_1.INT.udn
+++ b/Engine/Documentation/Source/Platforms/Linux/BeginnerLinuxDeveloper/SettingUpAnUnrealWorkflow/1/SettingUpAnUnrealWorkflow_1.INT.udn
@@ -3,7 +3,7 @@ Title: 1 - Required Setup
 Crumbs:%ROOT%, Platforms, Platforms/Linux, Platforms/Linux/BeginnerLinuxDeveloper, Platforms/Linux/BeginnerLinuxDeveloper/SettingUpAnUnrealWorkflow
 Description:Learn how to set up Git on Linux.
 Version: 4.13
-type:QuickStart
+type:quick start
 parent:Platforms/Linux/BeginnerLinuxDeveloper/SettingUpAnUnrealWorkflow
 order: 1
 tags:Linux
diff --git a/Engine/Documentation/Source/Platforms/Linux/BeginnerLinuxDeveloper/SettingUpAnUnrealWorkflow/2/SettingUpAnUnrealWorkflow_2.INT.udn b/Engine/Documentation/Source/Platforms/Linux/BeginnerLinuxDeveloper/SettingUpAnUnrealWorkflow/2/SettingUpAnUnrealWorkflow_2.INT.udn
--- a/Engine/Documentation/Source/Platforms/Linux/BeginnerLinuxDeveloper/SettingUpAnUnrealWorkflow/2/SettingUpAnUnrealWorkflow_2.INT.udn
+++ b/Engine/Documentation/Source/Platforms/Linux/BeginnerLinuxDeveloper/SettingUpAnUnrealWorkflow/2/SettingUpAnUnrealWorkflow_2.INT.udn
@@ -3,7 +3,7 @@ Title: 2 - Downloading UE4 on Linux
 Crumbs:%ROOT%, Platforms, Platforms/Linux, Platforms/Linux/BeginnerLinuxDeveloper, Platforms/Linux/BeginnerLinuxDeveloper/SettingUpAnUnrealWorkflow
 Description:Learn how to fork and clone Unreal Engine source code.
 Version: 4.13
-type:QuickStart
+type:quick start
 parent:Platforms/Linux/BeginnerLinuxDeveloper/SettingUpAnUnrealWorkflow
 order: 2
 tags:Linux
diff --git a/Engine/Documentation/Source/Platforms/Linux/GettingStarted/LinuxCrossGettingStarted.INT.udn b/Engine/Documentation/Source/Platforms/Linux/GettingStarted/LinuxCrossGettingStarted.INT.udn
--- a/Engine/Documentation/Source/Platforms/Linux/GettingStarted/LinuxCrossGettingStarted.INT.udn
+++ b/Engine/Documentation/Source/Platforms/Linux/GettingStarted/LinuxCrossGettingStarted.INT.udn
@@ -1,13 +1,34 @@
 Availability:Public
-Title: Getting Started: Developing Linux Projects
+Title: Getting Started with Linux Projects
 Crumbs:%ROOT%, Platforms, Platforms/Linux
 Description:Setting up for development for the Linux platform.
 Version: 4.9
 SkillLevel: Advanced
-type:quickstart
+type:quick start
 parent:Platforms/Linux
 tags:Linux
 
+[VAR:Params]
+[PARAM:image]
+	![%Platforms/Linux/GettingStarted:title%](GettingStarted_Topic.png)
+[/PARAM]
+[PARAM:icon]
+	![](%ROOT%/tutorial_icon.png)(convert:false)
+[/PARAM]
+[PARAM:title]
+	%Platforms/Linux/GettingStarted:title%
+[/PARAM]
+[PARAM:path]
+	[RELATIVE:Platforms/Linux/GettingStarted]
+[/PARAM]
+[/VAR]
+
+[VAR:Topic]
+[OBJECT:Topic]
+%params%
+[/OBJECT]
+[/VAR]
+
 <!--- %Globals:BinarySourceSelectVar% --->
 
 [TOC(start:2 end:3)]
diff --git a/Engine/Documentation/Source/Platforms/Linux/Linux.INT.udn b/Engine/Documentation/Source/Platforms/Linux/Linux.INT.udn
--- a/Engine/Documentation/Source/Platforms/Linux/Linux.INT.udn
+++ b/Engine/Documentation/Source/Platforms/Linux/Linux.INT.udn
@@ -9,7 +9,6 @@ type:landing
 parent:Platforms
 tags:Linux
 related: GettingStarted/RecommendedSpecifications
-related: Platforms/Linux/GettingStarted
 
 Fueling the future of Linux game development is at your fingertips! Now, more than ever, you can harness the power of
 Unreal Engine 4 (UE4) to create and share games with Linux users worldwide. If you're new to Linux game development,
@@ -17,7 +16,9 @@ you'll want to work your way through the Quick Start Guide before moving onto mo
 
 ## Getting Started
 
+%Platforms/Linux/GettingStarted:topic%
 %Platforms/Linux/BeginnerLinuxDeveloper/SettingUpAnUnrealWorkflow:topic%
+<!--%Platforms/Linux/BeginnerLinuxDeveloper/SettingUpAnIDE:topic%-->
 
 ## Reference
 
diff --git a/Engine/Documentation/Source/Platforms/Mobile/Lighting/MobileLighting.INT.udn b/Engine/Documentation/Source/Platforms/Mobile/Lighting/MobileLighting.INT.udn
--- a/Engine/Documentation/Source/Platforms/Mobile/Lighting/MobileLighting.INT.udn
+++ b/Engine/Documentation/Source/Platforms/Mobile/Lighting/MobileLighting.INT.udn
@@ -106,7 +106,7 @@ Stationary Directional lights can be used to cast Whole-scene Dynamic Cascade Sh
 ## Modulated Shadowing VS Dynamic Shadowing
 
 While Modulated and Dynamic shadowing looks and acts very similar to one another, they are quite different. 
-This is because Modulated shadows have a number of restrictions and hard limits that Dynamic shadows do not have that allows them to beperform bettert on mobile devices with limited hardware capabilities.
+This is because Modulated shadows have a number of restrictions and hard limits that Dynamic shadows do not have that allows them to perform better on mobile devices with limited hardware capabilities.
 In the image below you can see a comparison between Dynamic shadows and Modulated shadows.
 
 [OBJECT:ComparisonSlider]
diff --git a/Engine/Documentation/Source/Platforms/Mobile/Mobile.INT.udn b/Engine/Documentation/Source/Platforms/Mobile/Mobile.INT.udn
--- a/Engine/Documentation/Source/Platforms/Mobile/Mobile.INT.udn
+++ b/Engine/Documentation/Source/Platforms/Mobile/Mobile.INT.udn
@@ -24,7 +24,7 @@ type:landing
 	[PARAM:description]
 	[/PARAM]
 	[PARAM:links]
-		[DIR(output:"topiccompact" type:"quickstart" tags:"mobile")]
+		[DIR(output:"topiccompact" type:"quick start" tags:"mobile")]
 	[/PARAM]
 [/OBJECT]
 
@@ -52,7 +52,7 @@ type:landing
 	[PARAM:description]
 	[/PARAM]
 	[PARAM:links]
-		[DIR(end:"1" output:"listbutton" type:"howto" parent:"Platforms/Mobile")]
+		[DIR(end:"1" output:"listbutton" type:"how-to" parent:"Platforms/Mobile")]
 	[/PARAM]
 [/OBJECT]
 
