Copying //UE4/Dev-Documentation to Samples-Main (//UE4/Samples-Main)

diff --git a/Engine/Documentation/Source/Platforms/Mobile/MobilePackagingWizard/MobilePackagingWizard.INT.udn b/Engine/Documentation/Source/Platforms/Mobile/MobilePackagingWizard/MobilePackagingWizard.INT.udn
--- a/Engine/Documentation/Source/Platforms/Mobile/MobilePackagingWizard/MobilePackagingWizard.INT.udn
+++ b/Engine/Documentation/Source/Platforms/Mobile/MobilePackagingWizard/MobilePackagingWizard.INT.udn
@@ -1,7 +1,7 @@
 Title: Mobile Packaging Wizard 
 Description:Using the Mobile Packaging Wizard to reduce the install size of your UE4 mobile project
 Crumbs:%ROOT%, Platforms, Platforms/Mobile, Platforms/Mobile/MobilePackagingWizard/
-Availability:Docs
+Availability: Public
 version:4.12
 Parent:Platforms/Mobile
 tags:Mobile
@@ -13,15 +13,14 @@ Related:Engine/Deployment/Cooking
 Related:Engine/Deployment/Patching
 Related:Engine/Deployment/UnrealFrontend
 
-[TOC(start:2 end:3)]
-
 
 The Unreal Engine 4 (UE4) **Mobile Packaging Wizard** will help to reduce the size of your mobile project by packaging only the content that is needed to upload your project to the App store.  The rest of the content will be download from the cloud when the user launches the game on their device for the first time. In the following document, we will take a look at how you set up your UE4 mobile project to use the Mobile Packaging Wizard.
 
+%Globals:MobileSelect%
 
 ## Mobile Packaging Wizard Tool Location
 
-To access the Mobile Packaging Wizard you will need to do the folowing.
+To access the Mobile Packaging Wizard, you will need to do the following.
 
 1. From the **Main Toolbar**, locate the **Window** option and click on it. 
 
@@ -31,18 +30,45 @@ To access the Mobile Packaging Wizard you will need to do the folowing.
 
 	![](MPW_Project_Launcher.png)
 
+
+[OBJECT:ToggleButtonContent]
+	[PARAMLITERAL:category]
+	Platform
+	[/PARAMLITERAL]
+	[PARAMLITERAL:id]
+	android
+	[/PARAMLITERAL]
+	[PARAMLITERAL:active]
+	active_button_content
+	[/PARAMLITERAL]
+	[PARAM:content]
+	[INCLUDE:Platforms\Mobile\MobilePackagingWizard#android]
+	[/PARAM]
+[/OBJECT]
+[OBJECT:ToggleButtonContent]
+	[PARAMLITERAL:category]
+	Platform
+	[/PARAMLITERAL]
+	[PARAMLITERAL:id]
+	ios
+	[/PARAMLITERAL]
+	[PARAMLITERAL:active]
+	[/PARAMLITERAL]
+	[PARAM:content]
+	[INCLUDE:Platforms\Mobile\MobilePackagingWizard#ios]
+	[/PARAM]
+[/OBJECT]
+
+<!--
+[EXCERPT:android]
 ## Setting Up a Mobile Packaging Wizard Profile
 
-Before the Mobile Packaging Wizard can be used, a new profile that is specific to your project's needs has to be created. In the following section, we will take a look at what needs to be done in order to create a custom launch profile for your UE4 mobile project.
+Before the Mobile Packaging Wizard can be used, a new profile that is specific to your project's needs has to be created. In the following section, we will take a look at what needs to be done to create a custom launch profile for your UE4 mobile project.
 
 1. Click on the small white triangle next to **Custom Launch Profiles**, selecting the **Minimal Android APK + DLC…** option from the list.
 
 	![](MPW_Select_Android_Option.png)
 
-	[region:note]
-	While we are using Android for this example, the workflow will still be the same when using the Mobile Packaging Wizard for a game that is meant for IOS.
-	[/REGION]
-
 	(#bookmark1)
 1. When the **Minimal Android APK + DLC** packaging window is displayed, click on the **Browse** button, and select a location for the packaged content to be placed.
 
@@ -66,7 +92,7 @@ Before the Mobile Packaging Wizard can be used, a new profile that is specific t
 
 ## Building the Application 
 
-Now that all of the options for the Mobile Packing Wizard have been set up and saved into a Custom Launch Profile, it is time to run that profile and package your project. In the following section, we will take a look at what needs to be done in order for the Mobile Packing Wizard to successfully package your UE4 project.
+Now that all of the options for the Mobile Packing Wizard have been set up and saved into a Custom Launch Profile, it is time to run that profile and package your project. In the following section, we will take a look at what needs to be done for the Mobile Packing Wizard to successfully package your UE4 project.
 
 1. If the Project Launcher is not already open, open it up and then look at the bottom of the Project Launcher in the **Custom Launch Profile** section. In this section, you should have two profiles that start with your project name and are followed by either **Android APK** or **Android DLC**. 
 
@@ -94,7 +120,7 @@ Now that all of the options for the Mobile Packing Wizard have been set up and s
 
 ## Building the DLC
 
-Once your project has finished the cooking process for the initial download from the App store, it is now time to cook the content that is to be downloaded from the cloud after your App has been installed to a mobile device. 
+Once your project has finished the cooking process for the initial download from the App store, it is now time to cook the content that is to be downloaded from the cloud after your App has been installed on a mobile device. 
 
 1. From the Project Launcher, go to the **Custom Launch Profiles** section and look for your project's **DLC** option. When the DLC option for your project has been located, click on the **Launch this Profile** button to begin the DLC packaging process. 
 
@@ -108,4 +134,90 @@ Once your project has finished the cooking process for the initial download from
 
 Once everything has finished building, open up the folder that was specified in the [second step](#bookmark1) of the "Setting Up a Mobile Packaging Wizard Profile" section. You'll see the following two folders:
 
-![](MPW_App_DLC_Location.png)
\ No newline at end of file
+![](MPW_App_DLC_Location.png)
+
+Inside of the folders, you will find the various pieces of your project that have been set up to be downloaded after the users have downloaded and installed the application on their mobile device.
+
+[/EXCERPT:android]
+
+
+[EXCERPT:ios]
+## Setting Up a Mobile Packaging Wizard Profile
+
+Before the Mobile Packaging Wizard can be used, a new profile that is specific to your project's needs has to be created. In the following section, we will take a look at what needs to be done to create a custom launch profile for your UE4 mobile project.
+
+1. Click on the small white triangle next to **Custom Launch Profiles**, selecting the **Minimal IOS App + DLC…** option from the list.
+
+	![](MPW_Select_iOS_Option.png)
+
+	(#bookmark1)
+1. When the **Minimal IOS App + DLC** packaging window is displayed, click on the **Browse** button and select a location for the packaged content to be placed.
+
+	![](MPW_iOS_Press_Browse_Button.png)
+
+1. Once you have selected a location for the packaged content to be placed, press the **Next** button to continue the packaging process.
+
+	![](MPW_iOS_Press_Next.png)
+
+1. In the **Distributable Application** menu, set the **Build Configuration** to **Development** and then select the map you want to use, pressing the **Next** button when finished. 
+
+	![](MPW_Dis_iOS.png)
+
+	[region:note]
+	Only the contents referenced by the map(s) you choose will be packaged into the application for upload to the App Store. Typically, you should choose only a single map that contains a Level Blueprint to kick off the **BuildPatchServices** code with a **UMG** user interface to show the download progress (showing any applicable error conditions).
+	[/region]
+
+1. In the **Downloadable Content** menu, select the maps that should be downloaded and then press the **Create Profile** button to create a new Custom Launch Profile.
+
+	![](MPW_Downloadable_Content_iOS.png)
+
+## Building the Application 
+
+Now that all of the options for the Mobile Packing Wizard have been set up and saved into a Custom Launch Profile, it is time to run that profile and package your project. In the following section, we will take a look at what needs to be done for the Mobile Packing Wizard to successfully package your UE4 project.
+
+1. If the Project Launcher is not already open, open it up and then look at the bottom of the Project Launcher in the **Custom Launch Profile** section. In this section, you should have two profiles that start with your project name and are followed by either **IOS App** or **IOS DLC**. 
+
+	![](MPW_Custom_Launch_Profiles_iOS.png)
+
+	[region:note]
+	When you first open up the Project Launcher, the Custom Launch Profiles section could have a few items in it by default. To remove these items, right-click on the item you want to remove, and then select the **Delete** option to remove it.
+	[/region]
+
+1. Locate the **IOS App** Profile option and then click on the **Launch this Profile** button to begin the packaging process. 
+
+	![](MPW_Launch_APK_Profile_iOS.png)
+
+	[region:warning]
+	You must package your project's files before you build the DLC. <!--Trying to package the DLC first and then packing your project will not.-->
+	[/region]
+
+1. Once the **Launch this Profile** button has been pressed, the Project Launcher will begin the packaging process. You can see the progress of the packaging process in the Project Launcher window.
+
+	![](MPW_Packaging_Progress_iOS.png)
+
+1. When the packaging has completed, click on the **Done** button to load the Project Launcher.
+
+	![](MPW_Packaging_Completed_iOS.png)
+
+## Building the DLC
+
+Once your project has finished the cooking process for the initial download from the App store, it is now time to cook the content that is to be downloaded from the cloud after your App has been installed on a mobile device. 
+
+1. From the Project Launcher, go to the **Custom Launch Profiles** section and look for your project's **DLC** option. When the DLC option for your project has been located, click on the **Launch this Profile** button to begin the DLC packaging process. 
+
+	![](MPW_Build_DLC_iOS.png)
+
+1. After pressing the **Launch this Profile** button, the Project Launcher will begin to build the DLC for your project. When the DLC has finished being built, click on the **Done** button to continue. 
+
+	![](MPW_DLC_Completed_iOS.png)
+
+## Packaging Results 
+
+Once everything has finished building, open up the folder that was specified in the [second step](#bookmark1) of the "Setting Up a Mobile Packaging Wizard Profile" section. You'll see the following two folders:
+
+![](MPW_App_DLC_Location.png)
+
+Inside of the folders, you will find the various pieces of your project that have been set up to be downloaded after the users have downloaded and installed the application on their mobile device.
+
+[/EXCERPT:ios]
+-->
\ No newline at end of file
diff --git a/Engine/Documentation/Source/Platforms/Mobile/Performance/MobilePerformance.INT.udn b/Engine/Documentation/Source/Platforms/Mobile/Performance/MobilePerformance.INT.udn
--- a/Engine/Documentation/Source/Platforms/Mobile/Performance/MobilePerformance.INT.udn
+++ b/Engine/Documentation/Source/Platforms/Mobile/Performance/MobilePerformance.INT.udn
@@ -175,7 +175,7 @@ The iOS scale factor is directly related to Apple's scale factor system.
 |r.MobileContentScaleFactor|0.0|This will use the native resolution of the device.|
 |r.MobileContentScaleFactor|1.0|On a Retina device is the non-retina resolution.|
 |r.MobileContentScaleFactor|2.0|Full native resolution on iPhone 5S, iPad Air, etc.|
-|r.MobileContentScaleFactor|3.0|ull native resolution for iPhone 6+.|
+|r.MobileContentScaleFactor|3.0|Full native resolution for iPhone 6+.|
 
 ### Mobile Content Scale Factor for Android
 When using the Mobile Content Scale Factor for Android devices inputting the following number after the r.MobileContentScaleFactor will produce the following results. Please note that inputting value other than **0.0** will use this value as the scale factor for a standard 1280x720 or 720x1280 resolution depending on orientation.
@@ -272,7 +272,7 @@ The **Quality Switch** Material Expression allows you to define different levels
 
 ![](image08.png)
 
-To use the Quality Switch in your Materials, you will need to first add the Quality Switch Material Expression node to the Material Graph and then plug it's output into any if the inputs on the Main Material node you want to change when the quality level is changed. In the following image, the M_Brick_Clay_Old Material that can be found with the Starter Content was set up so that when the Quality Level is switched the complexity of the Material is reduced. 
+To use the Quality Switch in your Materials, you will need to first add the Quality Switch Material Expression node to the Material Graph and then plug it's output into any of the inputs on the Main Material node you want to change when the quality level is changed. In the following image, the M_Brick_Clay_Old Material that can be found with the Starter Content was set up so that when the Quality Level is switched the complexity of the Material is reduced. 
 
 [REGION:lightbox]
 [![](MQS_Material_Setup.png)(w:789 h:412)](MQS_Material_Setup.png)
@@ -335,8 +335,8 @@ Under the Mobile section you will find a number of different options that can be
 |--------------------|-------|
 |Mobile HDR|If true, mobile renders in full HDR. Disable this setting for games that do not require lighting features for better performance on slow devices.|
 |Max Dynamic Point Lights|The number of dynamic point lights to support on mobile devices. Setting this to 0 for games which do not require dynamic point lights will reduce the number of shader generated. Changing this setting requires restarting the editor.|
-|Use Shared Dynamic Point LIght Shaders|If this setting is enabled, the same shader will be sued for any number of dynamic point lights(up to the maximum specified above)hitting a surface. This is slightly slower but reduced the number of shaders generated. Changing this setting requires restarting the editor. |
-|Enabled Combined Static and CSM Shadowing|Allo primitives to receive both static and CSM shadows from a stationary light. Disabling will free a mobile texture sampler.|
+|Use Shared Dynamic Point Light Shaders|If this setting is enabled, the same shader will be used for any number of dynamic point lights(up to the maximum specified above)hitting a surface. This is slightly slower but reduced the number of shaders generated. Changing this setting requires restarting the editor. |
+|Enabled Combined Static and CSM Shadowing|Allows primitives to receive both static and CSM shadows from a stationary light. Disabling will free a mobile texture sampler.|
 
 
 <!-----
diff --git a/Engine/Documentation/Source/Platforms/Mobile/PostProcessEffects/MobilePostProcessEffects.INT.udn b/Engine/Documentation/Source/Platforms/Mobile/PostProcessEffects/MobilePostProcessEffects.INT.udn
--- a/Engine/Documentation/Source/Platforms/Mobile/PostProcessEffects/MobilePostProcessEffects.INT.udn
+++ b/Engine/Documentation/Source/Platforms/Mobile/PostProcessEffects/MobilePostProcessEffects.INT.udn
@@ -3,21 +3,28 @@ Title:Post Process Effects on Mobile Platforms
 Crumbs:%ROOT%, Platforms, Platforms/Mobile
 Description:Using Post Process Effects on mobile games.
 Platform:Mobile
-Version: 4.9
+Version: 4.14
 parent:Platforms/Mobile
 order:5
 type:overview
 tags:mobile
+reference-image:Mobile_PostProcessEffect.gif
 
 [TOC(start:2)]
 
+Mobile Post Processing is implemented as a separate code path from PC/Console Post Processing. This is to reflect the slower dependent texture reads, missing hardware features, extra render target resolve costs, and slower performance in general. 
+To enable Post Processing for your mobile applications, please make sure that **Mobile HDR** is enabled in your **Project Settings** as most of the settings below require Mobile HDR to be enabled for the effect(s) to work.
 
-Mobile Post Processing is implemented as a separate code path from PC/Console post processing. This is to reflect the slower dependent texture reads, missing hardware features, extra render target resolve costs, and the slower performance in general. The following categories can be found in a Post Process volume (or a camera Actor).
+[region: lightbox]
+[![](MobileHDR_Enabled.png)(w:400 h:200)](MobileHDR_Enabled.png)
+[/region]
 
-[region:note]
-Make sure that Mobile HDR is enabled in your project settings as most of the settings below require this setting to be enabled for the effect to work.
+[region:caption]
+Click for full image.
 [/region]
 
+The following categories can be found in a Post Process volume (or a camera Actor). 
+
 ## Enabling Tonemapper for Mobile Devices
 To use the Tonemapper Post Processing options on Mobile Devices, you will need to set the value of **r.tonemapperfilm** from its default of **0** to **1**.  You can change, enable, or disable the Tonemapper inside of the Unreal Engine 4 (UE4) Editor in the following ways: 
 
@@ -216,6 +223,37 @@ The AA Method supports "None" (the fastest) and "TemporalAA" which provides a sp
 |Max Roughness|Yes|0 & 1|
 
 
+## Custom Post Process Materials
+
+Unreal Engine ships with a variety of Post Process effects, allowing you to customize the look and feel of your mobile applications. There are times, though, when you'll want to create 
+custom Post Process Materials, enabling your own effects (such as "VCR Static").
+
+[REGION:asyncgif]
+![image alt text](Mobile_PostProcessEffect.gif)
+[/REGION]
+
+Adding effects to your mobile application requires the development of custom Post Process materials. When creating custom Post Process materials for mobile applications, you should note that there 
+are some key things to consider. Specifically, it's important to note that custom Post Process materials for mobile applications can only fetch from PostProcessInput0 (Scene Color) with the 
+following **Blendable Locations**:
+
+* **Before Tonemapping**
+* **After Tonemapping**
+
+To learn more about Blendable Locations, read through our [](Engine/Rendering/PostProcessEffects/PostProcessMaterials) reference page. 
+Currently, you cannot create custom Post Process materials for older Android devices that require 'mosaic' mode for HDR rendering. Finally, pixel depth information is not yet supported.
+
+### Custom Depth
+
+You can now use the **Custom Depth** Post Process option with the Mobile rendering path. Custom Post-Process Materials can now sample from **Scene Depth**, **Custom Depth** as well as **Scene Color**. 
+
+![](Custom_Scene_Depth.png)
+
+As this features requires Post-Processing to work correctly, **Mobile HDR** must be enabled. To enable Mobile HDR go to **Project Settings** > **Engine** > **Rendering** > **Mobile** and make sure that **Mobile HDR** is checked.
+
+[region:warning]
+Please note that this feature also does not currently work while **Mobile MSAA** is enabled so Mobile MSAA will have to be disabled for you to see the effect on the device.
+[/region]
+
 ## General Performance Tips
 
 * For best performance and quality, you should stick to using just Bloom and TemporalAA.
diff --git a/Engine/Documentation/Source/Platforms/Oculus/Oculus.INT.udn b/Engine/Documentation/Source/Platforms/Oculus/Oculus.INT.udn
--- a/Engine/Documentation/Source/Platforms/Oculus/Oculus.INT.udn
+++ b/Engine/Documentation/Source/Platforms/Oculus/Oculus.INT.udn
@@ -9,11 +9,12 @@ type:landing
 parent:Platforms/VR
 tags:VR
 tags:Oculus Rift
+tags:Platform
 
 
 ## Starting Out
 
-[DIR(output:"listbutton" type:"quickstart" parent:"Platforms/Oculus")]
+[DIR(output:"listbutton" type:"quick start" parent:"Platforms/Oculus")]
 
 ## Oculus Rift Topics
 
diff --git a/Engine/Documentation/Source/Platforms/SteamVR/SteamVR.INT.udn b/Engine/Documentation/Source/Platforms/SteamVR/SteamVR.INT.udn
--- a/Engine/Documentation/Source/Platforms/SteamVR/SteamVR.INT.udn
+++ b/Engine/Documentation/Source/Platforms/SteamVR/SteamVR.INT.udn
@@ -8,12 +8,13 @@ type:landing
 parent:Platforms/VR
 tags:VR
 tags:SteamVR
+tags:Platform
 
 ![](T_Vive_HMD.png)
 
 ## Starting Out
 
-[DIR(output:"listbutton" type:"quickstart" parent:"Platforms/SteamVR")]
+[DIR(output:"listbutton" type:"quick start" parent:"Platforms/SteamVR")]
 
 ## SteamVR Topics
 
diff --git a/Engine/Documentation/Source/Platforms/VR/ContentSetup/VRContent.INT.udn b/Engine/Documentation/Source/Platforms/VR/ContentSetup/VRContent.INT.udn
--- a/Engine/Documentation/Source/Platforms/VR/ContentSetup/VRContent.INT.udn
+++ b/Engine/Documentation/Source/Platforms/VR/ContentSetup/VRContent.INT.udn
@@ -3,7 +3,7 @@ Title:Virtual Reality Best Practices
 Crumbs:%ROOT%, Platforms, Platforms/VR
 Description:Information over developing for VR.
 Navigation:platform
-Version: 4.9
+Version: 4.14
 parent:Platforms/VR
 type:overview
 
@@ -113,6 +113,31 @@ After Instanced Stereo rendering has been enabled, it will require an engine res
 	[/PARAMLITERAL]
 [/OBJECT]
 
+### VR Multiview Support for Mobile
+
+Certain supported mobile devices can now support VR Multiview. Mobile Multiview is similar to Instanced Stereo rendering that is available to desktop PC's and works by providing an optimized path for Stereo Rendering on the mobile devices CPU. To enable this feature you will need to do that following:
+
+[region:warning]
+Due to the experimental nature of this feature, it will only work on modern Mali-based GPU's. As this feature matures and support for more mobile devices come online, this warning will be removed.
+[/region]
+
+1. From the Main Toolbar go to **Edit** >  **Projects Settings** to bring up the Project's Settings options.
+
+1. Then go to **Engine** > **Rendering** > **VR** and look for the **Mobile Multi-View (Experimental)** option. 
+
+1. Enable the Mobile Multi-View (Experimental) option by click on the checkmark box next to its name and then restart the UE4 Editor for the changes to take effect.
+
+    ![](VR_MobileMultiView.png)
+
+For this feature to work you will need your Android build settings to be setup like the following:
+
+* Make sure that **Support OpenGL ES2** is enabled in the **Android** > **Build** section. 
+* Make sure that both **Mobile HDR** and **Instanced Stereo Rendering** are disabled.
+* Currently Mobile Multiview will only work with modern Mali-based GPU's. 
+* If you package a project with this feature on but don't have a compatible GPU, it will be disabled at runtime. 
+
+
+
 ## VR Performance Profiling
 
 To track down assets that are too expensive for VR, you need to profile what your project is doing on the CPU and GPU as much as possible over the lifetime of your project. 
@@ -122,6 +147,17 @@ To track down assets that are too expensive for VR, you need to profile what you
 
 * **CPU Profiler** - Profiling what your project is doing on the CPU is a bit more involved than profiling for the GPU. To find out more on how to do this please refer to the [Performance Profiler](https://docs.unrealengine.com/latest/INT/Engine/Performance/Profiler/index.html) documents. 
 
+## VR Forward Rendering
+
+![](Engine\Performance\ForwardRenderer\ForwardRenderingBanner.png)(w:840)
+
+[INCLUDE:Engine\Performance\ForwardRenderer#intro]
+
+[REGION:note]
+Please see the [](Engine\Performance\ForwardRenderer) documentation for more information on this feature. 
+[/REGION]
+
+
 ## VR Post Process Settings
 
 Due to the demanding rendering requirements of VR, many of the advanced Post Processing features that are enabled by default should be disabled or your project could suffer
@@ -337,4 +373,4 @@ Below you will find some different techniques that can be used in place of Norma
 |Shader Type | Description|
 |---|---|
 |[Parallax Mapping](http://en.wikipedia.org/wiki/Parallax_mapping)| Parallax mapping takes Normal mapping to the next level by accounting for depth cues Normal mapping does not. A Parallax mapping shader can better display depth information, making objects appear to have more detail than they do. This is because no matter what angle you look at, a Parallax map will always correct itself to show you the correct depth information from that view point. The best use of a Parallax map would be for cobblestone pathways and fine detail on surfaces. 
-|[Tessellation Shader Displacement](Resources/ContentExamples/MaterialNodes/1_11) | Tessellation Shader Displacement will displace 3D Geometry in real time by adding details that are not modeled into the object. Tessellation shaders do a great job of displaying information because tessellation shaders actually create the missing detail by creating more vertices and displacing them in 3D Space in real time.  |
+
diff --git a/Engine/Documentation/Source/Platforms/VR/Profiling/Considerations/ConsiderationsAndInterpretations.INT.udn b/Engine/Documentation/Source/Platforms/VR/Profiling/Considerations/ConsiderationsAndInterpretations.INT.udn
new file mode 100644
--- /dev/null
+++ b/Engine/Documentation/Source/Platforms/VR/Profiling/Considerations/ConsiderationsAndInterpretations.INT.udn
@@ -0,0 +1,39 @@
+Availability:Docs
+Title: VR Profiling Interpretations and Considerations
+Crumbs:%ROOT%, Platforms
+Description: Things to keep in mind when interpreting the data from profiling tools.
+Version:4.14
+Skilllevel:Advanced
+Parent:Platforms/VR/Profiling
+order:3
+related:Engine/Performance/Profiler
+related:Engine/Performance/GPU
+related: Platforms/VR/Profiling/Overview
+related: Platforms/VR/Profiling/Tools
+
+
+While [profiling tools](Platforms/VR/Profiling/Tools) can provide data about how your project is performing in VR, there are some points to consider when interpreting those values.
+
+It's worth bearing in mind that in certain CPU-bound cases on DX12 and Vulcan, the GPU timings can be affected by CPU waits (bubbles), so please consider that if you see unexpected results in cases where draw thread time is high. 
+
+Additionally, often when profiling, you'll notice that you pop between 90 Hz and 45 Hz.  The difference is accounted for because the Compositor acts much like a vsync.  If you miss framerate, it delays you until the next frame entirely.  So, you 
+tend to go down in brackets of 90 / n, where n is a whole number.  Because the fps reporting is an average, it won't always report as a whole number, unless the drop is sustained.  If you're bouncing back and forth between making 
+framerate and not, you may see some fraction of a jump.
+
+Note that these delays can appear in odd places in the GPU and CPU profilers, most often as occlusion or scene graph traversal time. If those numbers are abnormally high, it's possible that it's a false positive.
+
+If you'd like to remove this for testing, sometimes it's useful to run the game emulating stereo rendering, instead of running in the device itself.  To do so:
+
+* Launch the game with `-game -emulatestereo -res=2160x1200` on the commandline
+* Ensure vsync is off with `r.vsync 0` in the console
+* Update the screen percentage to emulate the oversampling we have to do for VR with `r.screenpercentage 137` in the console
+
+This will emulate the GPU and CPU performance characteristics without the annoyances of the variable framerate.
+
+It's worth explicitly noting that if you're hitting 90 frames a second most of the time, and then make a change and notice a drastic drop, you're probably falling prey to the Compositor "vsync" issue noted above.  As soon as 
+you tip over the cliff, you'll see drastic changes in numbers.
+
+
+
+
+
diff --git a/Engine/Documentation/Source/Platforms/VR/Profiling/Overview/Overview.INT.udn b/Engine/Documentation/Source/Platforms/VR/Profiling/Overview/Overview.INT.udn
new file mode 100644
--- /dev/null
+++ b/Engine/Documentation/Source/Platforms/VR/Profiling/Overview/Overview.INT.udn
@@ -0,0 +1,95 @@
+Availability:Docs
+Title: VR Performance Testing
+Crumbs:%ROOT%, Platforms
+Description: Overview of the key principles for performance and profiling of UE4 virtual reality projects.
+Version:4.14
+Skilllevel:Advanced
+Parent:Platforms/VR/Profiling
+order:1
+related: Platforms/VR/Profiling/Considerations
+related: Platforms/VR/Profiling/Tools
+
+In general, maintaining VR frame rates is extremely challenging even on recommended hardware. Frame rate consistency is more important than in regular real-time rendering as dropping 
+frames can make your experience uncomfortable for the user.
+
+With that in mind, you need to leave yourself a buffer to ensure you can absorb frames which may run longer than average. We don't allow pipelining multiple frames due to the input/tracking latency it 
+would introduce. This means the game thread is only one frame ahead of the render thread at all times and the GPU and render thread are in sync on the same frame. This means that you don't get 11 ms, you get closer to 8.5.
+ Long frame times in the game thread due to gameplay code overhead, the render thread due to draw call overhead, or the GPU due to transform or shading overhead can cause you to miss frame time. All three 
+ need to be considered and balanced.
+
+If you follow the console variable settings in the [Best Practices guide](Platforms/Oculus/BestPractices), your render thread performance should be around where it needs to be for performance. From there, 
+it's important to optimize your content and gameplay code to ensure that you make frame rate. On the content side, keep everything as simple as it can be. 
+
+General practices for simplifying content include the following.
+* Avoid dynamic lights and shadows. 
+* Avoid heavy use of translucency.
+* Instance where you can.
+* Make LODs for everything.
+* Keep material complexity and number of materials per object low. 
+* Bake out anything that you can get away with. 
+
+[REGION:tip]
+See the [Bullet Train](https://s3.amazonaws.com/epiccms_epiccms_bucket/Resources/files/Going+Off+the+Rails.zip) and [Showdown](https://cdn2.unrealengine.com/Resources/files/UE4-Integration-and-Demos_OC-100270768.pptx) slide decks for 
+examples of how Epic has produced content for virtual reality that looks more complex than it actually is. Also, be aware that most screen space tricks don't behave as you'd expect in stereo. They generally look flat or simply wrong.
+[/REGION]
+
+Be sure to profile your work throughout the project and ensure you're always maintaining frame rate. Slowly adding complexity until you reach a tipping point is much easier to manage than trying to find areas to scale back down
+to meet the target performance a few days before shipping.
+
+Reducing draw calls really comes down to two techniques.  Make sure that you combine as many materials as possible and make sure to make heavy use of instanced static mesh components.  
+On PC, there's some potential in trying to automatically batch these things together, but on mobile, the cost is too high.
+
+The other thing that you may consider doing is using cull volumes to ensure that non-visible rooms are completely culled out when outside of them.  Our occlusion is slightly more liberal in what it 
+chooses to draw, and cull volumes help make this a bit more conservative in those choices.  Additionally, you can toggle visibility on Actors (or entire sublevels) to get the same effect, if your content is built like that.
+
+[REGION:note]
+Another important thing to remember is that if something has an alpha of zero, it may still be drawn.  So, if you don't want to render it, be sure to actually hide it to prevent it from being rendered.
+[/REGION]
+
+## Getting General Stats
+
+When in the game, there are a few ways to get general stats:
+
+| Console Command | Description |
+| --------------- | ----------- |
+| stat unit | This will give you general game thread, draw thread, and GPU time, as well as overall frame timing.  This is most useful for gathering whether the total frametime is where it needs to be, and what the game thread time is, but it should not be used for Draw Thread and GPU time.  Oculus has a variable delay that it adds in order to try and time GPU work with its compositor. |
+| startfpschart / stopfpschart | If you want to get an idea of what percentage of time is spent above 90Hz, you can run these commands.  It'll capture and aggregate stats data over the window between when you started and stopped, and then dump a file that has bucketed framerate information.  Note that since sometimes the game reports slightly under 90Hz when actually at 90, it's best to check the 80+ buckets to determine the actual time spent at framerate. |
+| stat gpu | This command, added in 4.14,  gives similar stats to the GPU profiler, but in a form that you can watch and monitor from in the game.  Great for checking quick costs on GPU work. |
+
+###  Real-Time GPU Profiler 
+To use the new real-time GPU profiler, type `stat gpu` in the game or editor. The stats are cumulative and non-hierarchical, so you can see the major categories without having to dig down through a tree of events. For example, shadow projection is the sum of all the shadow projections for all lights (across all the views). 
+
+![](StatGPU.png)
+
+The realtime stats are particularly useful if you want to gather stats during gameplay (e.g. for graphing later). You can use the realtime display to profile features enabled on console variables or quality settings, or for in-editor optimization where you want to see the results instantly. 
+
+Stats are declared in code as float counters, e.g.
+`DECLARE_FLOAT_COUNTER_STAT(TEXT("Postprocessing"), Stat_GPU_Postprocessing, STATGROUP_GPU);`
+
+Renderthread code blocks can be instrumented with `SCOPED_GPU_STAT` macros, which work similarly to `SCOPED_DRAW_EVENT`, e.g:
+`SCOPED_GPU_STAT(RHICmdList, Stat_GPU_Postprocessing);`
+
+Unlike the draw events, GPU stats are cumulative. You can add multiple entries for the same stat, and these will get aggregated. 
+
+Anything not marked up explicitly should be included in the catch-all `[unaccounted]` stat. If that gets high, then it indicates something is not being caught in the explicit stats and more tracking macros should be added.
+
+
+
+
+
+
+
+
+
+
+
+
+
+
+
+
+
+
+
+
+
diff --git a/Engine/Documentation/Source/Platforms/VR/Profiling/Tools/ProfilingTools.INT.udn b/Engine/Documentation/Source/Platforms/VR/Profiling/Tools/ProfilingTools.INT.udn
new file mode 100644
--- /dev/null
+++ b/Engine/Documentation/Source/Platforms/VR/Profiling/Tools/ProfilingTools.INT.udn
@@ -0,0 +1,54 @@
+Availability:Docs
+Title: VR Profiling Tools
+Crumbs:%ROOT%, Platforms
+Description:Internal and external tools you can use for profiling your Unreal Engine VR project
+Version:4.14
+Skilllevel:Advanced
+Parent:Platforms/VR/Profiling
+order:2
+related: Engine/Performance
+related:Engine/Performance/Profiler
+related:Engine/Performance/GPU
+related: Platforms/VR/Profiling/Overview
+related: Platforms/VR/Profiling/Considerations
+
+
+Unreal Engine 4 has a [set of CPU and GPU profiling tools](Engine/Performance) that will help drill explicitly into what is costing time in your project. 
+For CPU profiling, you can drill into specific events and timings using the [CPU profiling tool](Engine/Performance/Profiler), which connects to a remote process. 
+For GPU profiling, there are more options.  First of all, and most specifically, there's Unreal Engine 4's [GPU profiler](Engine/Performance/GPU), which can be brought up with 
+**Ctrl + Shift + Comma** in editor builds running with `-game` emulation.  
+
+Additionally, you could go further drilling into more systemic problems using both GPUView and Nsight. 
+
+## RenderDoc
+
+RenderDoc is an external tool you can also use for profiling Unreal Engine projects. 
+To use it with UE4:
+
+1. Download the latest [RenderDoc](https://renderdoc.org/builds).
+1. In RenderDoc, configure it to launch your game (e.g. UE4Editor.exe, with appropriate command line arguments).
+1. Run `ToggleDrawEvents` from the command line of your app to get human readable events.
+1. Press F12 to capture a frame.
+
+
+## Oculus PerfHud
+
+Oculus PerfHud is a tool that can be displayed in game, whether in editor or in packaged builds, that will help verify the actual CPU and GPU timings, excluding any throttling done by the application.  There's a great overview of the screens 
+and what the values mean in Oculus' [documentation](https://developer3.oculus.com/documentation/pcsdk/latest/concepts/dg-hud/).
+
+In general, you'll want to look at the graphs and values on the **Application Render Timing** screen.  These will give you the CPU and GPU timings as the Oculus compositor sees them, and will be more accurate for Draw and GPU timings.
+
+To enable the PerfHud in UE4, type the following in the console:
+
+	hmd setint PerfHudMode N
+	
+where N is:
+
+| N | Display type |
+| --- | -----------|
+|	0	| Turns off the performance HUD |
+|	1	| Shows performance summary and headroom |
+|	2	| Shows latency related timing info |
+|	3	| Shows render timing info for application |
+|	4	| Shows render timing info for OVR compositor |
+|	5	| Shows SDK & HMD version Info |
diff --git a/Engine/Documentation/Source/Platforms/VR/Profiling/VRPerformanceAndProfiling.INT.udn b/Engine/Documentation/Source/Platforms/VR/Profiling/VRPerformanceAndProfiling.INT.udn
new file mode 100644
--- /dev/null
+++ b/Engine/Documentation/Source/Platforms/VR/Profiling/VRPerformanceAndProfiling.INT.udn
@@ -0,0 +1,15 @@
+Availability:Docs
+Title: VR Performance and Profiling
+Crumbs:%ROOT%, Platforms
+Description:Tools and strategies for profiling your virtual reality projects in UE4
+Version:4.14
+Parent:Platforms/VR
+Skilllevel:Advanced
+type:overview
+
+Because of the hardware requirements of virtual reality, as well as the stronger impact poor performance can have on your users' comfort, maintaining good VR project performance is key. Unreal Engine 4 contains
+a suite of tools that allow you to view the performance of your project on both the CPU and the GPU. By interpreting those results, and then applying best content and gameplay code practices, you can provide
+the best possible experience with your project.
+
+[DIR(output:"listbutton" parent:"Platforms/VR/Profiling")]
+
diff --git a/Engine/Documentation/Source/Platforms/VR/SplashScreens/VRSplashScreens.INT.udn b/Engine/Documentation/Source/Platforms/VR/SplashScreens/VRSplashScreens.INT.udn
new file mode 100644
--- /dev/null
+++ b/Engine/Documentation/Source/Platforms/VR/SplashScreens/VRSplashScreens.INT.udn
@@ -0,0 +1,289 @@
+Availability:Docs
+Title:VR Splash Screens
+Crumbs:%ROOT%, Platforms, Platforms/VR
+Description:Setting up and working with VR splash screens and movies.
+Navigation:platform
+Version: 4.14
+parent:Platforms/VR
+Type:overview
+Skilllevel: Intermediate
+
+[TOC(start:2)]
+
+Unreal Engine 4 (UE4) offers support for displaying movies or images while a level loads when using Oculus, Gear VR, SteamVR, and PSVR Head Mounted Displays (HMD). The loading move or image will run on the rendering thread and can also help to mask frame rate issues while your content is loading. In the following document, we will take a look at how you can set these VR splash screens in your UE4 projects.
+
+## VR Splash Screen Supported Platforms 
+
+VR Splash Screens will only work when used with the following VR platforms. 
+
+* **Oculus VR**
+* **Gear VR**
+* **Steam VR**
+* **PSVR**
+
+## VR Splash Screen Project Setup
+
+Before you can setup and use the VR splash screens you will first need to create a new **Game Mode** and **Pawn** Blueprint to make it easier to work with the VR splash screens. In the following section, we will go over what Blueprints need to be created and how they need to be setup.
+
+1. Right - click in the **Content Browser**, and from the **Create Basic Asset** list select the **Blueprint Class** option.
+
+    ![](VRLS_Intro_01.png)
+
+1. In the **All Classes** section of the **Pick Parent Class** dialogue box, search for **Game Mode**. When Game Mode is located, select it and then press the **Select** button and give it a name of **VR_GameMode**.
+
+    ![](VRLS_Intro_02.png)
+
+1. Repeat the above process but this time select the **Pawn** option from the Pick Parent Class dialogue box and call this new Pawn, **VR_Pawn**.
+
+    ![](VRLS_Intro_03.png)
+
+1. Open up the VR_GameMode Blueprint and in the **Classes** section, set the **Default Pawn Class** to the VR_Pawn that was created. When completed, Compile and Save the VR_Gamemode Blueprint.
+
+    ![](VRLS_Intro_04.png)
+
+Now that the Pawn and Game Mode have been set up, in the next section we will take a look at how to setup the VR splash screens in the Pawn's Event Graph.
+
+## VR Splash Screen Setup
+
+In the following section, we will go over how you set up a VR splash screen in your Pawns Event Graph.
+
+[region:note]
+Please note that the following is showing the simplest way to set VR splash screens up. While each project has unique needs, the overall setup for VR splash screens will be very similar to what is described.
+[/region]
+1. Open up the VR_Pawn and go to the **Event Graph**. Right - click inside the Event Graph and from the displayed menu, search for and add the following Blueprint nodes.
+
+    * Enable Auto Loading Splash Screen
+    * Set Splash Screen
+    * Hide Splash Screen
+
+    ![](VRLS_Setup_01.png)
+
+1. With the nodes now added to the Event graph, arrange and connect the nodes so that they match the setup in the following image.
+
+    ![](VRLS_Setup_02.png)
+
+1. Since we want the VR splash screen to load any time a level is loaded, on the Enable Auto Loading Splash Screen node, check the **In Auto Show Enabled**.
+
+    ![](VRLS_Setup_03.png)
+
+1. The Set Splash Screen node is where you specify which Texture to use for the VR splash screen. To set a set a VR splash screen Texture, select a Texture in the Content Browser and then press the small white arrow in the Texture section of the Ste Splash Screen node to use that Texture as the VR splash screen.
+
+    ![](VRLS_Setup_04.png)
+    [region:note]
+    For this example, a noise Texture was created in Photoshop and then imported into UE4 and used. Feel free to use any Texture you want for this section.
+    [/region]
+    [region:warning]
+    If you are targeting **SteamVR** you will need to make sure that you set the **Compression Settings** of the Texture you are using for the splash screen to **UserInterface2D(RGBA)** or it will not be displayed when the splash screen is called.
+
+    ![](VIL_SteamVR_Comp_Settings.png)
+    [/region]
+
+1. Finally, on the Set Splash Screen node check the **Show on Set** option to make sure that the Texture you specified is displayed when a level is loaded.
+
+    ![](VRLS_Setup_05.png)
+
+With the VR splash screen now set up in the next section, we will take a look at how to go about calling the VR splash screen when a new level is loaded.
+
+## VR Splash Screen Use
+
+Now that we have the Pawn setup to use the VR splash screens, in this section we will go over how to display the VR Splash Screens when a level is loaded.
+
+1. For this to work you will need to create **two** levels that we can load. To create the levels go to the **Main Toolbar** and select the **File** > **New Level** option. When the **New Level** popup box is displayed, select the **Default** level option.  
+
+    [REGION:lightbox]
+    [![](VLI_NewLevel_00.png)(w:746)](VLI_NewLevel_00.png)
+    [/REGION]
+
+    [REGION:caption]
+    Click for full image.
+    [/REGION]
+
+1. Inside of the **Content Browser** press the **Save All** button and name the new level **Level_00**. Once it has saved, **Right - Click** on Level_00 and from the context menu select the **Duplicate** option and name it **Level_01**. When all of that is completed, you should now have two new level in the Content Browser like in the image below. 
+
+    [REGION:lightbox]
+    [![](VLI_NewLevel_01.png)(w:746)](VLI_NewLevel_01.png)
+    [/REGION]
+
+    [REGION:caption]
+    Click for full image.
+    [/REGION]
+
+1. Open up each level and go to **World Settings** > **Game Mode** and the **Game Mode Override** to the VR_GameMode Blueprint that was created. This way we can be sure that the correct Pawn will be used when the level is loaded.
+
+    ![](VLI_NewLevel_06.png)
+
+1. Open up the VR_Pawn Blueprint, if not already open, and in the Event, graph add key inputs for the **R** and **T** keys.
+
+    [REGION:lightbox]
+    [![](VLI_NewLevel_02.png)(w:746)](VLI_NewLevel_02.png)
+    [/REGION]
+
+    [REGION:caption]
+    Click for full image.
+    [/REGION]
+
+1.  Right - Click on the Event Graph and search for the **Open Level** Blueprint node and add it to the Event Graph. Connect the **Pressed** output for the T key to the Open Level input and in the **Level Name** option input **Level_00**.
+
+    [REGION:lightbox]
+    [![](VLI_NewLevel_03.png)(w:746)](VLI_NewLevel_03.png)
+    [/REGION]
+
+    [REGION:caption]
+    Click for full image.
+    [/REGION]
+
+1. Select the Open Level node and then press **CTRL + W** on the keyboard to duplicate it. Connect the Pressed output for the **R** key to the new Open Level node and change the name of the Level to **Level_01**.
+
+    [REGION:lightbox]
+    [![](VLI_NewLevel_04.png)(w:746)](VLI_NewLevel_04.png)
+    [/REGION]
+
+    [REGION:caption]
+    Click for full image.
+    [/REGION]
+
+1. Press the **Compile** and **Save** buttons to compile and save your Pawn Blueprint and then launch your project using the VR Preview.
+
+    [REGION:lightbox]
+    [![](VLI_NewLevel_05.png)(w:746)](VLI_NewLevel_05.png)
+    [/REGION]
+
+    [REGION:caption]
+    Click for full image.
+    [/REGION]
+
+1. Put your HMD on and then press the **R** key to load Level_01. While the level is being loaded you should see the image that you specified in the Texture input of the Set Splash Screen node being displayed. Once Level_01 loads press the **T** key to load Level_00, and you should now see the same image being displayed. 
+
+    [OBJECT:EmbeddedVideo]
+        [PARAMLITERAL:width]
+        640
+        [/PARAMLITERAL]
+        [PARAMLITERAL:height]
+        360
+        [/PARAMLITERAL]
+        [PARAMLITERAL:videoid]
+        DWkuJhEOtwc
+        [/PARAMLITERAL]
+    [/OBJECT]
+
+With the ability to change levels on button presses now working, the next step would be to try and get this working when a player enters a level streaming volume or trigger as you would in any game or interactive project.
+
+## VR Splash Screen Node Reference
+
+The following breakdown covers what each one of the properties on the VR splash screen nodes does. When looking at the VR splash screen nodes, you will notice that there is a generic type and a type specific to Oculus. You should only use the Oculus-specific types when you are targeting Oculus HMD's. All other HMD's should use the generic type.
+
+### Generic Splash Screens
+
+* Enable Auto Loading Splash Screen
+
+    ![](VRLS_EnableAutoLoading.png)
+
+    |Setting  Name| Description|
+    |---------------|----------------------------------------------------|
+    |In Auto Show Enabled|If true, load the splash screen when a level is loaded.|
+
+* Hide Splash Screen
+
+    ![](VRLS_Hide_Splash.png)
+
+    |Setting  Name| Description|
+    |---------------|----------------------------------------------------|
+    |Hide Splash Screen|Hides the splash screen.|
+
+* Set Splash Screen
+
+    ![](VRLS_Set_Splash.png)
+
+    |Setting  Name| Description|
+    |---------------|----------------------------------------------------|
+    |Texture|Texture to use fo the splash. Must be in B8R8G8A8 format.|
+    |Scale|Scale of the Texture.|
+    |Offest|Position from which to start rendering the Texture.|
+    |Show Loading Moive|Show a loading movie.|
+    |Show on Set|Show the Texture specified when the compositor is called.|
+
+* Show Splash Screen
+
+    ![](VRLS_Show_Splash.png)
+
+    |Setting  Name| Description|
+    |---------------|----------------------------------------------------|
+    |Show Splash Screen|Shows the splash screen.|
+
+### Oculus Specific Splash Screens
+
+The following breakdown covers what each of the Oculus-specific splash screen nodes will do. Keep in mind that the following nodes should only be used when you are targeting Oculus HMD's. 
+
+* Add Loading Splash Screen
+
+    ![](VRLS_OC_Add_Loading.png)
+
+    |Setting  Name| Description|
+    |---------------|----------------------------------------------------|
+    |Texture|A Texture asset to be used for the splash screen. GearVR uses it as a path for loading icon; all other params are currently ignored by GearVR.|
+    |Translation in Meters|Initial translation of the center of the splash screen in meters.|
+    |Rotation|Initial rotation of the splash screen, with the origin at the center of the splash screen.|
+    |Size in Meters|Size, in meters, of the quad with the splash screen.|
+    |Delta Rotation|Incremental rotation, that is added each 2nd frame to the quad transform. The quad is rotated around the center of the quad.|
+    |Ckear Before Add |If true clears splash before adding a new one.|
+
+* Clear Loading Splash Screen
+
+    ![](VRLS_OC_Clear_Loading_Splash.png)
+
+    |Setting  Name| Description|
+    |---------------|----------------------------------------------------|
+    |Clear Loading Splash Screen|Removes all splash screens.|
+
+* Enable Auto Loading Splash Screen
+
+    ![](VRLS_OC_Enable_Auto_Load.png)
+
+    |Setting  Name| Description|
+    |---------------|----------------------------------------------------|
+    |Auto Show Enabled|If true, automatic showing of splash screens is enabled when a map is being loaded.|
+
+* Hide Loading Icon
+
+    ![](VRLS_OC_Hide_Loading.png)
+
+    |Setting  Name| Description|
+    |---------------|----------------------------------------------------|
+    |Hide Loading Icon|Clears the loading icon. This call will clear all the splashes.|
+
+* Hide Loading Splash Screen
+
+    ![](VRLS_OC_Hide_Loading_SS.png)
+
+    |Setting  Name| Description|
+    |---------------|----------------------------------------------------|
+    |Clear|Clear all splash screen after hiding.|
+
+* Is Auto Loading Splash Screen Enabled
+
+    ![](VRLS_OC_IsAutoLoasding.png)
+
+    |Setting  Name| Description|
+    |---------------|----------------------------------------------------|
+    |Is Auto Loading Splash Screen Enabled| Returns true if the splash screen is automatically shown when LoadMap is called.|
+
+* Is Loading Icon Enabled
+
+    ![](VRLS_OC_IsLoadingIconEnabled.png)
+
+    |Setting  Name| Description|
+    |---------------|----------------------------------------------------|
+    |Is Loading Icon Enabled| Returns true if the splash screen is in loading icon mode.|
+
+* Show Loading Splash Screen
+
+    ![](VRLS_OC_ShowLoadingSplashScreen.png)
+
+    |Setting  Name| Description|
+    |---------------|----------------------------------------------------|
+    |Show Loading Splash Screen| Shows loading splash screen.|
+
+
+
+
+
diff --git a/Engine/Documentation/Source/Platforms/VR/StereoLayers/StereoLayers.INT.udn b/Engine/Documentation/Source/Platforms/VR/StereoLayers/StereoLayers.INT.udn
new file mode 100644
--- /dev/null
+++ b/Engine/Documentation/Source/Platforms/VR/StereoLayers/StereoLayers.INT.udn
@@ -0,0 +1,146 @@
+Availability:Docs
+Title:VR Stereo Layers
+Crumbs:%ROOT%, Platforms, Platforms/VR
+Description:Setting up and working with VR Stereo Layers.
+Navigation:platform
+Version: 4.14
+parent:Platforms/VR
+type:overview
+Skilllevel: Intermediate
+tags:VR
+tags:Oculus Rift
+tags:GearVR
+tags:SteamVR
+tags:PSVR
+
+[TOC(start:2)]
+
+Virtuality Reality (VR) Stereo Layers allow you to send a separate Texture to the VR Head Mounted Display (HMD) and have it re-projected on its own in a separate rendering pass than the rest of the project.  This is extremely useful for User Interface (UI) elements that you want to attach to the HMD but not be affected by post processing or anti-aliasing. In the following document, we will go over how you set up and use Stereo Layers across all VR Platforms that are supported.  
+
+
+## Supported Platforms 
+
+VR Stereo Layers will only work when used with the following VR platforms. 
+
+* **Oculus VR**
+* **Gear VR**
+* **Steam VR**
+* **PSVR**
+
+## Stereo Layer Setup
+
+To use Stereo Layers with your **VR Pawn**, you will need to do the following.
+
+1. Open up the **Pawn** that you want to use Stereo Layers with and then go to the **Viewport** pannel.
+
+    [REGION:lightbox]
+    [![](VRSL_Pawn_Setup_01.png)(w:667)](VRSL_Pawn_Setup_01.png)
+    [/REGION]
+
+    [REGION:caption]
+    Click for full image.
+    [/REGION]
+
+1. Click on the **Add Components** button and from the drop-down list that is displayed and search for the **Stereo Layer** Component. When you find it, click on it to add it to your **Pawn**.
+
+    [REGION:lightbox]
+    [![](VRSL_Pawn_Setup_02.png)(w:667)](VRSL_Pawn_Setup_02.png)
+    [/REGION]
+
+    [REGION:caption]
+    Click for full image.
+    [/REGION]
+
+1. Now drag the **Stereo Layer Component** on top of the **Pawn's** camera to make it a child actor of the camera.
+    
+    ![](VRSL_Pawn_Setup_03.png)(convert:false)
+
+1. Select the **Stereo Layer Component** from the Components list and then in the **Details Panel** under the **Transform** section, set the **X** value of **Location** to **100** so that it is easy to see in VR.
+
+    [REGION:lightbox]
+    [![](VRSL_Pawn_Setup_04.png)(w:667)](VRSL_Pawn_Setup_04.png)
+    [/REGION]
+
+    [REGION:caption]
+    Click for full image.
+    [/REGION]
+
+1. Over in the **Details Panel**  expand the **Stereo Layer** section.
+
+    [REGION:lightbox]
+    [![](VRSL_Pawn_Setup_05.png)(w:667)](VRSL_Pawn_Setup_05.png)
+    [/REGION]
+
+    [REGION:caption]
+    Click for full image.
+    [/REGION]
+
+1. Inside of the **Content Browser** locate a Texture that you would want to display and click on it to select it. Then back in the **Pawn** Blueprint under the Stereo Layers section, press the small **white triangle** next to the **Texture** input to use that Texture the Stereo Layer should display.
+
+    [REGION:lightbox]
+    [![](VRSL_Pawn_Setup_06.png)(w:667)](VRSL_Pawn_Setup_06.png)
+    [/REGION]
+
+    [REGION:caption]
+    Click for full image.
+    [/REGION]
+
+1. Next press the **Play in VR** button to launch your project in VR. When you put your HMD on you should see the Texture you specified in Texture input being displayed in front of you. The Texture should also follow your HMD precisely and as quickly as you can move your head around like in the video below.
+
+    [OBJECT:EmbeddedVideo]
+        [PARAMLITERAL:width]
+        640
+        [/PARAMLITERAL]
+        [PARAMLITERAL:height]
+        360
+        [/PARAMLITERAL]
+        [PARAMLITERAL:videoid]
+        Eb_kndx6tac
+        [/PARAMLITERAL]
+    [/OBJECT]
+
+    [region:note]
+    If the Texture you selected does not show up when your project is viewed in VR, make sure you are using the correct Texture Compressions settings. If you are unsure what those settings should be for the HMD you are using, please check out the Stereo Layers Texture Compression Settings section.
+    [/region]
+
+## Stereo Layers Texture Compression Settings
+
+Depending on which HMD you are using you will need to make adjustments to the Compression setting that is used for the Texture to be displayed. Below you will find more information about which Texture Compression settings are needed for Stereo Layers to display correctly on various HMDs.
+
+[REGION:lightbox]
+[![](VRSL_Texture_Settings.png)(w:667)](VRSL_Texture_Settings.png)
+[/REGION]
+
+[REGION:caption]
+Click for full image.
+[/REGION]
+
+|HMD Name| Required Texture Compression Settings|
+|---------------|----------------------------------------------------|
+|Oculus Rift| None|
+|Gear VR| None|
+| HTC Vive| UserInterFace2D (RGBA)|
+|PSVR| None|
+
+## Stereo Layer Properties Breakdown
+
+Stereo Layers have a number of properties that affect the way in which it is rendered. The following table breakdown each one of these settings and it's affect on the the Stereo Layers.
+
+![](VRSL_Prop_Breakdown_00.png)
+
+|Setting  Name| Description|
+|---------------|----------------------------------------------------|
+|Live Texture|True if the stereo layer texture needs to update itself every frame(scene capture, video, etc.)|
+|Supports Depth|True if the stereo layer needs to support depth intersections with the scene geometry, if available on the platform.|
+|No Alpha Channel|True if the texture should not use its own alpha channel(1.0 will be substituted).|
+|Texture|Texture displayed on the stereo layer (is stereoscopic textures are supported on the platform and more than one texture is provided, this will be the right eye).|
+|Quad Preserve Texture Ratio|True if the quad should internally set it's Y value based on the set texture's dimensions.|
+|Stereo Layer Type|Specifies how and where the quad is rendered to the screen.|
+|Stereo Layer Shape|Specifies which type of layer it is. Note that some shapes will be supported only on certain platforms!|
+|Priority|Render priority among all stereo layers, higher priority render on top of lower priority.|
+|Left Texture|Texture displayed on the stereo layer for left eye, if stereoscopic textures are supported on the platform.|
+|Quad Size|Size of the rendered stereo layer quad.|
+|UVRect|UV coordinates mapped to the quad face.|
+|Cylinder Radius|Radial size of the rendered stereo layer cylinder.|
+|Cylinder Overly Arc|Arc angle for the stereo layer cylinder.|
+|Cylinder Height|Height of the stereo layer cylinder.|
\ No newline at end of file
diff --git a/Engine/Documentation/Source/Platforms/VR/StereoPanoramicCapture/HowTo/CreatingViewerMaterial/CreatingViewerMaterial.INT.udn b/Engine/Documentation/Source/Platforms/VR/StereoPanoramicCapture/HowTo/CreatingViewerMaterial/CreatingViewerMaterial.INT.udn
new file mode 100644
--- /dev/null
+++ b/Engine/Documentation/Source/Platforms/VR/StereoPanoramicCapture/HowTo/CreatingViewerMaterial/CreatingViewerMaterial.INT.udn
@@ -0,0 +1,319 @@
+Availability: Docs
+Title:3. Creating a Material for Viewing Stereo Images
+Crumbs: 
+Description:Here we will create and set up a new Material to view our stereo screenshots. 
+Version: 4.13
+SkillLevel: Intermediate
+type:multi-step
+tags: VR
+order:3
+
+[NAV]
+
+In this section, we will create and set up a Material that will allow us to view the stereo images that were created in the previous section. 
+
+1. Locate the stereo images that were created in the previous section by going to your projects saved folder and then going to **StereoPanoramaImages\Saved\StereoPanorama\(Date & Time it was taken)**.
+
+    [REGION:lightbox]
+    [![](Platforms\VR\StereoPanoramicCapture\HowTo\TakingStereoImages\SPHT_Rendered_Images.png)(w:473)](Platforms\VR\StereoPanoramicCapture\HowTo\TakingStereoImages\SPHT_Rendered_Images.png)
+    
+    [/REGION]
+
+    [REGION:caption]
+    Click for full image.
+    [/REGION]
+
+1. Once you locate both images, select both images and then drag them from their location in Windows into the Content Browser which will import the images into your UE4 project.
+
+    [REGION:lightbox]
+    [![](SPHT_Import_Stereo_Images.png)(w:639)](SPHT_Import_Stereo_Images.png)
+    [/REGION]
+
+    [REGION:caption]
+    Click for full image.
+    [/REGION]
+
+1. Once both images have been imported, double-click on both of the images to open them up in the Texture Viewer and once opened, go to the **Level Of Detail** section and under **Mip Gen Settings**, change the setting from **From Texture Group** to **No Mipmaps**.
+
+    [REGION:lightbox]
+    [![](SPHT_Mip_Gen_Settings.png)(w:439)](SPHT_Mip_Gen_Settings.png)
+    [/REGION]
+
+    [REGION:caption]
+    Click for full image.
+    [/REGION]
+
+    [region:note]
+    Failing to do this could result in the textures showing a seam when viewed in stereo due to the distance this texture will be from the camera. 
+    [/region]
+
+1. Now Right - click in the Content Browser and from the **Create Basic Asset** section click on the Material option and give the newly created Material a name of **MAT_StereoViewer**.
+
+    [REGION:lightbox]
+    [![](SPHT_Create_Stereo_Material.png)(w:612)](SPHT_Create_Stereo_Material.png)
+    [/REGION]
+
+    [REGION:caption]
+    Click for full image.
+    [/REGION]
+
+1. Double-click on the MAT_StereoViewer Material to open it up and in the **Details** panel set the following options.
+
+    ![](SPHT_Material_Settings.png)
+
+    * Shading Model: Unlit
+    * Fully Rough: True
+    * Use Full Precision: True
+    * Use LIghtmap Directionality: False
+
+1. Now click on the **Copy Material Expression Graph** option that is below and copy the provided text and paste it into the Material Editor. After the text has been pasted into the Material Graph, make sure to wire up the output of the **If** node into the **Emissive Color** on the Main Material node. When completed your Material should look like the following image.
+
+    [Include:Platforms\VR\StereoPanoramicCapture\HowTo\CreatingViewerMaterial\#SPMaterial]
+
+1. Now that the viewer Material has been set up and connected, locate the left and right eye Textures in the Content Browser and then load the Textures into the corresponding eye Texture input and then press the **Apply** AND **Save** buttons to compile and save the Material. 
+
+    [REGION:lightbox]
+    [![](SPHT_Completed_Material.png)(w:439)](SPHT_Completed_Material.png)
+    [/REGION]
+
+    [REGION:caption]
+    Click for full image.
+    [/REGION]
+
+
+With the Material now set up and ready to be used, in the next section, we will look at how to set up a level to display the Material when the level is run.
+
+
+
+[NAV]
+
+<!-----
+[EXCERPT:SPMaterial]    
+[OBJECT:Snippet]
+    [PARAMLITERAL:id]
+        SPMaterial
+    [/PARAMLITERAL]
+    [PARAMLITERAL:type]
+        material
+    [/PARAMLITERAL]
+    [PARAM:image]
+        ![](SPHT_Stereo_Viewer_Material.png)
+    [/PARAM]
+    [PARAM:snippet]
+                            Begin Object Class=MaterialGraphNode Name="MaterialGraphNode_0"
+               Begin Object Class=MaterialExpressionIf Name="MaterialExpressionIf_0"
+               End Object
+               Begin Object Name="MaterialExpressionIf_0"
+                  A=(Expression=MaterialExpressionMaterialFunctionCall'MaterialGraphNode_2.MaterialExpressionMaterialFunctionCall_0')
+                  B=(Expression=MaterialExpressionConstant'MaterialGraphNode_1.MaterialExpressionConstant_0')
+                  AGreaterThanB=(Expression=MaterialExpressionTextureSampleParameter2D'MaterialGraphNode_7.MaterialExpressionTextureSampleParameter2D_0',Mask=1,MaskR=1,MaskG=1,MaskB=1)
+                  ALessThanB=(Expression=MaterialExpressionTextureSampleParameter2D'MaterialGraphNode_8.MaterialExpressionTextureSampleParameter2D_1',Mask=1,MaskR=1,MaskG=1,MaskB=1)
+                  MaterialExpressionEditorX=64
+                  MaterialExpressionEditorY=144
+                  MaterialExpressionGuid=81A589084F901693C24D6188C7E9963C
+                  Material=PreviewMaterial'/Engine/Transient.MAT_Stereo_Image'
+               End Object
+               MaterialExpression=MaterialExpressionIf'MaterialExpressionIf_0'
+               NodePosX=64
+               NodePosY=144
+               NodeGuid=5B7FA1FE457D602E7A7E578891582634
+               CustomProperties Pin (PinId=A0EE7D184D1901AC89D7CCB7D7B2AF54,PinName="A",PinType.PinCategory="required",PinType.PinSubCategoryObject=None,PinType.PinSubCategoryMemberReference=(),PinType.bIsArray=False,PinType.bIsReference=False,PinType.bIsConst=False,PinType.bIsWeakPointer=False,LinkedTo=(MaterialGraphNode_2 176CDC95486F3BF845430099EF2BE7FE,),PersistentGuid=00000000000000000000000000000000,bHidden=False,bNotConnectable=False,bDefaultValueIsReadOnly=False,bDefaultValueIsIgnored=False,bAdvancedView=False,)
+               CustomProperties Pin (PinId=D37D419748FE491453C139AF821D8F11,PinName="B",PinType.PinCategory="optional",PinType.PinSubCategoryObject=None,PinType.PinSubCategoryMemberReference=(),PinType.bIsArray=False,PinType.bIsReference=False,PinType.bIsConst=False,PinType.bIsWeakPointer=False,LinkedTo=(MaterialGraphNode_1 62DD318C41832241F7AB73BA06FB692E,),PersistentGuid=00000000000000000000000000000000,bHidden=False,bNotConnectable=False,bDefaultValueIsReadOnly=False,bDefaultValueIsIgnored=False,bAdvancedView=False,)
+               CustomProperties Pin (PinId=44EC5EE54EACAD4AA705D8866FE16EB5,PinName="A > B",PinType.PinCategory="required",PinType.PinSubCategoryObject=None,PinType.PinSubCategoryMemberReference=(),PinType.bIsArray=False,PinType.bIsReference=False,PinType.bIsConst=False,PinType.bIsWeakPointer=False,LinkedTo=(MaterialGraphNode_7 6073EDCF453D52887F1ED597234FAF0C,),PersistentGuid=00000000000000000000000000000000,bHidden=False,bNotConnectable=False,bDefaultValueIsReadOnly=False,bDefaultValueIsIgnored=False,bAdvancedView=False,)
+               CustomProperties Pin (PinId=6E32A98741E5BFCB34BE82870AE85A02,PinName="A == B",PinType.PinCategory="optional",PinType.PinSubCategoryObject=None,PinType.PinSubCategoryMemberReference=(),PinType.bIsArray=False,PinType.bIsReference=False,PinType.bIsConst=False,PinType.bIsWeakPointer=False,PersistentGuid=00000000000000000000000000000000,bHidden=False,bNotConnectable=False,bDefaultValueIsReadOnly=False,bDefaultValueIsIgnored=False,bAdvancedView=False,)
+               CustomProperties Pin (PinId=F937930D47BA8EA98B9510B0C025B9A4,PinName="A < B",PinType.PinCategory="required",PinType.PinSubCategoryObject=None,PinType.PinSubCategoryMemberReference=(),PinType.bIsArray=False,PinType.bIsReference=False,PinType.bIsConst=False,PinType.bIsWeakPointer=False,LinkedTo=(MaterialGraphNode_8 CDF923064C8491C8DD875A9FD733B329,),PersistentGuid=00000000000000000000000000000000,bHidden=False,bNotConnectable=False,bDefaultValueIsReadOnly=False,bDefaultValueIsIgnored=False,bAdvancedView=False,)
+               CustomProperties Pin (PinId=33CE8B8B4E0E854F7C6164AD5735E5F2,PinName="Output",PinFriendlyName=" ",Direction="EGPD_Output",PinType.PinSubCategoryObject=None,PinType.PinSubCategoryMemberReference=(),PinType.bIsArray=False,PinType.bIsReference=False,PinType.bIsConst=False,PinType.bIsWeakPointer=False,LinkedTo=(MaterialGraphNode_Root_0 C45D0E1F4E227CF5D5EC2C9805D42CE9,),PersistentGuid=00000000000000000000000000000000,bHidden=False,bNotConnectable=False,bDefaultValueIsReadOnly=False,bDefaultValueIsIgnored=False,bAdvancedView=False,)
+            End Object
+            Begin Object Class=MaterialGraphNode Name="MaterialGraphNode_1"
+               Begin Object Class=MaterialExpressionConstant Name="MaterialExpressionConstant_0"
+               End Object
+               Begin Object Name="MaterialExpressionConstant_0"
+                  R=0.500000
+                  MaterialExpressionEditorX=-80
+                  MaterialExpressionEditorY=128
+                  MaterialExpressionGuid=3D90A5C443D14DF0CF49A68F6E13583F
+                  Material=PreviewMaterial'/Engine/Transient.MAT_Stereo_Image'
+               End Object
+               MaterialExpression=MaterialExpressionConstant'MaterialExpressionConstant_0'
+               NodePosX=-80
+               NodePosY=128
+               NodeGuid=6190E3714CEDD46F0F98D8B9F38E5751
+               CustomProperties Pin (PinId=62DD318C41832241F7AB73BA06FB692E,PinName="Output",PinFriendlyName=" ",Direction="EGPD_Output",PinType.PinSubCategoryObject=None,PinType.PinSubCategoryMemberReference=(),PinType.bIsArray=False,PinType.bIsReference=False,PinType.bIsConst=False,PinType.bIsWeakPointer=False,LinkedTo=(MaterialGraphNode_0 D37D419748FE491453C139AF821D8F11,),PersistentGuid=00000000000000000000000000000000,bHidden=False,bNotConnectable=False,bDefaultValueIsReadOnly=False,bDefaultValueIsIgnored=False,bAdvancedView=False,)
+            End Object
+            Begin Object Class=MaterialGraphNode Name="MaterialGraphNode_2"
+               Begin Object Class=MaterialExpressionMaterialFunctionCall Name="MaterialExpressionMaterialFunctionCall_0"
+               End Object
+               Begin Object Name="MaterialExpressionMaterialFunctionCall_0"
+                  MaterialFunction=MaterialFunction'/Engine/Functions/Engine_MaterialFunctions02/Utility/BreakOutFloat2Components.BreakOutFloat2Components'
+                  FunctionInputs(0)=(ExpressionInputId=0FC5C56B4335798F20CB928CF1C87421,Input=(Expression=MaterialExpressionScreenPosition'MaterialGraphNode_3.MaterialExpressionScreenPosition_0',InputName="Float2"))
+                  FunctionOutputs(0)=(ExpressionOutputId=6BC893C247FCFC567178CE944E213385,Output=(OutputName="R"))
+                  FunctionOutputs(1)=(ExpressionOutputId=9720C81E449DF8E5514FCB8A80FD5FA3,Output=(OutputName="G"))
+                  MaterialExpressionEditorX=-272
+                  MaterialExpressionGuid=0DB6A16C443E70F09CC35BAA076DD2C3
+                  Material=PreviewMaterial'/Engine/Transient.MAT_Stereo_Image'
+                  Outputs(0)=(OutputName="R")
+                  Outputs(1)=(OutputName="G")
+               End Object
+               MaterialExpression=MaterialExpressionMaterialFunctionCall'MaterialExpressionMaterialFunctionCall_0'
+               NodePosX=-272
+               NodeGuid=43D001084D713B8304D69BBC75A0F775
+               CustomProperties Pin (PinId=361F5A0A417E832EB7A5E48C380D9DA4,PinName="Float2 (V2)",PinType.PinCategory="optional",PinType.PinSubCategoryObject=None,PinType.PinSubCategoryMemberReference=(),PinType.bIsArray=False,PinType.bIsReference=False,PinType.bIsConst=False,PinType.bIsWeakPointer=False,LinkedTo=(MaterialGraphNode_3 A6AE01F64B365D34D81BD981E64175F1,),PersistentGuid=00000000000000000000000000000000,bHidden=False,bNotConnectable=False,bDefaultValueIsReadOnly=False,bDefaultValueIsIgnored=False,bAdvancedView=False,)
+               CustomProperties Pin (PinId=176CDC95486F3BF845430099EF2BE7FE,PinName="R",Direction="EGPD_Output",PinType.PinSubCategoryObject=None,PinType.PinSubCategoryMemberReference=(),PinType.bIsArray=False,PinType.bIsReference=False,PinType.bIsConst=False,PinType.bIsWeakPointer=False,LinkedTo=(MaterialGraphNode_0 A0EE7D184D1901AC89D7CCB7D7B2AF54,),PersistentGuid=00000000000000000000000000000000,bHidden=False,bNotConnectable=False,bDefaultValueIsReadOnly=False,bDefaultValueIsIgnored=False,bAdvancedView=False,)
+               CustomProperties Pin (PinId=C686EF7142811E46C988739585372C29,PinName="G",Direction="EGPD_Output",PinType.PinSubCategoryObject=None,PinType.PinSubCategoryMemberReference=(),PinType.bIsArray=False,PinType.bIsReference=False,PinType.bIsConst=False,PinType.bIsWeakPointer=False,PersistentGuid=00000000000000000000000000000000,bHidden=False,bNotConnectable=False,bDefaultValueIsReadOnly=False,bDefaultValueIsIgnored=False,bAdvancedView=False,)
+            End Object
+            Begin Object Class=MaterialGraphNode Name="MaterialGraphNode_3"
+               Begin Object Class=MaterialExpressionScreenPosition Name="MaterialExpressionScreenPosition_0"
+               End Object
+               Begin Object Name="MaterialExpressionScreenPosition_0"
+                  MaterialExpressionEditorX=-448
+                  MaterialExpressionEditorY=-32
+                  MaterialExpressionGuid=F5C4FBB24FB4BE834B8623BAD82AA1FF
+                  Material=PreviewMaterial'/Engine/Transient.MAT_Stereo_Image'
+               End Object
+               MaterialExpression=MaterialExpressionScreenPosition'MaterialExpressionScreenPosition_0'
+               NodePosX=-448
+               NodePosY=-32
+               NodeGuid=0F5E2A6C4507E7940A41E281E944C1D8
+               CustomProperties Pin (PinId=A6AE01F64B365D34D81BD981E64175F1,PinName="Output",PinFriendlyName=" ",Direction="EGPD_Output",PinType.PinSubCategoryObject=None,PinType.PinSubCategoryMemberReference=(),PinType.bIsArray=False,PinType.bIsReference=False,PinType.bIsConst=False,PinType.bIsWeakPointer=False,LinkedTo=(MaterialGraphNode_2 361F5A0A417E832EB7A5E48C380D9DA4,),PersistentGuid=00000000000000000000000000000000,bHidden=False,bNotConnectable=False,bDefaultValueIsReadOnly=False,bDefaultValueIsIgnored=False,bAdvancedView=False,)
+            End Object
+            Begin Object Class=MaterialGraphNode Name="MaterialGraphNode_4"
+               Begin Object Class=MaterialExpressionMultiply Name="MaterialExpressionMultiply_0"
+               End Object
+               Begin Object Name="MaterialExpressionMultiply_0"
+                  A=(Expression=MaterialExpressionCameraVectorWS'MaterialGraphNode_5.MaterialExpressionCameraVectorWS_0')
+                  B=(Expression=MaterialExpressionVectorParameter'MaterialGraphNode_9.MaterialExpressionVectorParameter_0',Mask=1,MaskR=1,MaskG=1,MaskB=1)
+                  MaterialExpressionEditorX=-720
+                  MaterialExpressionEditorY=304
+                  MaterialExpressionGuid=4F10B53C4ECBDE27F4BD38B46FACAD18
+                  Material=PreviewMaterial'/Engine/Transient.MAT_Stereo_Image'
+                  bCollapsed=False
+               End Object
+               MaterialExpression=MaterialExpressionMultiply'MaterialExpressionMultiply_0'
+               NodePosX=-720
+               NodePosY=304
+               NodeGuid=630A92064469FB428330238D474586C7
+               CustomProperties Pin (PinId=E68092674090694BB584AEAEF50A40E0,PinName="A",PinType.PinCategory="optional",PinType.PinSubCategoryObject=None,PinType.PinSubCategoryMemberReference=(),PinType.bIsArray=False,PinType.bIsReference=False,PinType.bIsConst=False,PinType.bIsWeakPointer=False,LinkedTo=(MaterialGraphNode_5 CF46DC7F490ECC1D0F932B8242298926,),PersistentGuid=00000000000000000000000000000000,bHidden=False,bNotConnectable=False,bDefaultValueIsReadOnly=False,bDefaultValueIsIgnored=False,bAdvancedView=False,)
+               CustomProperties Pin (PinId=F1DE55B14567D0B67759DEA9C934BA05,PinName="B",PinType.PinCategory="optional",PinType.PinSubCategoryObject=None,PinType.PinSubCategoryMemberReference=(),PinType.bIsArray=False,PinType.bIsReference=False,PinType.bIsConst=False,PinType.bIsWeakPointer=False,LinkedTo=(MaterialGraphNode_9 9C23E1264DF642DC356B788B08CB3693,),PersistentGuid=00000000000000000000000000000000,bHidden=False,bNotConnectable=False,bDefaultValueIsReadOnly=False,bDefaultValueIsIgnored=False,bAdvancedView=False,)
+               CustomProperties Pin (PinId=BABCD0D743049748AF3DA996E4665D7D,PinName="Output",PinFriendlyName=" ",Direction="EGPD_Output",PinType.PinSubCategoryObject=None,PinType.PinSubCategoryMemberReference=(),PinType.bIsArray=False,PinType.bIsReference=False,PinType.bIsConst=False,PinType.bIsWeakPointer=False,LinkedTo=(MaterialGraphNode_6 6018F5F544A36CE1C7E511989710AD1D,),PersistentGuid=00000000000000000000000000000000,bHidden=False,bNotConnectable=False,bDefaultValueIsReadOnly=False,bDefaultValueIsIgnored=False,bAdvancedView=False,)
+            End Object
+            Begin Object Class=MaterialGraphNode Name="MaterialGraphNode_5"
+               Begin Object Class=MaterialExpressionCameraVectorWS Name="MaterialExpressionCameraVectorWS_0"
+               End Object
+               Begin Object Name="MaterialExpressionCameraVectorWS_0"
+                  MaterialExpressionEditorX=-912
+                  MaterialExpressionEditorY=256
+                  MaterialExpressionGuid=BA70FEF141C3FC6429AA56B337E6F88C
+                  Material=PreviewMaterial'/Engine/Transient.MAT_Stereo_Image'
+               End Object
+               MaterialExpression=MaterialExpressionCameraVectorWS'MaterialExpressionCameraVectorWS_0'
+               NodePosX=-912
+               NodePosY=256
+               NodeGuid=8F0586154298BB5972CE548ACAD0E0D2
+               CustomProperties Pin (PinId=CF46DC7F490ECC1D0F932B8242298926,PinName="Output",PinFriendlyName=" ",Direction="EGPD_Output",PinType.PinSubCategoryObject=None,PinType.PinSubCategoryMemberReference=(),PinType.bIsArray=False,PinType.bIsReference=False,PinType.bIsConst=False,PinType.bIsWeakPointer=False,LinkedTo=(MaterialGraphNode_4 E68092674090694BB584AEAEF50A40E0,),PersistentGuid=00000000000000000000000000000000,bHidden=False,bNotConnectable=False,bDefaultValueIsReadOnly=False,bDefaultValueIsIgnored=False,bAdvancedView=False,)
+            End Object
+            Begin Object Class=MaterialGraphNode Name="MaterialGraphNode_6"
+               Begin Object Class=MaterialExpressionMaterialFunctionCall Name="MaterialExpressionMaterialFunctionCall_1"
+               End Object
+               Begin Object Name="MaterialExpressionMaterialFunctionCall_1"
+                  MaterialFunction=MaterialFunction'/Engine/Functions/Engine_MaterialFunctions01/Coordinates/LongLatToUV.LongLatToUV'
+                  FunctionInputs(0)=(ExpressionInputId=150D02534FA6F797DFD995BA692D7801,Input=(Expression=MaterialExpressionMultiply'MaterialGraphNode_4.MaterialExpressionMultiply_0',InputName="Vector"))
+                  FunctionOutputs(0)=(ExpressionOutputId=0EB334074DD42132558C7EA6C023EC50,Output=(OutputName="Result"))
+                  MaterialExpressionEditorX=-526
+                  MaterialExpressionEditorY=299
+                  MaterialExpressionGuid=2B540AED4336652F5CC59D9DF979B358
+                  Material=PreviewMaterial'/Engine/Transient.MAT_Stereo_Image'
+                  Outputs(0)=(OutputName="Result")
+               End Object
+               MaterialExpression=MaterialExpressionMaterialFunctionCall'MaterialExpressionMaterialFunctionCall_1'
+               NodePosX=-526
+               NodePosY=299
+               NodeGuid=3A24EE65400825A12AE9DC844E79E18F
+               CustomProperties Pin (PinId=6018F5F544A36CE1C7E511989710AD1D,PinName="Vector (V3)",PinType.PinCategory="required",PinType.PinSubCategoryObject=None,PinType.PinSubCategoryMemberReference=(),PinType.bIsArray=False,PinType.bIsReference=False,PinType.bIsConst=False,PinType.bIsWeakPointer=False,LinkedTo=(MaterialGraphNode_4 BABCD0D743049748AF3DA996E4665D7D,),PersistentGuid=00000000000000000000000000000000,bHidden=False,bNotConnectable=False,bDefaultValueIsReadOnly=False,bDefaultValueIsIgnored=False,bAdvancedView=False,)
+               CustomProperties Pin (PinId=C87298924529FAC8728933A296BA7682,PinName="Result",Direction="EGPD_Output",PinType.PinSubCategoryObject=None,PinType.PinSubCategoryMemberReference=(),PinType.bIsArray=False,PinType.bIsReference=False,PinType.bIsConst=False,PinType.bIsWeakPointer=False,LinkedTo=(MaterialGraphNode_7 A60D4B0340117A9FC89B09BA6D69B3AF,MaterialGraphNode_8 0EA30C1443845B538B9305BFF8C98510,),PersistentGuid=00000000000000000000000000000000,bHidden=False,bNotConnectable=False,bDefaultValueIsReadOnly=False,bDefaultValueIsIgnored=False,bAdvancedView=False,)
+            End Object
+            Begin Object Class=MaterialGraphNode Name="MaterialGraphNode_7"
+               Begin Object Class=MaterialExpressionTextureSampleParameter2D Name="MaterialExpressionTextureSampleParameter2D_0"
+               End Object
+               Begin Object Name="MaterialExpressionTextureSampleParameter2D_0"
+                  ParameterName="Left_Eye"
+                  ExpressionGUID=6A1A633E43E045D020DFC3A368709F4D
+                  Coordinates=(Expression=MaterialExpressionMaterialFunctionCall'MaterialGraphNode_6.MaterialExpressionMaterialFunctionCall_1')
+                  Texture=None
+                  MaterialExpressionEditorX=-266
+                  MaterialExpressionEditorY=174
+                  MaterialExpressionGuid=8B896BCB490997B105768DA336F9575C
+                  Material=PreviewMaterial'/Engine/Transient.MAT_Stereo_Image'
+               End Object
+               MaterialExpression=MaterialExpressionTextureSampleParameter2D'MaterialExpressionTextureSampleParameter2D_0'
+               NodePosX=-266
+               NodePosY=174
+               bHasCompilerMessage=True
+               bCanRenameNode=True
+               ErrorType=1
+               ErrorMsg="Param2D> Requires Texture2D"
+               NodeGuid=9A5B88C84A2C1D5D71F007999C9E64DB
+               CustomProperties Pin (PinId=A60D4B0340117A9FC89B09BA6D69B3AF,PinName="UVs",PinType.PinCategory="optional",PinType.PinSubCategoryObject=None,PinType.PinSubCategoryMemberReference=(),PinType.bIsArray=False,PinType.bIsReference=False,PinType.bIsConst=False,PinType.bIsWeakPointer=False,LinkedTo=(MaterialGraphNode_6 C87298924529FAC8728933A296BA7682,),PersistentGuid=00000000000000000000000000000000,bHidden=False,bNotConnectable=False,bDefaultValueIsReadOnly=False,bDefaultValueIsIgnored=False,bAdvancedView=False,)
+               CustomProperties Pin (PinId=E039327F417689ACF875B5A184D2C251,PinName="Tex",PinType.PinCategory="optional",PinType.PinSubCategoryObject=None,PinType.PinSubCategoryMemberReference=(),PinType.bIsArray=False,PinType.bIsReference=False,PinType.bIsConst=False,PinType.bIsWeakPointer=False,PersistentGuid=00000000000000000000000000000000,bHidden=False,bNotConnectable=False,bDefaultValueIsReadOnly=False,bDefaultValueIsIgnored=False,bAdvancedView=False,)
+               CustomProperties Pin (PinId=6073EDCF453D52887F1ED597234FAF0C,PinName="Output",PinFriendlyName=" ",Direction="EGPD_Output",PinType.PinCategory="mask",PinType.PinSubCategoryObject=None,PinType.PinSubCategoryMemberReference=(),PinType.bIsArray=False,PinType.bIsReference=False,PinType.bIsConst=False,PinType.bIsWeakPointer=False,LinkedTo=(MaterialGraphNode_0 44EC5EE54EACAD4AA705D8866FE16EB5,),PersistentGuid=00000000000000000000000000000000,bHidden=False,bNotConnectable=False,bDefaultValueIsReadOnly=False,bDefaultValueIsIgnored=False,bAdvancedView=False,)
+               CustomProperties Pin (PinId=1BFBF7694F3E966628AFC58E7C534726,PinName="Output2",PinFriendlyName=" ",Direction="EGPD_Output",PinType.PinCategory="mask",PinType.PinSubCategory="red",PinType.PinSubCategoryObject=None,PinType.PinSubCategoryMemberReference=(),PinType.bIsArray=False,PinType.bIsReference=False,PinType.bIsConst=False,PinType.bIsWeakPointer=False,PersistentGuid=00000000000000000000000000000000,bHidden=False,bNotConnectable=False,bDefaultValueIsReadOnly=False,bDefaultValueIsIgnored=False,bAdvancedView=False,)
+               CustomProperties Pin (PinId=3BFCC8DA419D81F993A7B391CF2FA32E,PinName="Output3",PinFriendlyName=" ",Direction="EGPD_Output",PinType.PinCategory="mask",PinType.PinSubCategory="green",PinType.PinSubCategoryObject=None,PinType.PinSubCategoryMemberReference=(),PinType.bIsArray=False,PinType.bIsReference=False,PinType.bIsConst=False,PinType.bIsWeakPointer=False,PersistentGuid=00000000000000000000000000000000,bHidden=False,bNotConnectable=False,bDefaultValueIsReadOnly=False,bDefaultValueIsIgnored=False,bAdvancedView=False,)
+               CustomProperties Pin (PinId=C25C6F6F42144A2FC2C5DB911ECD7C50,PinName="Output4",PinFriendlyName=" ",Direction="EGPD_Output",PinType.PinCategory="mask",PinType.PinSubCategory="blue",PinType.PinSubCategoryObject=None,PinType.PinSubCategoryMemberReference=(),PinType.bIsArray=False,PinType.bIsReference=False,PinType.bIsConst=False,PinType.bIsWeakPointer=False,PersistentGuid=00000000000000000000000000000000,bHidden=False,bNotConnectable=False,bDefaultValueIsReadOnly=False,bDefaultValueIsIgnored=False,bAdvancedView=False,)
+               CustomProperties Pin (PinId=4FACBD794BD17982F6AFE2A83001ED2E,PinName="Output5",PinFriendlyName=" ",Direction="EGPD_Output",PinType.PinCategory="mask",PinType.PinSubCategory="alpha",PinType.PinSubCategoryObject=None,PinType.PinSubCategoryMemberReference=(),PinType.bIsArray=False,PinType.bIsReference=False,PinType.bIsConst=False,PinType.bIsWeakPointer=False,PersistentGuid=00000000000000000000000000000000,bHidden=False,bNotConnectable=False,bDefaultValueIsReadOnly=False,bDefaultValueIsIgnored=False,bAdvancedView=False,)
+            End Object
+            Begin Object Class=MaterialGraphNode Name="MaterialGraphNode_8"
+               Begin Object Class=MaterialExpressionTextureSampleParameter2D Name="MaterialExpressionTextureSampleParameter2D_1"
+               End Object
+               Begin Object Name="MaterialExpressionTextureSampleParameter2D_1"
+                  ParameterName="Right_Eye"
+                  ExpressionGUID=20152A504733C1C1C6CCBFBEDE38D8C7
+                  Coordinates=(Expression=MaterialExpressionMaterialFunctionCall'MaterialGraphNode_6.MaterialExpressionMaterialFunctionCall_1')
+                  ConstMipValue=1
+                  Texture=None
+                  MaterialExpressionEditorX=-241
+                  MaterialExpressionEditorY=445
+                  MaterialExpressionGuid=0D48E23F4EF41B3303AB1AA8DC273216
+                  Material=PreviewMaterial'/Engine/Transient.MAT_Stereo_Image'
+               End Object
+               MaterialExpression=MaterialExpressionTextureSampleParameter2D'MaterialExpressionTextureSampleParameter2D_1'
+               NodePosX=-241
+               NodePosY=445
+               bHasCompilerMessage=True
+               bCanRenameNode=True
+               ErrorType=1
+               ErrorMsg="Param2D> Requires Texture2D"
+               NodeGuid=A93210F24FF9C8FE30CEECBBAF308703
+               CustomProperties Pin (PinId=0EA30C1443845B538B9305BFF8C98510,PinName="UVs",PinType.PinCategory="optional",PinType.PinSubCategoryObject=None,PinType.PinSubCategoryMemberReference=(),PinType.bIsArray=False,PinType.bIsReference=False,PinType.bIsConst=False,PinType.bIsWeakPointer=False,LinkedTo=(MaterialGraphNode_6 C87298924529FAC8728933A296BA7682,),PersistentGuid=00000000000000000000000000000000,bHidden=False,bNotConnectable=False,bDefaultValueIsReadOnly=False,bDefaultValueIsIgnored=False,bAdvancedView=False,)
+               CustomProperties Pin (PinId=6E12F03A4EBBD551878771B92104B952,PinName="Tex",PinType.PinCategory="optional",PinType.PinSubCategoryObject=None,PinType.PinSubCategoryMemberReference=(),PinType.bIsArray=False,PinType.bIsReference=False,PinType.bIsConst=False,PinType.bIsWeakPointer=False,PersistentGuid=00000000000000000000000000000000,bHidden=False,bNotConnectable=False,bDefaultValueIsReadOnly=False,bDefaultValueIsIgnored=False,bAdvancedView=False,)
+               CustomProperties Pin (PinId=CDF923064C8491C8DD875A9FD733B329,PinName="Output",PinFriendlyName=" ",Direction="EGPD_Output",PinType.PinCategory="mask",PinType.PinSubCategoryObject=None,PinType.PinSubCategoryMemberReference=(),PinType.bIsArray=False,PinType.bIsReference=False,PinType.bIsConst=False,PinType.bIsWeakPointer=False,LinkedTo=(MaterialGraphNode_0 F937930D47BA8EA98B9510B0C025B9A4,),PersistentGuid=00000000000000000000000000000000,bHidden=False,bNotConnectable=False,bDefaultValueIsReadOnly=False,bDefaultValueIsIgnored=False,bAdvancedView=False,)
+               CustomProperties Pin (PinId=F4A512534340958957BAC39A44E565FF,PinName="Output2",PinFriendlyName=" ",Direction="EGPD_Output",PinType.PinCategory="mask",PinType.PinSubCategory="red",PinType.PinSubCategoryObject=None,PinType.PinSubCategoryMemberReference=(),PinType.bIsArray=False,PinType.bIsReference=False,PinType.bIsConst=False,PinType.bIsWeakPointer=False,PersistentGuid=00000000000000000000000000000000,bHidden=False,bNotConnectable=False,bDefaultValueIsReadOnly=False,bDefaultValueIsIgnored=False,bAdvancedView=False,)
+               CustomProperties Pin (PinId=0080AD0C498AB9A1D974ED80486A3FEE,PinName="Output3",PinFriendlyName=" ",Direction="EGPD_Output",PinType.PinCategory="mask",PinType.PinSubCategory="green",PinType.PinSubCategoryObject=None,PinType.PinSubCategoryMemberReference=(),PinType.bIsArray=False,PinType.bIsReference=False,PinType.bIsConst=False,PinType.bIsWeakPointer=False,PersistentGuid=00000000000000000000000000000000,bHidden=False,bNotConnectable=False,bDefaultValueIsReadOnly=False,bDefaultValueIsIgnored=False,bAdvancedView=False,)
+               CustomProperties Pin (PinId=1B4746BE460A06F7E4C99091F01E070D,PinName="Output4",PinFriendlyName=" ",Direction="EGPD_Output",PinType.PinCategory="mask",PinType.PinSubCategory="blue",PinType.PinSubCategoryObject=None,PinType.PinSubCategoryMemberReference=(),PinType.bIsArray=False,PinType.bIsReference=False,PinType.bIsConst=False,PinType.bIsWeakPointer=False,PersistentGuid=00000000000000000000000000000000,bHidden=False,bNotConnectable=False,bDefaultValueIsReadOnly=False,bDefaultValueIsIgnored=False,bAdvancedView=False,)
+               CustomProperties Pin (PinId=2059FDD14284DE0A16798E92259F2C7A,PinName="Output5",PinFriendlyName=" ",Direction="EGPD_Output",PinType.PinCategory="mask",PinType.PinSubCategory="alpha",PinType.PinSubCategoryObject=None,PinType.PinSubCategoryMemberReference=(),PinType.bIsArray=False,PinType.bIsReference=False,PinType.bIsConst=False,PinType.bIsWeakPointer=False,PersistentGuid=00000000000000000000000000000000,bHidden=False,bNotConnectable=False,bDefaultValueIsReadOnly=False,bDefaultValueIsIgnored=False,bAdvancedView=False,)
+            End Object
+            Begin Object Class=MaterialGraphNode Name="MaterialGraphNode_9"
+               Begin Object Class=MaterialExpressionVectorParameter Name="MaterialExpressionVectorParameter_0"
+               End Object
+               Begin Object Name="MaterialExpressionVectorParameter_0"
+                  DefaultValue=(R=1.000000,G=-1.000000,B=-1.000000,A=1.000000)
+                  ParameterName="Image_Direction"
+                  ExpressionGUID=F1F2C7C348FA6F958EA8E099B55A37ED
+                  MaterialExpressionEditorX=-912
+                  MaterialExpressionEditorY=384
+                  MaterialExpressionGuid=6873E4014A670FD8059959BFE93D5CDA
+                  Material=PreviewMaterial'/Engine/Transient.MAT_Stereo_Image'
+               End Object
+               MaterialExpression=MaterialExpressionVectorParameter'MaterialExpressionVectorParameter_0'
+               NodePosX=-912
+               NodePosY=384
+               bCanRenameNode=True
+               NodeGuid=37303B2F48466367CBE440A181DB2539
+               CustomProperties Pin (PinId=9C23E1264DF642DC356B788B08CB3693,PinName="Output",PinFriendlyName=" ",Direction="EGPD_Output",PinType.PinCategory="mask",PinType.PinSubCategoryObject=None,PinType.PinSubCategoryMemberReference=(),PinType.bIsArray=False,PinType.bIsReference=False,PinType.bIsConst=False,PinType.bIsWeakPointer=False,LinkedTo=(MaterialGraphNode_4 F1DE55B14567D0B67759DEA9C934BA05,),PersistentGuid=00000000000000000000000000000000,bHidden=False,bNotConnectable=False,bDefaultValueIsReadOnly=False,bDefaultValueIsIgnored=False,bAdvancedView=False,)
+               CustomProperties Pin (PinId=6B8E6BCA4EF15BDF5D11B6B1CB4B1882,PinName="Output2",PinFriendlyName=" ",Direction="EGPD_Output",PinType.PinCategory="mask",PinType.PinSubCategory="red",PinType.PinSubCategoryObject=None,PinType.PinSubCategoryMemberReference=(),PinType.bIsArray=False,PinType.bIsReference=False,PinType.bIsConst=False,PinType.bIsWeakPointer=False,PersistentGuid=00000000000000000000000000000000,bHidden=False,bNotConnectable=False,bDefaultValueIsReadOnly=False,bDefaultValueIsIgnored=False,bAdvancedView=False,)
+               CustomProperties Pin (PinId=4F58ABD7410B38EA9235A497F25BD3C1,PinName="Output3",PinFriendlyName=" ",Direction="EGPD_Output",PinType.PinCategory="mask",PinType.PinSubCategory="green",PinType.PinSubCategoryObject=None,PinType.PinSubCategoryMemberReference=(),PinType.bIsArray=False,PinType.bIsReference=False,PinType.bIsConst=False,PinType.bIsWeakPointer=False,PersistentGuid=00000000000000000000000000000000,bHidden=False,bNotConnectable=False,bDefaultValueIsReadOnly=False,bDefaultValueIsIgnored=False,bAdvancedView=False,)
+               CustomProperties Pin (PinId=252E95C9465BAEED4430B5B9DA37851D,PinName="Output4",PinFriendlyName=" ",Direction="EGPD_Output",PinType.PinCategory="mask",PinType.PinSubCategory="blue",PinType.PinSubCategoryObject=None,PinType.PinSubCategoryMemberReference=(),PinType.bIsArray=False,PinType.bIsReference=False,PinType.bIsConst=False,PinType.bIsWeakPointer=False,PersistentGuid=00000000000000000000000000000000,bHidden=False,bNotConnectable=False,bDefaultValueIsReadOnly=False,bDefaultValueIsIgnored=False,bAdvancedView=False,)
+               CustomProperties Pin (PinId=B5B286C94A078205748BA1880BD0B1D1,PinName="Output5",PinFriendlyName=" ",Direction="EGPD_Output",PinType.PinCategory="mask",PinType.PinSubCategory="alpha",PinType.PinSubCategoryObject=None,PinType.PinSubCategoryMemberReference=(),PinType.bIsArray=False,PinType.bIsReference=False,PinType.bIsConst=False,PinType.bIsWeakPointer=False,PersistentGuid=00000000000000000000000000000000,bHidden=False,bNotConnectable=False,bDefaultValueIsReadOnly=False,bDefaultValueIsIgnored=False,bAdvancedView=False,)
+            End Object
+    [/PARAM]
+[/OBJECT]
+[/EXCERPT:SPMaterial]
+-->
\ No newline at end of file
diff --git a/Engine/Documentation/Source/Platforms/VR/StereoPanoramicCapture/HowTo/LevelSetup/LevelSetup.INT.udn b/Engine/Documentation/Source/Platforms/VR/StereoPanoramicCapture/HowTo/LevelSetup/LevelSetup.INT.udn
new file mode 100644
--- /dev/null
+++ b/Engine/Documentation/Source/Platforms/VR/StereoPanoramicCapture/HowTo/LevelSetup/LevelSetup.INT.udn
@@ -0,0 +1,109 @@
+Availability: Docs
+Title:4. Level Setup
+Crumbs: %ROOT%, Platforms, Platforms/VR, Platforms/VR/StereoPanoramicCapture/HowTo
+Description: In this section, we will go over how to set up a level to display the stereo screenshots.
+Version: 4.13
+SkillLevel: Intermediate
+type: multi-step
+tags: VR
+
+[VAR:Steps]
+[OBJECT:Navigation]
+    [PARAM:previous]
+        [Previous Step](Platforms\VR\StereoPanoramicCapture\HowTo\CreatingViewerMaterial "%Platforms\VR\StereoPanoramicCapture\HowTo\CreatingViewerMaterial:title%")
+    [/PARAM]
+    [PARAM:current]
+    [/PARAM]
+    [PARAM:home]
+        [Stereo Panormaic ](Platforms\VR\StereoPanoramicCapture\HowTo)
+    [/PARAM]
+    [PARAM:next]
+    [/PARAM]
+[/OBJECT]
+[/VAR]
+
+%Steps%
+
+In this step, we will go over how to set up a level that will display the stereo screenshots when the level is run.
+
+1. From the **Main Toolbar** go to the **File** option and then from the list that is displayed, select the **New Level** option.
+
+    ![](SPHT_New_Level.png)
+
+1. When the **New Level** dialogue box shows up, choose the **Empty Level** option to load a new empty level.
+
+    ![](SPHT_Blank_Level.png)
+
+1. Click on the link below to download the required Static Mesh that is needed to display the Material. 
+
+    [](SM_PhotoSphere.zip)
+
+	[region:note]
+	The Static Mesh you are downloading has had its UV’s setup to work correctly with Material that was created. While you can use any Static Mesh you want, downloading this file will ensure that you will get the correct results.
+	[/region]
+
+1. Once the file has downloaded, unzip the file and then import the Static mesh into the Content Browser.
+
+    [REGION:lightbox]
+    [![](SPHT_Import_Sky.png)(w:617)](SPHT_Import_Sky.png)
+    [/REGION]
+
+    [REGION:caption]
+    Click for full image.
+    [/REGION]
+
+    [region:note]
+    If you receive a warning about smoothing groups not being found in the file, you can safely ignore that warning as we do not need smoothing information for this object as it uses an unlit shader. 
+    [/region]
+
+1. After the **SM_PhotoSphere**  Static Mesh has been imported, drag the **SM_PhotoSphere** from the Content Browser into the level, positioning it so that it's location is a **0,0,0** in the world.
+
+    [REGION:lightbox]
+    [![](SPHT_Position_SM.png)(w:617)](SPHT_Position_SM.png)
+    [/REGION]
+
+    [REGION:caption]
+    Click for full image.
+    [/REGION]
+
+1. Once the SM_PhotoSphere has been positioned in the center of the level, select the MAT_StereoViewer Material in the Content Browser and then drag it to the **Materials** input in the Details panel.
+
+    [REGION:lightbox]
+    [![](SPHY_Apply_Material.png)(w:617)](SPHY_Apply_Material.png)
+    [/REGION]
+
+    [REGION:caption]
+    Click for full image.
+    [/REGION]
+
+1. Once the Material has been applied, press the **Play in VR** preview button and then put your HMD on to see what the images look like when viewed in stereo. If everything has been set up correctly, you should not be able to tell that the level is being rendered using just two images instead of 3D objects.
+
+    [OBJECT:EmbeddedVideo]
+        [PARAMLITERAL:width]
+        640
+        [/PARAMLITERAL]
+        [PARAMLITERAL:height]
+        360
+        [/PARAMLITERAL]
+        [PARAMLITERAL:videoid]
+        cg7xV85GRzo
+        [/PARAMLITERAL]
+    [/OBJECT]
+
+	[region:note]
+	The black sphere you are seeing is the shadow that is cast from the Static Mesh that is attached to the Default Pawn.  To get rid of this shadow you can either create a new Game Mode and Pawn or when you run the game you can press the player eject button, find the Default Pawn in the World Outliner and under the Static Mesh section, set the Static Mesh to none and the shadow will be gone.
+	[/region]
+
+
+
+With this simple set up, you can now view any images that you take using the Stereo Panoramic plugin by replacing the left and right eye images with the new images that you created.
+
+
+%Steps%
+
+
+
+
+
+
+
diff --git a/Engine/Documentation/Source/Platforms/VR/StereoPanoramicCapture/HowTo/ProjectSetup/StereoPanoramicProjectSetup.INT.udn b/Engine/Documentation/Source/Platforms/VR/StereoPanoramicCapture/HowTo/ProjectSetup/StereoPanoramicProjectSetup.INT.udn
new file mode 100644
--- /dev/null
+++ b/Engine/Documentation/Source/Platforms/VR/StereoPanoramicCapture/HowTo/ProjectSetup/StereoPanoramicProjectSetup.INT.udn
@@ -0,0 +1,56 @@
+Availability: Docs
+Title:1. Project Setup
+Crumbs:
+Description:Here we will create and setup a new project to capture stereo screenshots in. 
+Version:4.13
+SkillLevel: Intermediate
+type:multi-step
+tags: VR
+order:2
+
+[NAV]
+
+In this section, we will create and then set up a new UE4 project to work out of.
+
+[region:note]
+While we will be creating a new blank project for this how-to, the information presented here can be applied to any UE4 project.
+[/region]
+
+## Steps
+
+1. From the **New Project** section of the Unreal Project Browser create a new Blank Blueprint base project with the following options and when completed, press the **Create Project** button to set up and load the new project.
+
+    * Blank Project Template
+    * Mobile / Tablet
+    * Scalable 3D or 2D
+    * No Starter Content
+
+    [REGION:lightbox]
+    [![](SPHT_Create_New_Project.png)(w:555)](SPHT_Create_New_Project.png)
+    [/REGION]
+
+    [REGION:caption]
+    Click for full image.
+    [/REGION]
+
+1. Once the project has been loaded go to the top menu bar and then go to **Edit** > **Plugins** to open up the plugins options.
+
+    ![](SPHT_Plugins_Menu.png)
+
+1. In the Plugins menu, click on the **Movie Capture** option and then in the **Stereo Panoramic Movie Capture** section, click on the checkmark box to enable the plugin and then press the **Restart Now** button to restart the UE4 editor.
+
+    [REGION:lightbox]
+    [![](SPHT_Enable_SPMP.png)(w:543)](SPHT_Enable_SPMP.png)
+    [/REGION]
+
+    [REGION:caption]
+    Click for full image.
+    [/REGION]
+
+1. Once the editor re-starts open up the console by pressing the backtick (`) on the keyboard and then input **SP.**. If the Stereo Panoramic Movie Plugin has been successfully enabled, you should see a list of commands you can choose from. If you do not see the command list make sure that you enabled and restarted the editor and that you are using Unreal Engine 4.13 or later.
+
+    ![](SPHT_Check_Plugin.png)
+
+With the Stereo Panoramic Move Plugin now enabled and verified to be working, we will next take a look at how to use the plugin to capture stereo images.
+
+[NAV]
\ No newline at end of file
