Copying //UE4/Dev-Documentation to Samples-Main (//UE4/Samples-Main)

diff --git a/Engine/Documentation/Source/Platforms/VR/StereoPanoramicCapture/HowTo/StereoPanoramicHowTo.INT.udn b/Engine/Documentation/Source/Platforms/VR/StereoPanoramicCapture/HowTo/StereoPanoramicHowTo.INT.udn
new file mode 100644
--- /dev/null
+++ b/Engine/Documentation/Source/Platforms/VR/StereoPanoramicCapture/HowTo/StereoPanoramicHowTo.INT.udn
@@ -0,0 +1,29 @@
+Availability:Docs
+Title: Using the Stereo Panoramic Capture Tool
+Crumbs:%ROOT%, Platforms, Platforms/VR
+Description:A look at how to use the Stereo Panoramic plugin to capture VR images and movies.
+Navigation:platform
+Version: 4.13
+parent:Platforms/VR/StereoPanoramicCapture
+type:how-to
+tags: VR
+
+![](Platforms\VR\StereoPanoramicCapture\SPE_Header_00.png)
+
+In the following example, we will take a look at how you can setup, create and view Stereo Panoramic screenshots all from inside UE4. When you have completed this project you will have created two Stereo Panoramic images and a Material that can be used to display those images in 3D. 
+
+
+## Steps 
+
+
+[DIR(output:"checkpoint" parent:"Platforms/VR/StereoPanoramicCapture/HowTo")]
+
+[REGION:call_to_action]
+    [Click to Start](Platforms\VR\StereoPanoramicCapture\HowTo\ProjectSetup "%Platforms\VR\StereoPanoramicCapture\HowTo\ProjectSetup:title%")
+[/REGION]
+
+
+
+
+
+
diff --git a/Engine/Documentation/Source/Platforms/VR/StereoPanoramicCapture/HowTo/TakingStereoImages/TakingStereoImages.INT.udn b/Engine/Documentation/Source/Platforms/VR/StereoPanoramicCapture/HowTo/TakingStereoImages/TakingStereoImages.INT.udn
new file mode 100644
--- /dev/null
+++ b/Engine/Documentation/Source/Platforms/VR/StereoPanoramicCapture/HowTo/TakingStereoImages/TakingStereoImages.INT.udn
@@ -0,0 +1,55 @@
+Availability: Docs
+Title:2. Taking Stereo Images
+Crumbs:
+Description: In this section, we will go over how to capture stereo screenshots using the Stereo Panoramic plugin.
+Version:4.13
+SkillLevel: Intermediate
+type:multi-step
+tags: VR
+order:2
+
+
+[NAV]
+
+In this step, we will use the Stereo Panoramic plugin to capture the images of the level in stereo 3D.
+
+## Steps
+
+1. From the **Main Toolbar** press the **Play** button and then navigate your player to an area that you want to capture. For this example, we will be leaving the player in the middle of the Static Mesh but feel free to position the player anywhere you want to.
+
+    ![](SPHT_Play_Game.png)
+
+1. Press the **Backtick** key (`) on the keyboard to open up the UE4 console and then input the following command. Once the command has been input press the Enter key to execute the command.
+
+    * **SP.PanoramicScreenshot**
+
+    ![](SPHT_Input_Command.png)
+
+    [region:note]
+    Instead of typing out SP.PanoramicScreenshot in its entirety; you can just type **SP.P** and then select SP.PanoramicScreenshot from the displayed list using the arrow keys. Once SP.PanoramicScreenshot is selected, press the Enter key to execute the command.
+    [/region]
+
+1. During the capture process, the editor might appear to be frozen or unresponsive for a few seconds or minutes.  This is due to the demanding rendering requirements of the Stereo Panoramic plugin. When the editor becomes responsive again, you will be able to find the screenshots in the following location in your projects folder.
+
+    * **StereoPanoramaImages\Saved\StereoPanorama\(Date & Time it was taken)**
+
+    [REGION:lightbox]
+    [![](SPHT_Rendered_Images.png)(w:473)](SPHT_Rendered_Images.png)
+    [/REGION]
+
+    [REGION:caption]
+    Click for full image.
+    [/REGION]
+
+    [region:note]
+    If the image previews look washed in thumbnail images or when viewed in an image viewer it is because the Stereo Panoramic plugin will add an alpha channel to the output.PNG file. 
+    [/region]
+
+	[region:tip]
+	To create the images that you see in the folder the Stereo Panoramic plugin is rendering the scene from the viewpoint of each eye that is rotated around a fixed point to simulate the head turning around the neck.
+	[/region]
+
+
+With the images now created, in the next section we will take a look at how to build a Material that will allow the images to be viewed using a VR HMD.
+
+[NAV]
diff --git a/Engine/Documentation/Source/Platforms/VR/StereoPanoramicCapture/Reference/StereoPanoramicCaptureReference.INT.udn b/Engine/Documentation/Source/Platforms/VR/StereoPanoramicCapture/Reference/StereoPanoramicCaptureReference.INT.udn
new file mode 100644
--- /dev/null
+++ b/Engine/Documentation/Source/Platforms/VR/StereoPanoramicCapture/Reference/StereoPanoramicCaptureReference.INT.udn
@@ -0,0 +1,103 @@
+Availability:Docs
+Title:Stereo Panoramic Capture Reference
+Crumbs:%ROOT%, Platforms, Platforms/VR
+Description:An in-depth look at what each property for the Stereo Panoramic does
+Navigation:platform
+Version: 4.13
+parent:Platforms/VR/StereoPanoramicCapture
+type:reference
+related: Gameplay/HowTo/UsingCameras/
+related: Engine/Performance/Scalability/ScalabilityReference/
+
+On the following page, you will find a detailed description of what options and controls are available for the Stereo Panoramic Capture plugin and what those options do to the image that is output. The options have been broken up into groups based on what those options do. To input a command press the backtick or tilde key (`) to open up the UE4 console and then type **SP.** followed by the command you want to enter.
+
+[region:note]
+If you want to know more about any of the settings below, check out the **StereoPanoramaManager.cpp** file for more information.
+[/region]
+
+## Slice Controls
+The Slice Control options are for controlling how many horizontal or vertical slices are taken for each image.
+
+|Property|Default Value|Description|
+|------------|-----------------|-----------------------|
+|SP.HorizontalAngularIncrement| 1.0f| The number of degrees per horizontal step. Must be a factor of 360.|
+|SP.VerticalAngularIncrement| 90.0f| The number of degrees per vertical step. Must be a factor of 180.|
+|SP.CaptureHorizontalFOV| 90.0f| Horizontal FOV for scene capture component. Must be larger than SP.HorizontalAngularIncrement.|
+|SP.EyeSeparation| 6.4f| The separation of the stereo cameras.|
+
+## Atlas Controls
+The atlas controls actually control the atlas of shots collected from the slices, which are used to reconstruct the 360 degree image.
+
+|Property|Default Value|Description|
+|------------|-----------------|-----------------------|
+|SP.StepCaptureWidth| 4096| The final spherical atlas width.|
+|SP.ForceAlpha| false| Force the alpha value to completely opaque.|
+
+## Sampling Controls
+The Sampling Control options affect how the image is filtered.
+
+|Property|Default Value|Description|
+|------------|-----------------|-----------------------|
+|SP.CaptureSlicePixelWidth| 2048| Capture Slice Pixel Dimension.|
+|SP.EnableBilerp| true| 0 - No Filtering 1- Bilinear Filter slice samples.|
+|SP.SuperSamplingMethod| 1| 0 - No Supersampling, 1 - Rotated Grid SS.|
+
+## Debug Controls
+The Debug Control options allow you to see and adjust how the images are captured allowing you to easily track down any issues you might run into with the final output image.
+
+|Property|Default Value|Description|
+|------------|-----------------|-----------------------|
+|SP.ConcurrentCaptures| 30| The number of scene captures to capture at the same time.|
+|SP.GenerateDebugImages| 0| 0 - No Debug Images. 1 - Save out each strip as it is generated. 2 - Save each entire slice.|
+|SP.OutputDir| This will default to your projects saved folder| The location where your images will be saved to.|
+|SP.ShouldOverrideInitialYaw| true| Override Initial Camera Yaw. Set to true if you don't want to use PlayerController View-Dir.|
+|SP.ForcedInitialYaw| 90.0f| Yaw value for initial Camera view direction. Set ShouldOverrideInitialYaw to true to use this value.|
+|SP.FadeStereoToZeroAtSides| true| Fade stereo effect between left/right eye to zero at 90 degrees.|
+
+[region:note]
+Please note that raising or lowering the value of **SP.ConcurrentCaptures** can greatly affect capture times. Setting this too low and you will not be using the optimal number for parallel processing. Setting this too high and you will overwhelm the GPU.
+[/region]
+
+## Pre-set Quality Levels
+The Pre-set Quality levels allow you to select from three different pre-defined image quality levels.  To use a Pre-Set Quality level press the backtick or tilde key (`) to open up the UE4 console and then type **SP.(Quality Level you want to use)**.
+
+### Preview Quality Level
+|Property|Default Value|
+|------------|-----------------|
+|SP.HorizontalAngularIncrement| 5| 
+|SP.VerticalAngularIncrement| 60| 
+|SP.CaptureHorizontalFOV| 60| 
+|SP.StepCaptureWidth| 720| 
+
+### Average Quality Level
+|Property|Default Value|
+|------------|-----------------|
+|SP.HorizontalAngularIncrement| 2| 
+|SP.VerticalAngularIncrement| 30| 
+|SP.CaptureHorizontalFOV| 30| 
+|SP.StepCaptureWidth| 1440| 
+
+### Improved Quality Level
+|Property|Default Value|
+|------------|-----------------|
+|SP.HorizontalAngularIncrement| 0.5| 
+|SP.VerticalAngularIncrement| 22.5| 
+|SP.CaptureHorizontalFOV| 22.5| 
+|SP.StepCaptureWidth| 1440| 
+
+
+
+
+
+
+
+
+
+
+
+
+
+
+
+
+
diff --git a/Engine/Documentation/Source/Platforms/VR/StereoPanoramicCapture/StereoPanoramicCapture.INT.udn b/Engine/Documentation/Source/Platforms/VR/StereoPanoramicCapture/StereoPanoramicCapture.INT.udn
new file mode 100644
--- /dev/null
+++ b/Engine/Documentation/Source/Platforms/VR/StereoPanoramicCapture/StereoPanoramicCapture.INT.udn
@@ -0,0 +1,52 @@
+Availability:Docs
+Title:Stereo Panoramic Capture Tool
+Crumbs:%ROOT%, Platforms, Platforms/VR
+Description:A look at how to use the Stereo Panoramic plugin to capture VR images and movies.
+Navigation:platform
+Version: 4.13
+parent:Platforms/VR
+type:landing
+
+![](SPE_Header_00.png)
+
+The Stereo Panoramic Movie Capture plugin for UE4 allows you to capture still images or movies that can then be viewed or played back in stereo giving the illusion that the world is being rendered in realtime.  In the following sections, we will take a look at how to set up and use the Stereo Panoramic Movie Capture plugin and also take a look at how you can create an image viewer inside of UE4 that will allow you to see the images you create in stereo. 
+
+[region:warning]
+Please note that the Stereo Panoramic Movie Capture plugin is still in an experimental state so you may run into issues like unexpected crashes when using it. 
+[/region]
+
+## Starting Out
+
+The Stereo Panoramic Capture How To guide will walk you through getting your UE4 project setup to be able to capture and view stereo images.
+
+[DIR(end:"1" output:"listbutton" type:"how-to" parent:"Platforms/VR/StereoPanoramicCapture")]
+
+## Stereo Panoramic Capture Reference 
+
+The Stereo Panoramic Capture Reference will go over the available options, and settings the Stereo Panoramic Capture plugin has and the effect those options have on the rendered images.
+
+[DIR(end:"1" output:"listbutton" type:"reference" parent:"Platforms/VR/StereoPanoramicCapture")]
+
+## Stereo Panoramic Capture Tips & Tricks
+
+The Stereo Panoramic Capture Tips & Tricks section will go over several tips and tricks so that you can ensure the best possible results when using the Stereo Panoramic Capture plugin.
+
+[DIR(end:"1" output:"listbutton" type:"overview" parent:"Platforms/VR/StereoPanoramicCapture")]
+
+
+
+
+
+
+
+
+
+
+
+
+
+
+
+
+
+
diff --git a/Engine/Documentation/Source/Platforms/VR/StereoPanoramicCapture/TipsAndTricks/TipAndTricks.INT.udn b/Engine/Documentation/Source/Platforms/VR/StereoPanoramicCapture/TipsAndTricks/TipAndTricks.INT.udn
new file mode 100644
--- /dev/null
+++ b/Engine/Documentation/Source/Platforms/VR/StereoPanoramicCapture/TipsAndTricks/TipAndTricks.INT.udn
@@ -0,0 +1,367 @@
+Availability:Docs
+Title:Stereo Panoramic Capture Tips & Tricks
+Crumbs:%ROOT%, Platforms, Platforms/VR
+Description:Taking a look at how to fix some common issues with the Stereo Panoramic capture tool
+Navigation:platform
+Version: 4.13
+parent:Platforms/VR/StereoPanoramicCapture
+type:overview
+
+[TOC(start:2)]
+
+In the following section, we will go over some of the most common errors and issues that you might run into when using the Stereo Panoramic plugin and how to go about solving or getting around them.
+
+## Capturing Movies in Stereo
+
+You can also use the Stereo Panoramic capture tool to output an image sequence that can then be turned into a Stereo Movie. To get the best possible results before you try and capture a stereo video, make sure you do the following:
+
+* **Command Line Options:** These commands should be added to your game / editor.exe, so they are enabled when your project is loaded.
+    * **-usefixedtimestep** 
+    * **-notexturestreaming**
+
+	![](TP_CommandLineOptions.png)
+
+* **Console Commands:** The following Console Commands will help to make sure nothing is missed during the capture. You can place the following Console commands in the level Blueprint that the capture is taking place in.
+    * **r.AllowOcclusionQueries 0**
+	* **t.MaxFPS 60**
+	* **SP.HorizontalAngularIncrement 2**
+	* **SP.VerticalAngularIncrement 30**
+	* **SP.CaptureHorizontalFOV 30**
+	* **SP.ConcurrentCaptures 16**
+	* **SP.CaptureSlicePixelWidth 720**
+	* **SP.StepCaptureWidth 4096**
+
+	[REGION:lightbox]
+	[![](TP_Blueprint_Setup.png)(w:912)](TP_Blueprint_Setup.png)
+	[/REGION]
+
+	[REGION:caption]
+	Click for full image.
+	[/REGION]
+
+* **Movie Capture Setup:** The best way to capture a stereo movie is to animate a camera using Matinee or Sequencer and then call that sequence along with the following command to output a movie image sequence.
+	
+	* **SP.PanoramicMovie (Frame Range Here)** 
+
+	![](TP_Movie_Capture.png)
+
+When completed your Level Blueprint should look similar to the following image.
+
+
+[REGION:fullwidth]
+[INCLUDE:Platforms\VR\StereoPanoramicCapture\TipsAndTricks\#StereoBP]
+[/REGION]
+
+[REGION:caption]
+You can copy the following Blueprint code by clicking on the **Copy Node Graph** option that is in the upper left hand corner of the image.
+[/REGION]
+
+
+## Image Seams
+
+When viewing, images using the Stereo Image Material you might notice a seam where the left and right eye Textures met like in the image below.
+
+[OBJECT:ComparisonSlider]
+    [PARAM:before]
+    ![Seam](TP_Seam_Example.png)
+    [/PARAM]
+    [PARAM:after]
+    ![No Seam](TP_No_Seam.png)
+    [/PARAM]
+[/OBJECT]
+
+To remove these seams, you will need to open up your Left and Right eye images in the UE4 Texture Viewer. Then in the **Level of Detail** section, change the value for **Mip Gen Settings** to **NoMipmaps**. 
+
+[REGION:lightbox]
+[![](Platforms\VR\StereoPanoramicCapture\HowTo\CreatingViewerMaterial\SPHT_Mip_Gen_Settings.png)(w:439)](Platforms\VR\StereoPanoramicCapture\HowTo\CreatingViewerMaterial\SPHT_Mip_Gen_Settings.png)
+[/REGION]
+
+[REGION:caption]
+Click for full image.
+[/REGION]
+
+## Image Darkness
+
+If you notice that the captured images are coming out very dark, it could mean you are using the old tonemapper. To change to the new tonemapper enter the following command into the console and then re-shoot your images.
+
+**r.tonemapperfilm=1**
+
+[OBJECT:ComparisonSlider]
+    [PARAM:before]
+    ![r.ToneMapperFilm Off](TP_No_PP.png)
+    [/PARAM]
+    [PARAM:after]
+    ![r.ToneMapperFilm On](TP_PP.png)
+    [/PARAM]
+[/OBJECT]
+
+
+With the new tonemapper now enabled. You can adjust the following Post Process settings to your liking.
+
+ * White Balance
+ * Color Grading
+ * Film
+ * Scene Color
+ * Global Illumination
+
+    [region:note]
+    Please note that while a lot of the above categories have settings in them that will change how the screen will look, you might find some settings that do not work.
+    [/region]
+
+## Image Distortion 
+
+Due to the manner in which the Stereo images are created, you will get some image distortion / pinching at the top and bottom of the image. You can see an example of what this distortion / pinching looks like in the following image.
+
+[REGION:lightbox]
+[![](TP_Image_Distortion.png)(w:256)](TP_Image_Distortion.png)
+[/REGION]
+
+[REGION:caption]
+Click for full image.
+[/REGION]
+
+While you can not eliminate this distortion, there are a few things you can do to help minimize it. Below you will find a few tips on how to address this distortion / pinching issue. 
+
+* You can add a small gradient that fades to black at the bottom and or top of the image to help hide this distortion / pinching issue like in the image below.
+
+    [REGION:lightbox]
+    [![](TP_Image_Distortion_Fix.png)(w:256)](TP_Image_Distortion_Fix.png)
+    [/REGION]
+
+    [REGION:caption]
+    Click for full image.
+    [/REGION]
+
+* Objects that are very close to the camera that is capturing can cause divergence which makes it hard for a user to focus on the 3D image. Be careful with objects that get too close to the camera as this will cause divergence. While in some cases image divergence will not be possible to avoid, thinking about your camera placement can help lessen the effect. The following video shows how you can have different levels divergence in your images. 
+
+    [OBJECT:EmbeddedVideo]
+        [PARAMLITERAL:width]
+        640
+        [/PARAMLITERAL]
+        [PARAMLITERAL:height]
+        360
+        [/PARAMLITERAL]
+        [PARAMLITERAL:videoid]
+        sjwiAQW90qg
+        [/PARAMLITERAL]
+    [/OBJECT]
+
+    When the user looks up, they will have issues getting the top part of the image to be in focus due to how close the top of the hang glider was to the capture camera. However, when the user looks down, they will have no focus issues as the ground was far enough away from the camera to not cause divergence.
+
+## Not all Effects Work
+
+Effects that are created in screen space or effects that are camera facing will not work with the Stereo Panoramic image capture tool. This means that the following types of effects should be disabled so that your images do not have any errors in them. 
+
+* Light Shafts
+* Screen based Distortion effects
+* Camera Facing Meshes or Particles
+* Vignette 
+
+[region:note]
+Please note that there could be other effects that are not listed here that also do not capture as well. If you see any similar artifacts in, your images try disabling the effect and re-capturing the image to see if the artifact goes away or not.
+[/region]
+
+In the following image, you can see what happens when you try and capture something that makes use of screen based effects or is camera facing. If you take a look at the highlighted areas of the image, you will notice what appears to be a hard line in the clouds that are off in the distance. This line is there because the capture tool works by capturing your scene in a series of zoomed in the image for both the horizontal and vertical axis. Because of this effects that face the camera do not know that the camera has moved and do not update their orientation to reflect this resulting in artifacts like you see in the image. 
+
+[REGION:lightbox]
+[![](TP_Screen_Based_Effects_Not_Working.png)(w:806)](TP_Screen_Based_Effects_Not_Working.png)
+[/REGION]
+
+[REGION:caption]
+Click for full image.
+[/REGION]
+
+## Additional Resources
+
+[Ninja Theory](http://www.ninjatheory.com/wp/) did a great write up on how they used the Stereo Panoramic tool to capture images and movies in stereo. If you have not checked out their blog post on the Epic website, you should as it is packed with useful information. You can view the blog post by clicking on the following iamge.
+
+[![](Hellblade_Link_Image.png)](https://www.unrealengine.com/blog/capturing-stereoscopic-360-screenshots-videos-movies-unreal-engine-4)
+
+<!--
+[EXCERPT:StereoBP]	
+[OBJECT:Snippet]
+	[PARAMLITERAL:id]
+		StereoBP
+	[/PARAMLITERAL]
+	[PARAMLITERAL:type]
+		blueprint
+	[/PARAMLITERAL]
+	[PARAM:image]
+		![](TP_FinalBPSetup.png)
+	[/PARAM]
+	[PARAM:snippet]
+			BEGIN OBJECT
+						Begin Object Class=K2Node_Event Name="K2Node_Event_125"
+		   EventReference=(MemberParent=Class'/Script/Engine.Actor',MemberName="ReceiveBeginPlay")
+		   bOverrideFunction=True
+		   NodePosX=-896
+		   NodePosY=752
+		   ErrorType=1
+		   ErrorMsg="Error Missing Event \'BeginGame\' for  Event BeginGame  "
+		   NodeGuid=A75DF43F4BF618856D0E5981842ADB4F
+		   CustomProperties Pin (PinId=CC02AA634F24989613B99EA5B50D165A,PinName="OutputDelegate",Direction="EGPD_Output",PinType.PinCategory="delegate",PinType.PinSubCategory="",PinType.PinSubCategoryObject=None,PinType.PinSubCategoryMemberReference=(MemberParent=Class'/Script/Engine.Actor',MemberName="ReceiveBeginPlay"),PinType.PinValueType=(),PinType.bIsMap=False,PinType.bIsSet=False,PinType.bIsArray=False,PinType.bIsReference=False,PinType.bIsConst=False,PinType.bIsWeakPointer=False,PersistentGuid=00000000000000000000000000000000,bHidden=False,bNotConnectable=False,bDefaultValueIsReadOnly=False,bDefaultValueIsIgnored=False,bAdvancedView=False,)
+		   CustomProperties Pin (PinId=B3DA9D9944CB25B1C6C25391C61DE208,PinName="then",Direction="EGPD_Output",PinType.PinCategory="exec",PinType.PinSubCategory="",PinType.PinSubCategoryObject=None,PinType.PinSubCategoryMemberReference=(),PinType.PinValueType=(),PinType.bIsMap=False,PinType.bIsSet=False,PinType.bIsArray=False,PinType.bIsReference=False,PinType.bIsConst=False,PinType.bIsWeakPointer=False,LinkedTo=(K2Node_CallFunction_44 FAE82F5E422A038876E1CE84A81BD4E5,),PersistentGuid=00000000000000000000000000000000,bHidden=False,bNotConnectable=False,bDefaultValueIsReadOnly=False,bDefaultValueIsIgnored=False,bAdvancedView=False,)
+		End Object
+		Begin Object Class=K2Node_CallFunction Name="K2Node_CallFunction_44"
+		   FunctionReference=(MemberParent=Class'/Script/Engine.KismetSystemLibrary',MemberName="ExecuteConsoleCommand")
+		   NodePosX=-688
+		   NodePosY=752
+		   NodeGuid=AFBCFC1344213DF7985B249E0B0D2AD1
+		   CustomProperties Pin (PinId=FAE82F5E422A038876E1CE84A81BD4E5,PinName="execute",PinToolTip="\nExec",PinType.PinCategory="exec",PinType.PinSubCategory="",PinType.PinSubCategoryObject=None,PinType.PinSubCategoryMemberReference=(),PinType.PinValueType=(),PinType.bIsMap=False,PinType.bIsSet=False,PinType.bIsArray=False,PinType.bIsReference=False,PinType.bIsConst=False,PinType.bIsWeakPointer=False,LinkedTo=(K2Node_Event_125 B3DA9D9944CB25B1C6C25391C61DE208,),PersistentGuid=00000000000000000000000000000000,bHidden=False,bNotConnectable=False,bDefaultValueIsReadOnly=False,bDefaultValueIsIgnored=False,bAdvancedView=False,)
+		   CustomProperties Pin (PinId=46C4EC604A2519BABCFA6B909568F518,PinName="then",PinToolTip="\nExec",Direction="EGPD_Output",PinType.PinCategory="exec",PinType.PinSubCategory="",PinType.PinSubCategoryObject=None,PinType.PinSubCategoryMemberReference=(),PinType.PinValueType=(),PinType.bIsMap=False,PinType.bIsSet=False,PinType.bIsArray=False,PinType.bIsReference=False,PinType.bIsConst=False,PinType.bIsWeakPointer=False,LinkedTo=(K2Node_CallFunction_9 FAE82F5E422A038876E1CE84A81BD4E5,),PersistentGuid=00000000000000000000000000000000,bHidden=False,bNotConnectable=False,bDefaultValueIsReadOnly=False,bDefaultValueIsIgnored=False,bAdvancedView=False,)
+		   CustomProperties Pin (PinId=AB16A1664BEB283188CC7EBFBD515297,PinName="self",PinFriendlyName=NSLOCTEXT("K2Node", "Target", "Target"),PinToolTip="Target\nKismet System Library Reference",PinType.PinCategory="object",PinType.PinSubCategory="",PinType.PinSubCategoryObject=Class'/Script/Engine.KismetSystemLibrary',PinType.PinSubCategoryMemberReference=(),PinType.PinValueType=(),PinType.bIsMap=False,PinType.bIsSet=False,PinType.bIsArray=False,PinType.bIsReference=False,PinType.bIsConst=False,PinType.bIsWeakPointer=False,DefaultObject="/Script/Engine.Default__KismetSystemLibrary",PersistentGuid=00000000000000000000000000000000,bHidden=True,bNotConnectable=False,bDefaultValueIsReadOnly=False,bDefaultValueIsIgnored=False,bAdvancedView=False,)
+		   CustomProperties Pin (PinId=07888691421445ABFA1A0CBEF8592B32,PinName="WorldContextObject",PinToolTip="World Context Object\nObject Reference",PinType.PinCategory="object",PinType.PinSubCategory="",PinType.PinSubCategoryObject=Class'/Script/CoreUObject.Object',PinType.PinSubCategoryMemberReference=(),PinType.PinValueType=(),PinType.bIsMap=False,PinType.bIsSet=False,PinType.bIsArray=False,PinType.bIsReference=False,PinType.bIsConst=False,PinType.bIsWeakPointer=False,PersistentGuid=00000000000000000000000000000000,bHidden=True,bNotConnectable=False,bDefaultValueIsReadOnly=False,bDefaultValueIsIgnored=False,bAdvancedView=False,)
+		   CustomProperties Pin (PinId=98C1505C47C09BFCC2553E93D46E8FA3,PinName="Command",PinToolTip="Command\nString\n\nCommand to send to the console",PinType.PinCategory="string",PinType.PinSubCategory="",PinType.PinSubCategoryObject=None,PinType.PinSubCategoryMemberReference=(),PinType.PinValueType=(),PinType.bIsMap=False,PinType.bIsSet=False,PinType.bIsArray=False,PinType.bIsReference=False,PinType.bIsConst=False,PinType.bIsWeakPointer=False,DefaultValue="r.AllowOcclusionQueries 0",PersistentGuid=00000000000000000000000000000000,bHidden=False,bNotConnectable=False,bDefaultValueIsReadOnly=False,bDefaultValueIsIgnored=False,bAdvancedView=False,)
+		   CustomProperties Pin (PinId=C53F578840C771FD1210108AD09D137F,PinName="SpecificPlayer",PinToolTip="Specific Player\nPlayer Controller Reference\n\nIf specified, the console command will be routed through the specified player",PinType.PinCategory="object",PinType.PinSubCategory="",PinType.PinSubCategoryObject=Class'/Script/Engine.PlayerController',PinType.PinSubCategoryMemberReference=(),PinType.PinValueType=(),PinType.bIsMap=False,PinType.bIsSet=False,PinType.bIsArray=False,PinType.bIsReference=False,PinType.bIsConst=False,PinType.bIsWeakPointer=False,PersistentGuid=00000000000000000000000000000000,bHidden=False,bNotConnectable=False,bDefaultValueIsReadOnly=False,bDefaultValueIsIgnored=False,bAdvancedView=False,)
+		End Object
+		Begin Object Class=K2Node_CallFunction Name="K2Node_CallFunction_9"
+		   FunctionReference=(MemberParent=Class'/Script/Engine.KismetSystemLibrary',MemberName="ExecuteConsoleCommand")
+		   NodePosX=-432
+		   NodePosY=752
+		   NodeGuid=BFBB43724C75DB94647061BC8A9CEF06
+		   CustomProperties Pin (PinId=FAE82F5E422A038876E1CE84A81BD4E5,PinName="execute",PinToolTip="\nExec",PinType.PinCategory="exec",PinType.PinSubCategory="",PinType.PinSubCategoryObject=None,PinType.PinSubCategoryMemberReference=(),PinType.PinValueType=(),PinType.bIsMap=False,PinType.bIsSet=False,PinType.bIsArray=False,PinType.bIsReference=False,PinType.bIsConst=False,PinType.bIsWeakPointer=False,LinkedTo=(K2Node_CallFunction_44 46C4EC604A2519BABCFA6B909568F518,),PersistentGuid=00000000000000000000000000000000,bHidden=False,bNotConnectable=False,bDefaultValueIsReadOnly=False,bDefaultValueIsIgnored=False,bAdvancedView=False,)
+		   CustomProperties Pin (PinId=46C4EC604A2519BABCFA6B909568F518,PinName="then",PinToolTip="\nExec",Direction="EGPD_Output",PinType.PinCategory="exec",PinType.PinSubCategory="",PinType.PinSubCategoryObject=None,PinType.PinSubCategoryMemberReference=(),PinType.PinValueType=(),PinType.bIsMap=False,PinType.bIsSet=False,PinType.bIsArray=False,PinType.bIsReference=False,PinType.bIsConst=False,PinType.bIsWeakPointer=False,LinkedTo=(K2Node_CallFunction_33 268218444F4DE85418EE0DA390FB028D,),PersistentGuid=00000000000000000000000000000000,bHidden=False,bNotConnectable=False,bDefaultValueIsReadOnly=False,bDefaultValueIsIgnored=False,bAdvancedView=False,)
+		   CustomProperties Pin (PinId=AB16A1664BEB283188CC7EBFBD515297,PinName="self",PinFriendlyName=NSLOCTEXT("K2Node", "Target", "Target"),PinToolTip="Target\nKismet System Library Reference",PinType.PinCategory="object",PinType.PinSubCategory="",PinType.PinSubCategoryObject=Class'/Script/Engine.KismetSystemLibrary',PinType.PinSubCategoryMemberReference=(),PinType.PinValueType=(),PinType.bIsMap=False,PinType.bIsSet=False,PinType.bIsArray=False,PinType.bIsReference=False,PinType.bIsConst=False,PinType.bIsWeakPointer=False,DefaultObject="/Script/Engine.Default__KismetSystemLibrary",PersistentGuid=00000000000000000000000000000000,bHidden=True,bNotConnectable=False,bDefaultValueIsReadOnly=False,bDefaultValueIsIgnored=False,bAdvancedView=False,)
+		   CustomProperties Pin (PinId=07888691421445ABFA1A0CBEF8592B32,PinName="WorldContextObject",PinToolTip="World Context Object\nObject Reference",PinType.PinCategory="object",PinType.PinSubCategory="",PinType.PinSubCategoryObject=Class'/Script/CoreUObject.Object',PinType.PinSubCategoryMemberReference=(),PinType.PinValueType=(),PinType.bIsMap=False,PinType.bIsSet=False,PinType.bIsArray=False,PinType.bIsReference=False,PinType.bIsConst=False,PinType.bIsWeakPointer=False,PersistentGuid=00000000000000000000000000000000,bHidden=True,bNotConnectable=False,bDefaultValueIsReadOnly=False,bDefaultValueIsIgnored=False,bAdvancedView=False,)
+		   CustomProperties Pin (PinId=98C1505C47C09BFCC2553E93D46E8FA3,PinName="Command",PinToolTip="Command\nString\n\nCommand to send to the console",PinType.PinCategory="string",PinType.PinSubCategory="",PinType.PinSubCategoryObject=None,PinType.PinSubCategoryMemberReference=(),PinType.PinValueType=(),PinType.bIsMap=False,PinType.bIsSet=False,PinType.bIsArray=False,PinType.bIsReference=False,PinType.bIsConst=False,PinType.bIsWeakPointer=False,DefaultValue="t.MaxFPS 60",PersistentGuid=00000000000000000000000000000000,bHidden=False,bNotConnectable=False,bDefaultValueIsReadOnly=False,bDefaultValueIsIgnored=False,bAdvancedView=False,)
+		   CustomProperties Pin (PinId=C53F578840C771FD1210108AD09D137F,PinName="SpecificPlayer",PinToolTip="Specific Player\nPlayer Controller Reference\n\nIf specified, the console command will be routed through the specified player",PinType.PinCategory="object",PinType.PinSubCategory="",PinType.PinSubCategoryObject=Class'/Script/Engine.PlayerController',PinType.PinSubCategoryMemberReference=(),PinType.PinValueType=(),PinType.bIsMap=False,PinType.bIsSet=False,PinType.bIsArray=False,PinType.bIsReference=False,PinType.bIsConst=False,PinType.bIsWeakPointer=False,PersistentGuid=00000000000000000000000000000000,bHidden=False,bNotConnectable=False,bDefaultValueIsReadOnly=False,bDefaultValueIsIgnored=False,bAdvancedView=False,)
+		End Object
+		Begin Object Class=K2Node_CallFunction Name="K2Node_CallFunction_33"
+		   FunctionReference=(MemberParent=Class'/Script/Engine.KismetSystemLibrary',MemberName="ExecuteConsoleCommand")
+		   NodePosX=-112
+		   NodePosY=752
+		   NodeGuid=89B56A934CDE5F572EDF79BFD7FEFD73
+		   CustomProperties Pin (PinId=268218444F4DE85418EE0DA390FB028D,PinName="execute",PinToolTip="\nExec",PinType.PinCategory="exec",PinType.PinSubCategory="",PinType.PinSubCategoryObject=None,PinType.PinSubCategoryMemberReference=(),PinType.PinValueType=(),PinType.bIsMap=False,PinType.bIsSet=False,PinType.bIsArray=False,PinType.bIsReference=False,PinType.bIsConst=False,PinType.bIsWeakPointer=False,LinkedTo=(K2Node_CallFunction_9 46C4EC604A2519BABCFA6B909568F518,),PersistentGuid=00000000000000000000000000000000,bHidden=False,bNotConnectable=False,bDefaultValueIsReadOnly=False,bDefaultValueIsIgnored=False,bAdvancedView=False,)
+		   CustomProperties Pin (PinId=45FE35AD47ABC7599556B9BD32D8C851,PinName="then",PinToolTip="\nExec",Direction="EGPD_Output",PinType.PinCategory="exec",PinType.PinSubCategory="",PinType.PinSubCategoryObject=None,PinType.PinSubCategoryMemberReference=(),PinType.PinValueType=(),PinType.bIsMap=False,PinType.bIsSet=False,PinType.bIsArray=False,PinType.bIsReference=False,PinType.bIsConst=False,PinType.bIsWeakPointer=False,LinkedTo=(K2Node_CallFunction_26 268218444F4DE85418EE0DA390FB028D,),PersistentGuid=00000000000000000000000000000000,bHidden=False,bNotConnectable=False,bDefaultValueIsReadOnly=False,bDefaultValueIsIgnored=False,bAdvancedView=False,)
+		   CustomProperties Pin (PinId=EC90966D40C68E27990B278689DD2416,PinName="self",PinFriendlyName=NSLOCTEXT("K2Node", "Target", "Target"),PinToolTip="Target\nKismet System Library Reference",PinType.PinCategory="object",PinType.PinSubCategory="",PinType.PinSubCategoryObject=Class'/Script/Engine.KismetSystemLibrary',PinType.PinSubCategoryMemberReference=(),PinType.PinValueType=(),PinType.bIsMap=False,PinType.bIsSet=False,PinType.bIsArray=False,PinType.bIsReference=False,PinType.bIsConst=False,PinType.bIsWeakPointer=False,DefaultObject="/Script/Engine.Default__KismetSystemLibrary",PersistentGuid=00000000000000000000000000000000,bHidden=True,bNotConnectable=False,bDefaultValueIsReadOnly=False,bDefaultValueIsIgnored=False,bAdvancedView=False,)
+		   CustomProperties Pin (PinId=41448AB34ADD5530E946A68865DE6440,PinName="WorldContextObject",PinToolTip="World Context Object\nObject Reference",PinType.PinCategory="object",PinType.PinSubCategory="",PinType.PinSubCategoryObject=Class'/Script/CoreUObject.Object',PinType.PinSubCategoryMemberReference=(),PinType.PinValueType=(),PinType.bIsMap=False,PinType.bIsSet=False,PinType.bIsArray=False,PinType.bIsReference=False,PinType.bIsConst=False,PinType.bIsWeakPointer=False,PersistentGuid=00000000000000000000000000000000,bHidden=True,bNotConnectable=False,bDefaultValueIsReadOnly=False,bDefaultValueIsIgnored=False,bAdvancedView=False,)
+		   CustomProperties Pin (PinId=D70D49A84DF692F95492EFBB818833E9,PinName="Command",PinToolTip="Command\nString\n\nCommand to send to the console",PinType.PinCategory="string",PinType.PinSubCategory="",PinType.PinSubCategoryObject=None,PinType.PinSubCategoryMemberReference=(),PinType.PinValueType=(),PinType.bIsMap=False,PinType.bIsSet=False,PinType.bIsArray=False,PinType.bIsReference=False,PinType.bIsConst=False,PinType.bIsWeakPointer=False,DefaultValue="SP.HorizontalAngularIncrement 2",PersistentGuid=00000000000000000000000000000000,bHidden=False,bNotConnectable=False,bDefaultValueIsReadOnly=False,bDefaultValueIsIgnored=False,bAdvancedView=False,)
+		   CustomProperties Pin (PinId=F37AE70C47B2BB78C35E2EBC9A8D006E,PinName="SpecificPlayer",PinToolTip="Specific Player\nPlayer Controller Reference\n\nIf specified, the console command will be routed through the specified player",PinType.PinCategory="object",PinType.PinSubCategory="",PinType.PinSubCategoryObject=Class'/Script/Engine.PlayerController',PinType.PinSubCategoryMemberReference=(),PinType.PinValueType=(),PinType.bIsMap=False,PinType.bIsSet=False,PinType.bIsArray=False,PinType.bIsReference=False,PinType.bIsConst=False,PinType.bIsWeakPointer=False,PersistentGuid=00000000000000000000000000000000,bHidden=False,bNotConnectable=False,bDefaultValueIsReadOnly=False,bDefaultValueIsIgnored=False,bAdvancedView=False,)
+		End Object
+		Begin Object Class=K2Node_CallFunction Name="K2Node_CallFunction_26"
+		   FunctionReference=(MemberParent=Class'/Script/Engine.KismetSystemLibrary',MemberName="ExecuteConsoleCommand")
+		   NodePosX=-112
+		   NodePosY=896
+		   NodeGuid=3DBD5B1D465AD78682F13EA3DB200937
+		   CustomProperties Pin (PinId=268218444F4DE85418EE0DA390FB028D,PinName="execute",PinToolTip="\nExec",PinType.PinCategory="exec",PinType.PinSubCategory="",PinType.PinSubCategoryObject=None,PinType.PinSubCategoryMemberReference=(),PinType.PinValueType=(),PinType.bIsMap=False,PinType.bIsSet=False,PinType.bIsArray=False,PinType.bIsReference=False,PinType.bIsConst=False,PinType.bIsWeakPointer=False,LinkedTo=(K2Node_CallFunction_33 45FE35AD47ABC7599556B9BD32D8C851,),PersistentGuid=00000000000000000000000000000000,bHidden=False,bNotConnectable=False,bDefaultValueIsReadOnly=False,bDefaultValueIsIgnored=False,bAdvancedView=False,)
+		   CustomProperties Pin (PinId=45FE35AD47ABC7599556B9BD32D8C851,PinName="then",PinToolTip="\nExec",Direction="EGPD_Output",PinType.PinCategory="exec",PinType.PinSubCategory="",PinType.PinSubCategoryObject=None,PinType.PinSubCategoryMemberReference=(),PinType.PinValueType=(),PinType.bIsMap=False,PinType.bIsSet=False,PinType.bIsArray=False,PinType.bIsReference=False,PinType.bIsConst=False,PinType.bIsWeakPointer=False,LinkedTo=(K2Node_CallFunction_34 268218444F4DE85418EE0DA390FB028D,),PersistentGuid=00000000000000000000000000000000,bHidden=False,bNotConnectable=False,bDefaultValueIsReadOnly=False,bDefaultValueIsIgnored=False,bAdvancedView=False,)
+		   CustomProperties Pin (PinId=EC90966D40C68E27990B278689DD2416,PinName="self",PinFriendlyName=NSLOCTEXT("K2Node", "Target", "Target"),PinToolTip="Target\nKismet System Library Reference",PinType.PinCategory="object",PinType.PinSubCategory="",PinType.PinSubCategoryObject=Class'/Script/Engine.KismetSystemLibrary',PinType.PinSubCategoryMemberReference=(),PinType.PinValueType=(),PinType.bIsMap=False,PinType.bIsSet=False,PinType.bIsArray=False,PinType.bIsReference=False,PinType.bIsConst=False,PinType.bIsWeakPointer=False,DefaultObject="/Script/Engine.Default__KismetSystemLibrary",PersistentGuid=00000000000000000000000000000000,bHidden=True,bNotConnectable=False,bDefaultValueIsReadOnly=False,bDefaultValueIsIgnored=False,bAdvancedView=False,)
+		   CustomProperties Pin (PinId=41448AB34ADD5530E946A68865DE6440,PinName="WorldContextObject",PinToolTip="World Context Object\nObject Reference",PinType.PinCategory="object",PinType.PinSubCategory="",PinType.PinSubCategoryObject=Class'/Script/CoreUObject.Object',PinType.PinSubCategoryMemberReference=(),PinType.PinValueType=(),PinType.bIsMap=False,PinType.bIsSet=False,PinType.bIsArray=False,PinType.bIsReference=False,PinType.bIsConst=False,PinType.bIsWeakPointer=False,PersistentGuid=00000000000000000000000000000000,bHidden=True,bNotConnectable=False,bDefaultValueIsReadOnly=False,bDefaultValueIsIgnored=False,bAdvancedView=False,)
+		   CustomProperties Pin (PinId=D70D49A84DF692F95492EFBB818833E9,PinName="Command",PinToolTip="Command\nString\n\nCommand to send to the console",PinType.PinCategory="string",PinType.PinSubCategory="",PinType.PinSubCategoryObject=None,PinType.PinSubCategoryMemberReference=(),PinType.PinValueType=(),PinType.bIsMap=False,PinType.bIsSet=False,PinType.bIsArray=False,PinType.bIsReference=False,PinType.bIsConst=False,PinType.bIsWeakPointer=False,DefaultValue="SP.VerticalAngularIncrement 30",PersistentGuid=00000000000000000000000000000000,bHidden=False,bNotConnectable=False,bDefaultValueIsReadOnly=False,bDefaultValueIsIgnored=False,bAdvancedView=False,)
+		   CustomProperties Pin (PinId=F37AE70C47B2BB78C35E2EBC9A8D006E,PinName="SpecificPlayer",PinToolTip="Specific Player\nPlayer Controller Reference\n\nIf specified, the console command will be routed through the specified player",PinType.PinCategory="object",PinType.PinSubCategory="",PinType.PinSubCategoryObject=Class'/Script/Engine.PlayerController',PinType.PinSubCategoryMemberReference=(),PinType.PinValueType=(),PinType.bIsMap=False,PinType.bIsSet=False,PinType.bIsArray=False,PinType.bIsReference=False,PinType.bIsConst=False,PinType.bIsWeakPointer=False,PersistentGuid=00000000000000000000000000000000,bHidden=False,bNotConnectable=False,bDefaultValueIsReadOnly=False,bDefaultValueIsIgnored=False,bAdvancedView=False,)
+		End Object
+		Begin Object Class=K2Node_CallFunction Name="K2Node_CallFunction_34"
+		   FunctionReference=(MemberParent=Class'/Script/Engine.KismetSystemLibrary',MemberName="ExecuteConsoleCommand")
+		   NodePosX=-112
+		   NodePosY=1040
+		   NodeGuid=ED0022E74518E4E655085E9BA08A638A
+		   CustomProperties Pin (PinId=268218444F4DE85418EE0DA390FB028D,PinName="execute",PinToolTip="\nExec",PinType.PinCategory="exec",PinType.PinSubCategory="",PinType.PinSubCategoryObject=None,PinType.PinSubCategoryMemberReference=(),PinType.PinValueType=(),PinType.bIsMap=False,PinType.bIsSet=False,PinType.bIsArray=False,PinType.bIsReference=False,PinType.bIsConst=False,PinType.bIsWeakPointer=False,LinkedTo=(K2Node_CallFunction_26 45FE35AD47ABC7599556B9BD32D8C851,),PersistentGuid=00000000000000000000000000000000,bHidden=False,bNotConnectable=False,bDefaultValueIsReadOnly=False,bDefaultValueIsIgnored=False,bAdvancedView=False,)
+		   CustomProperties Pin (PinId=45FE35AD47ABC7599556B9BD32D8C851,PinName="then",PinToolTip="\nExec",Direction="EGPD_Output",PinType.PinCategory="exec",PinType.PinSubCategory="",PinType.PinSubCategoryObject=None,PinType.PinSubCategoryMemberReference=(),PinType.PinValueType=(),PinType.bIsMap=False,PinType.bIsSet=False,PinType.bIsArray=False,PinType.bIsReference=False,PinType.bIsConst=False,PinType.bIsWeakPointer=False,LinkedTo=(K2Node_CallFunction_39 268218444F4DE85418EE0DA390FB028D,),PersistentGuid=00000000000000000000000000000000,bHidden=False,bNotConnectable=False,bDefaultValueIsReadOnly=False,bDefaultValueIsIgnored=False,bAdvancedView=False,)
+		   CustomProperties Pin (PinId=EC90966D40C68E27990B278689DD2416,PinName="self",PinFriendlyName=NSLOCTEXT("K2Node", "Target", "Target"),PinToolTip="Target\nKismet System Library Reference",PinType.PinCategory="object",PinType.PinSubCategory="",PinType.PinSubCategoryObject=Class'/Script/Engine.KismetSystemLibrary',PinType.PinSubCategoryMemberReference=(),PinType.PinValueType=(),PinType.bIsMap=False,PinType.bIsSet=False,PinType.bIsArray=False,PinType.bIsReference=False,PinType.bIsConst=False,PinType.bIsWeakPointer=False,DefaultObject="/Script/Engine.Default__KismetSystemLibrary",PersistentGuid=00000000000000000000000000000000,bHidden=True,bNotConnectable=False,bDefaultValueIsReadOnly=False,bDefaultValueIsIgnored=False,bAdvancedView=False,)
+		   CustomProperties Pin (PinId=41448AB34ADD5530E946A68865DE6440,PinName="WorldContextObject",PinToolTip="World Context Object\nObject Reference",PinType.PinCategory="object",PinType.PinSubCategory="",PinType.PinSubCategoryObject=Class'/Script/CoreUObject.Object',PinType.PinSubCategoryMemberReference=(),PinType.PinValueType=(),PinType.bIsMap=False,PinType.bIsSet=False,PinType.bIsArray=False,PinType.bIsReference=False,PinType.bIsConst=False,PinType.bIsWeakPointer=False,PersistentGuid=00000000000000000000000000000000,bHidden=True,bNotConnectable=False,bDefaultValueIsReadOnly=False,bDefaultValueIsIgnored=False,bAdvancedView=False,)
+		   CustomProperties Pin (PinId=D70D49A84DF692F95492EFBB818833E9,PinName="Command",PinToolTip="Command\nString\n\nCommand to send to the console",PinType.PinCategory="string",PinType.PinSubCategory="",PinType.PinSubCategoryObject=None,PinType.PinSubCategoryMemberReference=(),PinType.PinValueType=(),PinType.bIsMap=False,PinType.bIsSet=False,PinType.bIsArray=False,PinType.bIsReference=False,PinType.bIsConst=False,PinType.bIsWeakPointer=False,DefaultValue="SP.CaptureHorizontalFOV 30",PersistentGuid=00000000000000000000000000000000,bHidden=False,bNotConnectable=False,bDefaultValueIsReadOnly=False,bDefaultValueIsIgnored=False,bAdvancedView=False,)
+		   CustomProperties Pin (PinId=F37AE70C47B2BB78C35E2EBC9A8D006E,PinName="SpecificPlayer",PinToolTip="Specific Player\nPlayer Controller Reference\n\nIf specified, the console command will be routed through the specified player",PinType.PinCategory="object",PinType.PinSubCategory="",PinType.PinSubCategoryObject=Class'/Script/Engine.PlayerController',PinType.PinSubCategoryMemberReference=(),PinType.PinValueType=(),PinType.bIsMap=False,PinType.bIsSet=False,PinType.bIsArray=False,PinType.bIsReference=False,PinType.bIsConst=False,PinType.bIsWeakPointer=False,PersistentGuid=00000000000000000000000000000000,bHidden=False,bNotConnectable=False,bDefaultValueIsReadOnly=False,bDefaultValueIsIgnored=False,bAdvancedView=False,)
+		End Object
+		Begin Object Class=K2Node_CallFunction Name="K2Node_CallFunction_39"
+		   FunctionReference=(MemberParent=Class'/Script/Engine.KismetSystemLibrary',MemberName="ExecuteConsoleCommand")
+		   NodePosX=-112
+		   NodePosY=1184
+		   NodeGuid=F668B2DD48DF11CDFAFCAE8218B28221
+		   CustomProperties Pin (PinId=268218444F4DE85418EE0DA390FB028D,PinName="execute",PinToolTip="\nExec",PinType.PinCategory="exec",PinType.PinSubCategory="",PinType.PinSubCategoryObject=None,PinType.PinSubCategoryMemberReference=(),PinType.PinValueType=(),PinType.bIsMap=False,PinType.bIsSet=False,PinType.bIsArray=False,PinType.bIsReference=False,PinType.bIsConst=False,PinType.bIsWeakPointer=False,LinkedTo=(K2Node_CallFunction_34 45FE35AD47ABC7599556B9BD32D8C851,),PersistentGuid=00000000000000000000000000000000,bHidden=False,bNotConnectable=False,bDefaultValueIsReadOnly=False,bDefaultValueIsIgnored=False,bAdvancedView=False,)
+		   CustomProperties Pin (PinId=45FE35AD47ABC7599556B9BD32D8C851,PinName="then",PinToolTip="\nExec",Direction="EGPD_Output",PinType.PinCategory="exec",PinType.PinSubCategory="",PinType.PinSubCategoryObject=None,PinType.PinSubCategoryMemberReference=(),PinType.PinValueType=(),PinType.bIsMap=False,PinType.bIsSet=False,PinType.bIsArray=False,PinType.bIsReference=False,PinType.bIsConst=False,PinType.bIsWeakPointer=False,LinkedTo=(K2Node_CallFunction_40 268218444F4DE85418EE0DA390FB028D,),PersistentGuid=00000000000000000000000000000000,bHidden=False,bNotConnectable=False,bDefaultValueIsReadOnly=False,bDefaultValueIsIgnored=False,bAdvancedView=False,)
+		   CustomProperties Pin (PinId=EC90966D40C68E27990B278689DD2416,PinName="self",PinFriendlyName=NSLOCTEXT("K2Node", "Target", "Target"),PinToolTip="Target\nKismet System Library Reference",PinType.PinCategory="object",PinType.PinSubCategory="",PinType.PinSubCategoryObject=Class'/Script/Engine.KismetSystemLibrary',PinType.PinSubCategoryMemberReference=(),PinType.PinValueType=(),PinType.bIsMap=False,PinType.bIsSet=False,PinType.bIsArray=False,PinType.bIsReference=False,PinType.bIsConst=False,PinType.bIsWeakPointer=False,DefaultObject="/Script/Engine.Default__KismetSystemLibrary",PersistentGuid=00000000000000000000000000000000,bHidden=True,bNotConnectable=False,bDefaultValueIsReadOnly=False,bDefaultValueIsIgnored=False,bAdvancedView=False,)
+		   CustomProperties Pin (PinId=41448AB34ADD5530E946A68865DE6440,PinName="WorldContextObject",PinToolTip="World Context Object\nObject Reference",PinType.PinCategory="object",PinType.PinSubCategory="",PinType.PinSubCategoryObject=Class'/Script/CoreUObject.Object',PinType.PinSubCategoryMemberReference=(),PinType.PinValueType=(),PinType.bIsMap=False,PinType.bIsSet=False,PinType.bIsArray=False,PinType.bIsReference=False,PinType.bIsConst=False,PinType.bIsWeakPointer=False,PersistentGuid=00000000000000000000000000000000,bHidden=True,bNotConnectable=False,bDefaultValueIsReadOnly=False,bDefaultValueIsIgnored=False,bAdvancedView=False,)
+		   CustomProperties Pin (PinId=D70D49A84DF692F95492EFBB818833E9,PinName="Command",PinToolTip="Command\nString\n\nCommand to send to the console",PinType.PinCategory="string",PinType.PinSubCategory="",PinType.PinSubCategoryObject=None,PinType.PinSubCategoryMemberReference=(),PinType.PinValueType=(),PinType.bIsMap=False,PinType.bIsSet=False,PinType.bIsArray=False,PinType.bIsReference=False,PinType.bIsConst=False,PinType.bIsWeakPointer=False,DefaultValue="SP.ConcurrentCaptures 16",PersistentGuid=00000000000000000000000000000000,bHidden=False,bNotConnectable=False,bDefaultValueIsReadOnly=False,bDefaultValueIsIgnored=False,bAdvancedView=False,)
+		   CustomProperties Pin (PinId=F37AE70C47B2BB78C35E2EBC9A8D006E,PinName="SpecificPlayer",PinToolTip="Specific Player\nPlayer Controller Reference\n\nIf specified, the console command will be routed through the specified player",PinType.PinCategory="object",PinType.PinSubCategory="",PinType.PinSubCategoryObject=Class'/Script/Engine.PlayerController',PinType.PinSubCategoryMemberReference=(),PinType.PinValueType=(),PinType.bIsMap=False,PinType.bIsSet=False,PinType.bIsArray=False,PinType.bIsReference=False,PinType.bIsConst=False,PinType.bIsWeakPointer=False,PersistentGuid=00000000000000000000000000000000,bHidden=False,bNotConnectable=False,bDefaultValueIsReadOnly=False,bDefaultValueIsIgnored=False,bAdvancedView=False,)
+		End Object
+		Begin Object Class=K2Node_CallFunction Name="K2Node_CallFunction_40"
+		   FunctionReference=(MemberParent=Class'/Script/Engine.KismetSystemLibrary',MemberName="ExecuteConsoleCommand")
+		   NodePosX=288
+		   NodePosY=768
+		   NodeGuid=D7F64C2D4B99FE1060BF1EA8B26B170D
+		   CustomProperties Pin (PinId=268218444F4DE85418EE0DA390FB028D,PinName="execute",PinToolTip="\nExec",PinType.PinCategory="exec",PinType.PinSubCategory="",PinType.PinSubCategoryObject=None,PinType.PinSubCategoryMemberReference=(),PinType.PinValueType=(),PinType.bIsMap=False,PinType.bIsSet=False,PinType.bIsArray=False,PinType.bIsReference=False,PinType.bIsConst=False,PinType.bIsWeakPointer=False,LinkedTo=(K2Node_CallFunction_39 45FE35AD47ABC7599556B9BD32D8C851,),PersistentGuid=00000000000000000000000000000000,bHidden=False,bNotConnectable=False,bDefaultValueIsReadOnly=False,bDefaultValueIsIgnored=False,bAdvancedView=False,)
+		   CustomProperties Pin (PinId=45FE35AD47ABC7599556B9BD32D8C851,PinName="then",PinToolTip="\nExec",Direction="EGPD_Output",PinType.PinCategory="exec",PinType.PinSubCategory="",PinType.PinSubCategoryObject=None,PinType.PinSubCategoryMemberReference=(),PinType.PinValueType=(),PinType.bIsMap=False,PinType.bIsSet=False,PinType.bIsArray=False,PinType.bIsReference=False,PinType.bIsConst=False,PinType.bIsWeakPointer=False,LinkedTo=(K2Node_CallFunction_12371 268218444F4DE85418EE0DA390FB028D,),PersistentGuid=00000000000000000000000000000000,bHidden=False,bNotConnectable=False,bDefaultValueIsReadOnly=False,bDefaultValueIsIgnored=False,bAdvancedView=False,)
+		   CustomProperties Pin (PinId=EC90966D40C68E27990B278689DD2416,PinName="self",PinFriendlyName=NSLOCTEXT("K2Node", "Target", "Target"),PinToolTip="Target\nKismet System Library Reference",PinType.PinCategory="object",PinType.PinSubCategory="",PinType.PinSubCategoryObject=Class'/Script/Engine.KismetSystemLibrary',PinType.PinSubCategoryMemberReference=(),PinType.PinValueType=(),PinType.bIsMap=False,PinType.bIsSet=False,PinType.bIsArray=False,PinType.bIsReference=False,PinType.bIsConst=False,PinType.bIsWeakPointer=False,DefaultObject="/Script/Engine.Default__KismetSystemLibrary",PersistentGuid=00000000000000000000000000000000,bHidden=True,bNotConnectable=False,bDefaultValueIsReadOnly=False,bDefaultValueIsIgnored=False,bAdvancedView=False,)
+		   CustomProperties Pin (PinId=41448AB34ADD5530E946A68865DE6440,PinName="WorldContextObject",PinToolTip="World Context Object\nObject Reference",PinType.PinCategory="object",PinType.PinSubCategory="",PinType.PinSubCategoryObject=Class'/Script/CoreUObject.Object',PinType.PinSubCategoryMemberReference=(),PinType.PinValueType=(),PinType.bIsMap=False,PinType.bIsSet=False,PinType.bIsArray=False,PinType.bIsReference=False,PinType.bIsConst=False,PinType.bIsWeakPointer=False,PersistentGuid=00000000000000000000000000000000,bHidden=True,bNotConnectable=False,bDefaultValueIsReadOnly=False,bDefaultValueIsIgnored=False,bAdvancedView=False,)
+		   CustomProperties Pin (PinId=D70D49A84DF692F95492EFBB818833E9,PinName="Command",PinToolTip="Command\nString\n\nCommand to send to the console",PinType.PinCategory="string",PinType.PinSubCategory="",PinType.PinSubCategoryObject=None,PinType.PinSubCategoryMemberReference=(),PinType.PinValueType=(),PinType.bIsMap=False,PinType.bIsSet=False,PinType.bIsArray=False,PinType.bIsReference=False,PinType.bIsConst=False,PinType.bIsWeakPointer=False,DefaultValue="SP.CaptureSlicePixelWidth 720",PersistentGuid=00000000000000000000000000000000,bHidden=False,bNotConnectable=False,bDefaultValueIsReadOnly=False,bDefaultValueIsIgnored=False,bAdvancedView=False,)
+		   CustomProperties Pin (PinId=F37AE70C47B2BB78C35E2EBC9A8D006E,PinName="SpecificPlayer",PinToolTip="Specific Player\nPlayer Controller Reference\n\nIf specified, the console command will be routed through the specified player",PinType.PinCategory="object",PinType.PinSubCategory="",PinType.PinSubCategoryObject=Class'/Script/Engine.PlayerController',PinType.PinSubCategoryMemberReference=(),PinType.PinValueType=(),PinType.bIsMap=False,PinType.bIsSet=False,PinType.bIsArray=False,PinType.bIsReference=False,PinType.bIsConst=False,PinType.bIsWeakPointer=False,PersistentGuid=00000000000000000000000000000000,bHidden=False,bNotConnectable=False,bDefaultValueIsReadOnly=False,bDefaultValueIsIgnored=False,bAdvancedView=False,)
+		End Object
+		Begin Object Class=K2Node_CallFunction Name="K2Node_CallFunction_12371"
+		   FunctionReference=(MemberParent=Class'/Script/Engine.KismetSystemLibrary',MemberName="ExecuteConsoleCommand")
+		   NodePosX=288
+		   NodePosY=912
+		   NodeGuid=F934806E4BC115C6B46C1FB67BB8B59B
+		   CustomProperties Pin (PinId=268218444F4DE85418EE0DA390FB028D,PinName="execute",PinToolTip="\nExec",PinType.PinCategory="exec",PinType.PinSubCategory="",PinType.PinSubCategoryObject=None,PinType.PinSubCategoryMemberReference=(),PinType.PinValueType=(),PinType.bIsMap=False,PinType.bIsSet=False,PinType.bIsArray=False,PinType.bIsReference=False,PinType.bIsConst=False,PinType.bIsWeakPointer=False,LinkedTo=(K2Node_CallFunction_40 45FE35AD47ABC7599556B9BD32D8C851,),PersistentGuid=00000000000000000000000000000000,bHidden=False,bNotConnectable=False,bDefaultValueIsReadOnly=False,bDefaultValueIsIgnored=False,bAdvancedView=False,)
+		   CustomProperties Pin (PinId=45FE35AD47ABC7599556B9BD32D8C851,PinName="then",PinToolTip="\nExec",Direction="EGPD_Output",PinType.PinCategory="exec",PinType.PinSubCategory="",PinType.PinSubCategoryObject=None,PinType.PinSubCategoryMemberReference=(),PinType.PinValueType=(),PinType.bIsMap=False,PinType.bIsSet=False,PinType.bIsArray=False,PinType.bIsReference=False,PinType.bIsConst=False,PinType.bIsWeakPointer=False,LinkedTo=(K2Node_CallFunction_25 268218444F4DE85418EE0DA390FB028D,),PersistentGuid=00000000000000000000000000000000,bHidden=False,bNotConnectable=False,bDefaultValueIsReadOnly=False,bDefaultValueIsIgnored=False,bAdvancedView=False,)
+		   CustomProperties Pin (PinId=EC90966D40C68E27990B278689DD2416,PinName="self",PinFriendlyName=NSLOCTEXT("K2Node", "Target", "Target"),PinToolTip="Target\nKismet System Library Reference",PinType.PinCategory="object",PinType.PinSubCategory="",PinType.PinSubCategoryObject=Class'/Script/Engine.KismetSystemLibrary',PinType.PinSubCategoryMemberReference=(),PinType.PinValueType=(),PinType.bIsMap=False,PinType.bIsSet=False,PinType.bIsArray=False,PinType.bIsReference=False,PinType.bIsConst=False,PinType.bIsWeakPointer=False,DefaultObject="/Script/Engine.Default__KismetSystemLibrary",PersistentGuid=00000000000000000000000000000000,bHidden=True,bNotConnectable=False,bDefaultValueIsReadOnly=False,bDefaultValueIsIgnored=False,bAdvancedView=False,)
+		   CustomProperties Pin (PinId=41448AB34ADD5530E946A68865DE6440,PinName="WorldContextObject",PinToolTip="World Context Object\nObject Reference",PinType.PinCategory="object",PinType.PinSubCategory="",PinType.PinSubCategoryObject=Class'/Script/CoreUObject.Object',PinType.PinSubCategoryMemberReference=(),PinType.PinValueType=(),PinType.bIsMap=False,PinType.bIsSet=False,PinType.bIsArray=False,PinType.bIsReference=False,PinType.bIsConst=False,PinType.bIsWeakPointer=False,PersistentGuid=00000000000000000000000000000000,bHidden=True,bNotConnectable=False,bDefaultValueIsReadOnly=False,bDefaultValueIsIgnored=False,bAdvancedView=False,)
+		   CustomProperties Pin (PinId=D70D49A84DF692F95492EFBB818833E9,PinName="Command",PinToolTip="Command\nString\n\nCommand to send to the console",PinType.PinCategory="string",PinType.PinSubCategory="",PinType.PinSubCategoryObject=None,PinType.PinSubCategoryMemberReference=(),PinType.PinValueType=(),PinType.bIsMap=False,PinType.bIsSet=False,PinType.bIsArray=False,PinType.bIsReference=False,PinType.bIsConst=False,PinType.bIsWeakPointer=False,DefaultValue="SP.StepCaptureWidth 4096",PersistentGuid=00000000000000000000000000000000,bHidden=False,bNotConnectable=False,bDefaultValueIsReadOnly=False,bDefaultValueIsIgnored=False,bAdvancedView=False,)
+		   CustomProperties Pin (PinId=F37AE70C47B2BB78C35E2EBC9A8D006E,PinName="SpecificPlayer",PinToolTip="Specific Player\nPlayer Controller Reference\n\nIf specified, the console command will be routed through the specified player",PinType.PinCategory="object",PinType.PinSubCategory="",PinType.PinSubCategoryObject=Class'/Script/Engine.PlayerController',PinType.PinSubCategoryMemberReference=(),PinType.PinValueType=(),PinType.bIsMap=False,PinType.bIsSet=False,PinType.bIsArray=False,PinType.bIsReference=False,PinType.bIsConst=False,PinType.bIsWeakPointer=False,PersistentGuid=00000000000000000000000000000000,bHidden=False,bNotConnectable=False,bDefaultValueIsReadOnly=False,bDefaultValueIsIgnored=False,bAdvancedView=False,)
+		End Object
+		Begin Object Class=K2Node_CallFunction Name="K2Node_CallFunction_25"
+		   FunctionReference=(MemberParent=Class'/Script/Engine.KismetSystemLibrary',MemberName="ExecuteConsoleCommand")
+		   NodePosX=688
+		   NodePosY=912
+		   NodeGuid=28D6613942F2D4D3398E50B9F33525C8
+		   CustomProperties Pin (PinId=268218444F4DE85418EE0DA390FB028D,PinName="execute",PinToolTip="\nExec",PinType.PinCategory="exec",PinType.PinSubCategory="",PinType.PinSubCategoryObject=None,PinType.PinSubCategoryMemberReference=(),PinType.PinValueType=(),PinType.bIsMap=False,PinType.bIsSet=False,PinType.bIsArray=False,PinType.bIsReference=False,PinType.bIsConst=False,PinType.bIsWeakPointer=False,LinkedTo=(K2Node_CallFunction_12371 45FE35AD47ABC7599556B9BD32D8C851,),PersistentGuid=00000000000000000000000000000000,bHidden=False,bNotConnectable=False,bDefaultValueIsReadOnly=False,bDefaultValueIsIgnored=False,bAdvancedView=False,)
+		   CustomProperties Pin (PinId=45FE35AD47ABC7599556B9BD32D8C851,PinName="then",PinToolTip="\nExec",Direction="EGPD_Output",PinType.PinCategory="exec",PinType.PinSubCategory="",PinType.PinSubCategoryObject=None,PinType.PinSubCategoryMemberReference=(),PinType.PinValueType=(),PinType.bIsMap=False,PinType.bIsSet=False,PinType.bIsArray=False,PinType.bIsReference=False,PinType.bIsConst=False,PinType.bIsWeakPointer=False,LinkedTo=(K2Node_CallFunction_151 F95DFD104982DC30E863A897771D0B16,),PersistentGuid=00000000000000000000000000000000,bHidden=False,bNotConnectable=False,bDefaultValueIsReadOnly=False,bDefaultValueIsIgnored=False,bAdvancedView=False,)
+		   CustomProperties Pin (PinId=EC90966D40C68E27990B278689DD2416,PinName="self",PinFriendlyName=NSLOCTEXT("K2Node", "Target", "Target"),PinToolTip="Target\nKismet System Library Reference",PinType.PinCategory="object",PinType.PinSubCategory="",PinType.PinSubCategoryObject=Class'/Script/Engine.KismetSystemLibrary',PinType.PinSubCategoryMemberReference=(),PinType.PinValueType=(),PinType.bIsMap=False,PinType.bIsSet=False,PinType.bIsArray=False,PinType.bIsReference=False,PinType.bIsConst=False,PinType.bIsWeakPointer=False,DefaultObject="/Script/Engine.Default__KismetSystemLibrary",PersistentGuid=00000000000000000000000000000000,bHidden=True,bNotConnectable=False,bDefaultValueIsReadOnly=False,bDefaultValueIsIgnored=False,bAdvancedView=False,)
+		   CustomProperties Pin (PinId=41448AB34ADD5530E946A68865DE6440,PinName="WorldContextObject",PinToolTip="World Context Object\nObject Reference",PinType.PinCategory="object",PinType.PinSubCategory="",PinType.PinSubCategoryObject=Class'/Script/CoreUObject.Object',PinType.PinSubCategoryMemberReference=(),PinType.PinValueType=(),PinType.bIsMap=False,PinType.bIsSet=False,PinType.bIsArray=False,PinType.bIsReference=False,PinType.bIsConst=False,PinType.bIsWeakPointer=False,PersistentGuid=00000000000000000000000000000000,bHidden=True,bNotConnectable=False,bDefaultValueIsReadOnly=False,bDefaultValueIsIgnored=False,bAdvancedView=False,)
+		   CustomProperties Pin (PinId=D70D49A84DF692F95492EFBB818833E9,PinName="Command",PinToolTip="Command\nString\n\nCommand to send to the console",PinType.PinCategory="string",PinType.PinSubCategory="",PinType.PinSubCategoryObject=None,PinType.PinSubCategoryMemberReference=(),PinType.PinValueType=(),PinType.bIsMap=False,PinType.bIsSet=False,PinType.bIsArray=False,PinType.bIsReference=False,PinType.bIsConst=False,PinType.bIsWeakPointer=False,DefaultValue="SP.PanoramicMovie 0 600",PersistentGuid=00000000000000000000000000000000,bHidden=False,bNotConnectable=False,bDefaultValueIsReadOnly=False,bDefaultValueIsIgnored=False,bAdvancedView=False,)
+		   CustomProperties Pin (PinId=F37AE70C47B2BB78C35E2EBC9A8D006E,PinName="SpecificPlayer",PinToolTip="Specific Player\nPlayer Controller Reference\n\nIf specified, the console command will be routed through the specified player",PinType.PinCategory="object",PinType.PinSubCategory="",PinType.PinSubCategoryObject=Class'/Script/Engine.PlayerController',PinType.PinSubCategoryMemberReference=(),PinType.PinValueType=(),PinType.bIsMap=False,PinType.bIsSet=False,PinType.bIsArray=False,PinType.bIsReference=False,PinType.bIsConst=False,PinType.bIsWeakPointer=False,PersistentGuid=00000000000000000000000000000000,bHidden=False,bNotConnectable=False,bDefaultValueIsReadOnly=False,bDefaultValueIsIgnored=False,bAdvancedView=False,)
+		End Object
+		Begin Object Class=EdGraphNode_Comment Name="EdGraphNode_Comment_1"
+		   NodePosX=-731
+		   NodePosY=707
+		   NodeWidth=544
+		   NodeHeight=224
+		   NodeComment="Engine Setup"
+		   NodeGuid=813223D3400DB5753536C6BEA943CC8E
+		End Object
+		Begin Object Class=EdGraphNode_Comment Name="EdGraphNode_Comment_2"
+		   NodePosX=-160
+		   NodePosY=704
+		   NodeWidth=768
+		   NodeHeight=624
+		   NodeComment="Capture Settings"
+		   NodeGuid=1ABFA54A42DE5A7B6B934695CBEC41BC
+		End Object
+		Begin Object Class=EdGraphNode_Comment Name="EdGraphNode_Comment_3"
+		   NodePosX=640
+		   NodePosY=848
+		   NodeWidth=349
+		   NodeHeight=227
+		   NodeComment="Capture Moive "
+		   NodeGuid=3298F9E14591F5BC8880B1B703E81B7F
+		End Object
+		Begin Object Class=K2Node_Literal Name="K2Node_Literal_1"
+		   ObjectRef=MatineeActor'/Game/Maps/BlueprintOffice.BlueprintOffice:PersistentLevel.MatineeActor_0'
+		   NodePosX=1056
+		   NodePosY=1008
+		   NodeGuid=ED0D00074768DFD6A8FA108E15A340C7
+		   CustomProperties Pin (PinId=ADF345724EC9A5D2063510BA2ACE8B6C,PinName="MatineeActor",PinFriendlyName="MatineeActor",Direction="EGPD_Output",PinType.PinCategory="object",PinType.PinSubCategory="",PinType.PinSubCategoryObject=Class'/Script/Engine.MatineeActor',PinType.PinSubCategoryMemberReference=(),PinType.PinValueType=(),PinType.bIsMap=False,PinType.bIsSet=False,PinType.bIsArray=False,PinType.bIsReference=False,PinType.bIsConst=False,PinType.bIsWeakPointer=False,LinkedTo=(K2Node_CallFunction_151 ED8AFCFB4DE6ED18B3857AB571EB033E,),PersistentGuid=00000000000000000000000000000000,bHidden=False,bNotConnectable=False,bDefaultValueIsReadOnly=False,bDefaultValueIsIgnored=False,bAdvancedView=False,)
+		End Object
+		Begin Object Class=K2Node_CallFunction Name="K2Node_CallFunction_151"
+		   FunctionReference=(MemberParent=Class'/Script/Engine.MatineeActor',MemberName="Play")
+		   NodePosX=1072
+		   NodePosY=896
+		   NodeGuid=81F1A4714AE767D5F9B02FB036C381CB
+		   CustomProperties Pin (PinId=F95DFD104982DC30E863A897771D0B16,PinName="execute",PinToolTip="\nExec",PinType.PinCategory="exec",PinType.PinSubCategory="",PinType.PinSubCategoryObject=None,PinType.PinSubCategoryMemberReference=(),PinType.PinValueType=(),PinType.bIsMap=False,PinType.bIsSet=False,PinType.bIsArray=False,PinType.bIsReference=False,PinType.bIsConst=False,PinType.bIsWeakPointer=False,LinkedTo=(K2Node_CallFunction_25 45FE35AD47ABC7599556B9BD32D8C851,),PersistentGuid=00000000000000000000000000000000,bHidden=False,bNotConnectable=False,bDefaultValueIsReadOnly=False,bDefaultValueIsIgnored=False,bAdvancedView=False,)
+		   CustomProperties Pin (PinId=7F591BBA4208E1E61B501E801D583F17,PinName="then",PinToolTip="\nExec",Direction="EGPD_Output",PinType.PinCategory="exec",PinType.PinSubCategory="",PinType.PinSubCategoryObject=None,PinType.PinSubCategoryMemberReference=(),PinType.PinValueType=(),PinType.bIsMap=False,PinType.bIsSet=False,PinType.bIsArray=False,PinType.bIsReference=False,PinType.bIsConst=False,PinType.bIsWeakPointer=False,PersistentGuid=00000000000000000000000000000000,bHidden=False,bNotConnectable=False,bDefaultValueIsReadOnly=False,bDefaultValueIsIgnored=False,bAdvancedView=False,)
+		   CustomProperties Pin (PinId=ED8AFCFB4DE6ED18B3857AB571EB033E,PinName="self",PinFriendlyName=NSLOCTEXT("K2Node", "Target", "Target"),PinToolTip="Target\nMatinee Actor Reference",PinType.PinCategory="object",PinType.PinSubCategory="",PinType.PinSubCategoryObject=Class'/Script/Engine.MatineeActor',PinType.PinSubCategoryMemberReference=(),PinType.PinValueType=(),PinType.bIsMap=False,PinType.bIsSet=False,PinType.bIsArray=False,PinType.bIsReference=False,PinType.bIsConst=False,PinType.bIsWeakPointer=False,LinkedTo=(K2Node_Literal_1 ADF345724EC9A5D2063510BA2ACE8B6C,),PersistentGuid=00000000000000000000000000000000,bHidden=False,bNotConnectable=False,bDefaultValueIsReadOnly=False,bDefaultValueIsIgnored=False,bAdvancedView=False,)
+		End Object
+		Begin Object Class=EdGraphNode_Comment Name="EdGraphNode_Comment_9"
+		   NodePosX=1008
+		   NodePosY=848
+		   NodeWidth=272
+		   NodeHeight=224
+		   NodeComment="Start Matinee"
+		   NodeGuid=DC7FA49B4840FAED60CD43B60C5F26C1
+		End Object
+
+			END OBJECT
+	[/PARAM]
+[/OBJECT]
+[/EXCERPT:StereoBP]
+-->
+
diff --git a/Engine/Documentation/Source/Platforms/VR/VR.INT.udn b/Engine/Documentation/Source/Platforms/VR/VR.INT.udn
--- a/Engine/Documentation/Source/Platforms/VR/VR.INT.udn
+++ b/Engine/Documentation/Source/Platforms/VR/VR.INT.udn
@@ -4,7 +4,7 @@ Crumbs:%ROOT%, Platforms
 Description:Information over developing for virtual reality platforms.
 Navigation:platform
 Version: 4.9
-type:landing
+Type:landing
 parent:Platforms
 order:2
 
@@ -50,18 +50,18 @@ vr
 
 ## Starting Out
 
-[DIR(output:"listbutton" type:"quickstart" tags:"vr")]
+[DIR(output:"listbutton" type:"quick start" tags:"vr")]
 
 ## Virtual Reality Platforms
 
-[DIR(end:"1" output:"listbutton" type:"landing" parent:"Platforms/VR")]
+[DIR(end:"1" output:"listbutton" type:"landing" tags:"platform" parent:"Platforms/VR")]
 
 ## Virtual Reality Topics
 
-[DIR(end:"1" output:"listbutton" type:"overview" parent:"Platforms/VR")]
+[DIR(end:"1" output:"listbutton" parent:"Platforms/VR" tags:"!platform")]
 
 ## Training Streams
 
 [REGION:topics third]
 [DIR(output:"video" type:"video" tags:"Virtual Reality")]
-[/REGION]
\ No newline at end of file
+[/REGION]
diff --git a/Engine/Documentation/Source/Platforms/iOS/iOS.INT.udn b/Engine/Documentation/Source/Platforms/iOS/iOS.INT.udn
--- a/Engine/Documentation/Source/Platforms/iOS/iOS.INT.udn
+++ b/Engine/Documentation/Source/Platforms/iOS/iOS.INT.udn
@@ -20,7 +20,7 @@ tags:iOS
 	[PARAM:description]
 	[/PARAM]
 	[PARAM:links]
-		[DIR(output:"topiccompact" parent:"Platforms/iOS" type:"quickstart")]
+		[DIR(output:"topiccompact" parent:"Platforms/iOS" type:"quick start")]
 	[/PARAM]
 [/OBJECT]
 
@@ -34,7 +34,7 @@ tags:iOS
 	[PARAM:description]
 	[/PARAM]
 	[PARAM:links]
-		[DIR(end:"1" output:"listbutton" type:"howto" parent:"Platforms/iOS")]
+		[DIR(end:"1" output:"listbutton" type:"how-to" parent:"Platforms/iOS")]
 	[/PARAM]
 [/OBJECT]
 
diff --git a/Engine/Documentation/Source/Programming/Development/BuildGraph/BuildGraph.INT.udn b/Engine/Documentation/Source/Programming/Development/BuildGraph/BuildGraph.INT.udn
--- a/Engine/Documentation/Source/Programming/Development/BuildGraph/BuildGraph.INT.udn
+++ b/Engine/Documentation/Source/Programming/Development/BuildGraph/BuildGraph.INT.udn
@@ -1,4 +1,4 @@
-Availability: Docs
+Availability: Public
 Title: BuildGraph
 Crumbs: %ROOT%, Programming, Programming/Development
 Description: Customize your Installed Builds by using our BuildGraph scripting system.
@@ -68,12 +68,7 @@ The syntax of elements used to define BuildGraph constructs is listed in the fol
 
 ## BuildGraph Scripting Reference
 
-[REGION:topics third]
-%Programming/Development/BuildGraph/ScriptAnatomy/Types:topiccompact%
-%Programming/Development/BuildGraph/ScriptAnatomy/Elements:topiccompact%
-%Programming/Development/BuildGraph/ScriptAnatomy/Conditions:topiccompact%
-%Programming/Development/BuildGraph/ScriptAnatomy/Tasks:topiccompact%
-[/REGION]
+[DIR(output:"listbutton" parent:"Programming/Development/BuildGraph" org:"flat" type:"reference")]
 
 ## Using BuildGraph
 
diff --git a/Engine/Documentation/Source/Programming/Development/BuildGraph/ScriptAnatomy/Conditions/BuildGraphScriptConditions.INT.udn b/Engine/Documentation/Source/Programming/Development/BuildGraph/ScriptAnatomy/Conditions/BuildGraphScriptConditions.INT.udn
--- a/Engine/Documentation/Source/Programming/Development/BuildGraph/ScriptAnatomy/Conditions/BuildGraphScriptConditions.INT.udn
+++ b/Engine/Documentation/Source/Programming/Development/BuildGraph/ScriptAnatomy/Conditions/BuildGraphScriptConditions.INT.udn
@@ -1,11 +1,10 @@
-Availability: Docs
+Availability: Public
 Title: BuildGraph Script Conditions
-Crumbs: %ROOT%, Programming, Programming/Development, Programming/Development/BuildGraph, Programming/Development/BuildGraph/ScriptAnatomy
+Crumbs: 
 Description:Learn the syntax needed to write BuildGraph script conditions.
 version: 4.13
 parent:Programming/Development/BuildGraph
 type:Reference
-related: Programming/Development/BuildGraph/ScriptAnatomy
 tags:BuildGraph
 
 [VAR:TopicCompact]
diff --git a/Engine/Documentation/Source/Programming/Development/BuildGraph/ScriptAnatomy/Elements/BuildGraphScriptElements.INT.udn b/Engine/Documentation/Source/Programming/Development/BuildGraph/ScriptAnatomy/Elements/BuildGraphScriptElements.INT.udn
--- a/Engine/Documentation/Source/Programming/Development/BuildGraph/ScriptAnatomy/Elements/BuildGraphScriptElements.INT.udn
+++ b/Engine/Documentation/Source/Programming/Development/BuildGraph/ScriptAnatomy/Elements/BuildGraphScriptElements.INT.udn
@@ -1,11 +1,10 @@
-Availability: Docs
+Availability: Public
 Title: BuildGraph Script Elements
-Crumbs: %ROOT%, Programming, Programming/Development, Programming/Development/BuildGraph, Programming/Development/BuildGraph/ScriptAnatomy
+Crumbs:
 Description:Learn about the elements of a BuildGraph.
 version: 4.13
 parent:Programming/Development/BuildGraph
 type:Reference
-related: Programming/Development/BuildGraph/ScriptAnatomy
 tags:BuildGraph
 
 [VAR:TopicCompact]
@@ -28,7 +27,7 @@ tags:BuildGraph
 [/OBJECT]
 [/VAR]
 
-[TOC(start:1 end:3)]
+[TOC(start:1 end:2)]
 
 [EXCERPT:BuildGraphScriptElements]
 
@@ -54,23 +53,23 @@ A BuildGraph script is typically defined with the following elements:
 `<Node>` is the smallest unit of execution in BuildGraph, having a set of inputs and outputs. Each `<Node>` consists of a 
 sequence of tasks that are executed in order.
 
-|                  |              |              |                                                                                                                                             |
-| ---------------- | ------------ | ------------ | ------------------------------------------------------------------------------------------------------------------------------------------- |
-| Name             | Name         | Required     | Name of the node.                                                                                                                           |
-| Requires         | Target List  | Optional     | List of nodes, aggregates, or tagged file sets produced by other nodes that this node requires to execute, separated by semicolons.         |
-| Produces         | Tag List     | Optional     | Tagged file sets that this node makes available to other nodes, separated by semicolons.                                                    |
-| After            | Target List  | Optional     | List of nodes that this node should run after, if they are part of the current target (ignored if they are not). Separated by semicolons.   |
-| NotifyOnWarnings | Boolean      | Optional     | If false, this node will not produce notifications on warnings. Only used by build systems. Defaults to `true`. `[META]`                  |
-| If               | Condition    | Optional     | Condition to be evaluated. The element is ignored unless the condition evaluates to `true`.                                                |
+| Attribute | Type | Required? | Description |
+| ---------------- | ------------ | ------------ | ---------------- |
+| Name             | Name         | Required     | Name of the node. |
+| Requires         | Target List  | Optional     | List of nodes, aggregates, or tagged file sets produced by other nodes that this node requires to execute, separated by semicolons. |
+| Produces         | Tag List     | Optional     | Tagged file sets that this node makes available to other nodes, separated by semicolons. |
+| After            | Target List  | Optional     | List of nodes that this node should run after, if they are part of the current target (ignored if they are not). Separated by semicolons. |
+| NotifyOnWarnings | Boolean      | Optional     | If false, this node will not produce notifications on warnings. Only used by build systems. Defaults to `true`. `[META]` |
+| If               | Condition    | Optional     | Condition to be evaluated. The element is ignored unless the condition evaluates to `true`. |
 
 ### Aggregate
 
 `<Aggregate>` declares a named aggregate that can be used as a synonym for a set of other nodes (or produced tag sets).
 
-|                  |              |              |                                                                                              |
-| ---------------- | ------------ | ------------ | -------------------------------------------------------------------------------------------- |
-| Name             | Name         | Required     | Name of the aggregate.                                                                       |
-| Requires         | Target List  | Required     | List of dependencies for this aggregate. May be nodes, tagged file sets, or agent groups.    |
+| Attribute | Type | Required? | Description |                                                                                            |
+| ---------------- | ------------ | ------------ | -------------- |
+| Name             | Name         | Required     | Name of the aggregate. |
+| Requires         | Target List  | Required     | List of dependencies for this aggregate. May be nodes, tagged file sets, or agent groups. |
 | If               | Condition    | Optional     | Condition to be evaluated. The element is ignored unless the condition evaluates to `true`. |
 
 ### Agent
@@ -78,20 +77,20 @@ sequence of tasks that are executed in order.
 `<Agent>` defines requirements for an agent on nodes that will be executed in sequence (without cleaning intermediate directories). 
 `<Agent>` requirements are ignored when building locally, but must be specified.
 
-|                  |                 |              |                                                                                                                                             |
-| ---------------- | --------------- | ------------ | ------------------------------------------------------------------------------------------------------------------------------------------  |
-| Name             | Name            | Required     | Name of the group.                                                                                                                          |
+| Attribute | Type | Required? | Description |
+| ---------------- | --------------- | ------------ | -----------  |
+| Name             | Name            | Required     | Name of the group. |
 | Type             | Identifier List | Optional     | Type of agent to run on. The meaning of this string is inferred by the host build system; it does not have any intrinsic meaning. `[META]` |
-| If               | Condition       | Optional     | Condition to be evaluated. The element is ignored unless the condition evaluates to `true`.                                                |
+| If               | Condition       | Optional     | Condition to be evaluated. The element is ignored unless the condition evaluates to `true`. |
 
 ### Trigger
 
 `<Trigger>` serves as a container for portions of the graph that should only be executed with explicit user intervention. 
 To execute nodes behind a trigger, pass `-Trigger=<Name>` to the command line. 
 
-|                  |                 |              |                                                                                             |
-| ---------------- | --------------- | ------------ | ------------------------------------------------------------------------------------------- |
-| Name             | Name            | Required     | Name of the trigger.                                                                        |
+| Attribute | Type | Required? | Description |
+| ---------------- | --------------- | ------------ | ------------------ |
+| Name             | Name            | Required     | Name of the trigger. |
 | If               | Condition       | Optional     | Condition to be evaluated. The element is ignored unless the condition evaluates to `true`. |
 
 ## Property Manipulation
@@ -107,22 +106,22 @@ BuildGraph properties can be modified with the following elements:
 `<Property>` sets the value of a property. If a property in an outer scope has already been declared with the same name,
 `<Property>` overwrites it. Otherwise, a new property is declared at the current scope.
 
-|                  |            |          |                                                                                             |
-| ---------------- | ---------- | -------- | ------------------------------------------------------------------------------------------- |
-| Name             | Name       | Required | Name of the property to set.                                                                |
-| Value            | String     | Required | New value for the property.                                                                 |
+| Attribute | Type | Required? | Description |
+| ---------------- | ---------- | -------- | -------------- |
+| Name             | Name       | Required | Name of the property to set. |
+| Value            | String     | Required | New value for the property. |
 | If               | Condition  | Optional | Condition to be evaluated. The element is ignored unless the condition evaluates to `true`. |
 
 ### Option
 
 `<Option>` defines a user option that can be set from the command line. `<Option>` may only appear at a global scope.
 
-|              |           |              |                                                                                             |
-| ------------ | --------- | ------------ | ------------------------------------------------------------------------------------------- |
-| Name         | Name      | Required     | Name of the option (and property) to initialize with the option's value.                    |
-| Description  | String    | Required     | Description for the option to display when running BuildGraph with the -ListOnly argument.  |
-| Restrict     | Regex     | Optional     | Regex that matches valid values for this option (eg. `[a-zA-Z]+`, `true | false`).          |
-| DefaultValue | String    | Required     | Default value for the option if the user does not set it explicitly.                        |
+| Attribute | Type | Required? | Description |
+| ------------ | --------- | ------------ | ------------- |
+| Name         | Name      | Required     | Name of the option (and property) to initialize with the option's value. |
+| Description  | String    | Required     | Description for the option to display when running BuildGraph with the -ListOnly argument. |
+| Restrict     | Regex     | Optional     | Regex that matches valid values for this option (eg. `[a-zA-Z]+`, `true | false`). |
+| DefaultValue | String    | Required     | Default value for the option if the user does not set it explicitly. |
 | If           | Condition | Optional     | Condition to be evaluated. The element is ignored unless the condition evaluates to `true`. |
 
 ### EnvVar
@@ -130,9 +129,9 @@ BuildGraph properties can be modified with the following elements:
 `<EnvVar>` declares a property to contain the contents of an environment variable (or an empty string if it's not set).
 `<EnvVar>` may only appear at a global scope.
 
-|                  |             |          |                                                                                             |
-| ---------------- | ----------- | -------- | ------------------------------------------------------------------------------------------- |
-| Name             | Name        | Required | Name of the environment variable to introduce as a property.                                |
+| Attribute | Type | Required? | Description |
+| ------------ | --------- | ------------ | ------------- |
+| Name             | Name        | Required | Name of the environment variable to introduce as a property. |
 | If               | Condition   | Optional | Condition to be evaluated. The element is ignored unless the condition evaluates to `true`. |
 
 ## Flow Control
@@ -149,17 +148,17 @@ You can control the flow of your BuildGraph scripts with the following elements:
 `<Include>` processes the contents of another BuildGraph script as if it appeared within this file.
 `<Include>` may only appear at a global scope.
 
-|                  |                 |              |                                                                                             |
-| ---------------- | --------------- | ------------ | ------------------------------------------------------------------------------------------- |
-| Script           | String          | Required     | Path to the script you wan to include (relative to the current script).                     |
+| Attribute | Type | Required? | Description |
+| ------------ | --------- | ------------ | ------------- |
+| Script           | String          | Required     | Path to the script you wan to include (relative to the current script). |
 | If               | Condition       | Optional     | Condition to be evaluated. The element is ignored unless the condition evaluates to `true`. |
 
 ### Do
 
 `<Do>` groups a sequence of elements and only processes them if a condition evaluates as true.
 
-|    |           |          |                                                                                             |
-| -- | ----------| ---------| ------------------------------------------------------------------------------------------- |
+| Attribute | Type | Required? | Description |
+| ------------ | --------- | ------------ | ------------- |
 | If | Condition | Optional | Condition to be evaluated. The element is ignored unless the condition evaluates to `true`. |
 
 ### Switch
@@ -180,9 +179,9 @@ Example `<Switch>` statement:
 			<Elements go here>
 		</Default>
 	</Switch>
-
-|    |           |          |                                                                                             |
-| -- | ----------| ---------| ------------------------------------------------------------------------------------------- |
+	
+| Attribute | Type | Required? | Description |
+| ------------ | --------- | ------------ | ------------- |
 | If | Condition | Optional | Condition to be evaluated. The element is ignored unless the condition evaluates to `true`. |
 
 ### ForEach
@@ -195,11 +194,11 @@ Example `<ForEach>` statement:
 	<ForEach Name="Counter" Values="1;2;3;4;5">
 		<Log Message="Counter=$(Counter)"/>
 	</ForEach>
-
-|                  |                 |              |                                                                                             |
-| ---------------- | --------------- | ------------ | ------------------------------------------------------------------------------------------- |
-| Name             | Name            | Required     | Property to assign to each value in the list.                                               |
-| Values           | String List     | Required     | List of values that are separated by semicolons.                                            |
+	
+| Attribute | Type | Required? | Description |
+| ------------ | --------- | ------------ | ------------- |
+| Name             | Name            | Required     | Property to assign to each value in the list. |
+| Values           | String List     | Required     | List of values that are separated by semicolons. |
 | If               | Condition       | Optional     | Condition to be evaluated. The element is ignored unless the condition evaluates to `true`. |
 
 ## Diagnostics
@@ -215,9 +214,9 @@ You can perform diagnostics on your BuildGraph scripts with the following elemen
 at global scope. `<Warning>` will only provide an output if it's still part of the graph after it has been trimmed down to 
 the target being executed.
 
-|                  |                 |              |                                                                                             |
-| ---------------- | --------------- | ------------ | ------------------------------------------------------------------------------------------- |
-| Message          | String          | Required     | Text to be printed to the log.                                                              |
+| Attribute | Type | Required? | Description |
+| ------------ | --------- | ------------ | ------------- |
+| Message          | String          | Required     | Text to be printed to the log. |
 | If               | Condition       | Optional     | Condition to be evaluated. The element is ignored unless the condition evaluates to `true`. |
 
 ### Error
@@ -225,9 +224,9 @@ the target being executed.
 `<Error>`outputs an error message. `<Error>` may be included in nodes, agents, triggers, or at global scope. 
 `<Error>` will only provide an output if it's still part of the graph after it has been trimmed down to the target being executed.
 
-|                  |                 |              |                                                                                             |
-| ---------------- | --------------- | ------------ | ------------------------------------------------------------------------------------------- |
-| Message          | String          | Required     | Text to be printed to the log.                                                              |
+| Attribute | Type | Required? | Description |
+| ------------ | --------- | ------------ | ------------- |
+| Message          | String          | Required     | Text to be printed to the log. |
 | If               | Condition       | Optional     | Condition to be evaluated. The element is ignored unless the condition evaluates to `true`. |
 
 [/EXCERPT:BuildGraphScriptElements]
diff --git a/Engine/Documentation/Source/Programming/Development/BuildGraph/ScriptAnatomy/Types/BuildGraphScriptTypes.INT.udn b/Engine/Documentation/Source/Programming/Development/BuildGraph/ScriptAnatomy/Types/BuildGraphScriptTypes.INT.udn
--- a/Engine/Documentation/Source/Programming/Development/BuildGraph/ScriptAnatomy/Types/BuildGraphScriptTypes.INT.udn
+++ b/Engine/Documentation/Source/Programming/Development/BuildGraph/ScriptAnatomy/Types/BuildGraphScriptTypes.INT.udn
@@ -1,11 +1,10 @@
-Availability: Docs
+Availability: Public
 Title: BuildGraph Script Types
-Crumbs: %ROOT%, Programming, Programming/Development, Programming/Development/BuildGraph, Programming/Development/BuildGraph/ScriptAnatomy
+Crumbs:
 Description:Learn about valid data types for BuildGraph attributes.
 version: 4.13
 parent:Programming/Development/BuildGraph
 type:Reference
-related: Programming/Development/BuildGraph/ScriptAnatomy
 tags:BuildGraph
 
 [VAR:TopicCompact]
diff --git a/Engine/Documentation/Source/Programming/Development/BuildGraph/Usage/BuildGraphUsage.INT.udn b/Engine/Documentation/Source/Programming/Development/BuildGraph/Usage/BuildGraphUsage.INT.udn
--- a/Engine/Documentation/Source/Programming/Development/BuildGraph/Usage/BuildGraphUsage.INT.udn
+++ b/Engine/Documentation/Source/Programming/Development/BuildGraph/Usage/BuildGraphUsage.INT.udn
@@ -1,11 +1,10 @@
-Availability: Docs
+Availability: Public
 Title: BuildGraph Usage
 Crumbs: %ROOT%, Programming, Programming/Development, Programming/Development/BuildGraph
 Description:This page shows you how to use BuildGraph Scripts.
 version: 4.13
 parent:Programming/Development/BuildGraph
 type:Overview
-related: Programming/Development/BuildGraph/ScriptAnatomy
 related: Programming/Development/InstalledBuildReference
 tags:BuildGraph
 
diff --git a/Engine/Documentation/Source/Programming/Development/Development.INT.udn b/Engine/Documentation/Source/Programming/Development/Development.INT.udn
--- a/Engine/Documentation/Source/Programming/Development/Development.INT.udn
+++ b/Engine/Documentation/Source/Programming/Development/Development.INT.udn
@@ -4,46 +4,32 @@ Crumbs: %ROOT%, Programming
 Description:Setting up your code environment, creating code projects, and working with Unreal Engine source.
 version: 4.9
 
-[OBJECT:TopicButtonList]
-	[PARAM:icon]
-	[/PARAM]
-	[PARAM:title]
-		General Developer Information
-	[/PARAM]
-	[PARAM:description]
-	[/PARAM]
-	[PARAM:links]
-		* [](Programming/Development/VisualStudioSetup "%Programming/Development/VisualStudioSetup:description%")
-		* [](Programming/Development/Tools "%Programming/Development/Tools:description%")
-		* [](Programming/Development/CodingStandard "%Programming/Development/CodingStandard:description%")
-	[/PARAM]
-[/OBJECT]
-
-[OBJECT:TopicButtonList]
-	[PARAM:icon]
-	[/PARAM]
-	[PARAM:title]
-		Working With Code Projects
-	[/PARAM]
-	[PARAM:description]
-	[/PARAM]
-	[PARAM:links]
-		* [](Programming/Development/ManagingGameCode "%Programming/Development/ManagingGameCode:description%")
-			* [](Programming/Development/ManagingGameCode/CppClassWizard "%Programming/Development/ManagingGameCode/CppClassWizard:description%")
-		* [](Programming/Development/CompilingProjects "%Programming/Development/CompilingProjects:description%")
-	[/PARAM]
-[/OBJECT]
-
-[OBJECT:TopicButtonList]
-	[PARAM:icon]
-	[/PARAM]
-	[PARAM:title]
-		Working With Engine Source
-	[/PARAM]
-	[PARAM:description]
-	[/PARAM]
-	[PARAM:links]
-		* [](GettingStarted/DownloadingUnrealEngine "%GettingStarted/DownloadingUnrealEngine:description%")
-		* [](Programming/Development/BuildingUnrealEngine "%Programming/Development/BuildingUnrealEngine:description%")
-	[/PARAM]
-[/OBJECT]
+## General Developer Information
+
+[REGION:buttonlist]
+* [](Programming/Development/VisualStudioSetup "%Programming/Development/VisualStudioSetup:description%")
+* [](Programming/Development/Tools "%Programming/Development/Tools:description%")
+* [](Programming/Development/CodingStandard "%Programming/Development/CodingStandard:description%")
+[/REGION]
+
+## Working With Code Projects
+
+[REGION:buttonlist]
+* [](Programming/Development/ManagingGameCode "%Programming/Development/ManagingGameCode:description%")
+	* [](Programming/Development/ManagingGameCode/CppClassWizard "%Programming/Development/ManagingGameCode/CppClassWizard:description%")
+* [](Programming/Development/CompilingProjects "%Programming/Development/CompilingProjects:description%")
+[/REGION]
+
+## Working With Engine Source
+
+[REGION:buttonlist]
+* [](GettingStarted/DownloadingUnrealEngine "%GettingStarted/DownloadingUnrealEngine:description%")
+* [](Programming/Development/BuildingUnrealEngine "%Programming/Development/BuildingUnrealEngine:description%")
+[/REGION]
+
+## Creating Engine Builds
+
+[REGION:buttonlist]
+* [](Programming/Development/BuildGraph "%Programming/Development/BuildGraph:description%")
+* [](Programming/Development/InstalledBuildReference "%Programming/Development/InstalledBuildReference:description%")
+[/REGION]
\ No newline at end of file
