Copying //UE4/Dev-Documentation to Samples-Main (//UE4/Samples-Main)

diff --git a/Engine/Documentation/Source/Programming/Development/InstalledBuildReference/InstalledBuildReference.INT.udn b/Engine/Documentation/Source/Programming/Development/InstalledBuildReference/InstalledBuildReference.INT.udn
--- a/Engine/Documentation/Source/Programming/Development/InstalledBuildReference/InstalledBuildReference.INT.udn
+++ b/Engine/Documentation/Source/Programming/Development/InstalledBuildReference/InstalledBuildReference.INT.udn
@@ -1,4 +1,4 @@
-Availability: Docs
+Availability: Public
 Title: Installed Build Reference
 Crumbs: %ROOT%, Programming, Programming/Development
 Description:This page describes the Installed Build process, including an overview of how to write Installed Build scripts.
@@ -33,163 +33,60 @@ tags: BuildGraph
 [/VAR]
 
 Developers are always looking for ways to efficiently deliver build solutions for a variety of hardware profiles and operating systems. 
-With the release of our BuildGraph scripting system, you'll be able to create custom sequences of Unreal Engine 4 (UE4) build commands, 
-harnessing the Installed Build process that we've developed for your use. As you read through the following reference page, you'll learn 
-about the Installed Build process, including how to run scripts that make Installed Builds for targeted build solutions.
+An **Installed Build** is a fully featured engine build that can be redistributed to quickly get your team up and running with the engine. 
+From a technical standpoint, an Installed build contains precompiled binaries for Unreal Editor and a static library for each target platform 
+in development and shipping configurations. 
 
-## What's an Installed Build?
+The Installed Build process automates build packages and tools that are needed to make an Installed Build version of Unreal Engine 4 (UE4). In summary, 
+making an Installed Build is the process of setting up default target platforms, compiling UE4 with its tools and editor, running tests, 
+and deploying the Installed Build version of UE4 to targeted platforms. 
 
-An **Installed Build** of UE4 is the same as a version that's distributed using the [Epic Games Launcher](GettingStarted/Installation).
-
-![](InstalledBuilds.png)
+## Making an Installed Build
 
-This version of the engine is shipped with each release of UE4. Although the Installed Build is fully featured, its features are streamlined by an 
-Installed Build process that quickly gets users up and running with the engine. Additionally, the Unreal Editor is distributed as 
-precompiled binaries with a static library for each target platform in development and shipping configurations.
+You can create Installed Builds using the [](Programming/Development/BuildGraph) scripting system with the Installed Build script located at 
+`[UE4Root]/Engine/Build/InstalledBuild.xml`.
 
-[REGION:note]
-Contents found in the `Engine` folder should be treated as `read-only` because rebuilding game projects will not attempt to rebuild Engine code.
-[/REGION] 
+**To create an Installed Build:**
 
-## About the Installed Build Process
+1.	Run the Installed Build Script by invoking the **AutomationTool** with the following command line, replacing `[PLATFORM]` with either `Win64` or `Mac`. You can also configure the build using the available [Installed Build Script Options](#bookmark1).
 
-The Installed Build process automates build packages and tools that are needed to make an Installed Build version of UE4. 
-In summary, making an Installed Build is the process of setting up default target platforms, compiling UE4 with its tools and 
-editor, running tests, and deploying the Installed Build version of UE4 for targeted platforms.
+		BuildGraph -target="Make Installed Build [PLATFORM]" -script=Engine/Build/InstalledEngineBuild.xml -clean
 
-## Making an Installed Build
+1.	Navigate to the `LocalBuilds/Engine/` folder in your Unreal Engine directory to view the Installed Build that you created. If you specified a different directory to publish to, navigate to that directory instead.
 
-The Installed Build process is made with our BuildGraph scripting system. BuildGraph scripts create custom sequences 
-of UE4 build commands for build automation. You can declare graphs by scripting an XML file with syntax (using [XML Schema v1.0](https://www.w3.org/TR/NOTE-dcd])) 
-similar to MSBuild, Apache ANT, and NAnt. Scripted XML files define a network of node dependencies containing tasks that will be executed during the Installed Build process.
+## Installed Build Script Options
 
-If you're interested in reading through example BuildGraph scripts, navigate to the `[UE4Root]/Engine/Build/Graph/Examples` folder, where you'll find: 
+(#bookmark1)
 
-* `AllExamples.xml`
-* `Building.xml`
-* `TagsAndFiles.xml`
+The `InstalledEngineBuild.xml` script makes an Installed Build with all of the default options and platforms enabled; however, it can be configured using a set 
+of specified options, enabling you to target which platforms to include, where to publish the build to, etc. You can view a list of available options for customizing the 
+build process (along with a list of nodes that will be built), by passing the `-listonly` option to the script. 
 
-### Running the Installed Build Script
+The available options are: 
 
-The Installed Build script is located at `[UE4Root]/Engine/Build/InstalledEngineBuild.xml`. You can run the Installed Build script by invoking the 
-**AutomationTool** with the following command line:
 
-%Globals:osselect%
-
-[OBJECT:ToggleButtonContent]
-	[PARAMLITERAL:category]
-	OS
-	[/PARAMLITERAL]
-	[PARAMLITERAL:id]
-	windows
-	[/PARAMLITERAL]
-	[PARAMLITERAL:active]
-	active_button_content
-	[/PARAMLITERAL]
-	[PARAM:content]
-	`BuildGraph -target=”Make Installed Build Win64” -script=Engine/Build/InstalledEngineBuild.xml -clean`
-	[/PARAM]
-[/OBJECT]
-
-[OBJECT:ToggleButtonContent]
-	[PARAMLITERAL:category]
-	OS
-	[/PARAMLITERAL]
-	[PARAMLITERAL:id]
-	mac
-	[/PARAMLITERAL]
-	[PARAMLITERAL:active]
-	active_button_content
-	[/PARAMLITERAL]
-	[PARAM:content]
-	`BuildGraph -target=”Make Installed Build Mac” -script=Engine/Build/InstalledEngineBuild.xml -clean`
-	[/PARAM]
-[/OBJECT]
+| Option | Default | Description |
+| ------ | ------- | ----------- |
+| `-set:HostPlatformOnly=[true/false]` | false | A helper option to make an installed build for your host platform only, so that you don't have to disable each platform individually |
+| `-set:WithWin64=[true/false]` | true | Include the Win64 target platform |
+| `-set:WithWin32=[true/false]` | true | Include the Win32 target platform |
+| `-set:WithMac=[true/false]` | true |  Include the Mac target platform |
+| `-set:WithAndroid=[true/false]` | true | Include the Android target platform |
+| `-set:WithIOS=[true/false]` | true | Include the iOS target platform [REGION:note]Enabling this option requires a Mac set up for remote building.[/REGION] |
+| `-set:WithTVOS=[true/false]` | true | Include the tvOS target platform [REGION:note]Enabling this option requires a Mac set up for remote building.[/REGION] |
+| `-set:WithLinux=[true/false]` | true | Include the Linux target platform |
+| `-set:WithHTML5=[true/false]` | true | Include the HTML5 target platform |
+| `-set:WithPS4=[true/false]` | false | Include the PS4 target platform [REGION:note]Enabling this option requires the PS4 SDK.[/REGION] |
+| `-set:WithXboxOne=[true/false]` | false | Include the XboxOne target platform [REGION:note]Enabling this option requires the Xbox One SDK.[/REGION] |
+| `-set:WithDDC=[true/false]` | true | Build a stand-alone derived-data cache for the engine content and templates [REGION:note]Building a stand-alone [Derived Data Cache (DDC)](Engine/Basics/DerivedDataCache) for the Engine and Template content can be one of the slowest aspects of the build. If you don't need a stand-alone DDC, you can skip this step by appending `-set:WithDDC=false` to the command line.[/REGION] |
+| `-set:SignExecutables=[true/false]` | false | Sign the executables produced if you have to set up your machine to sign executables as part of the build process |
+| `-set:AnalyticsTypeOverride=[ID]` | N/A | Identifier for analytic events to send |
 
-In summary, running `InstalledEngineBuild.xml` makes an Installed Build with all of the default options and platforms enabled. 
-
-#### Installed Build Script Options
-
-If you want to see a list of available options for customizing the build process (along with a list of nodes that'll be built), 
-you'll want to append the `-listonly` option before running the script. Subsequently, the list output should look something like this: 
-
-	Options:
-	
-	    -set:PublishDir=...            Root directory to contain published builds and symbols (Default: $(RootDir)/LocalBuilds/Engine) 
-	    -set:WithWin64=...             Include the Win64 target platform (Default: true) 
-	    -set:WithWin32=...             Include the Win32 target platform (Default: true) 
-	    -set:WithMac=...               Include the Mac target platform (Default: true) 
-	    -set:WithAndroid=...           Include the Android target platform (Default: true) 
-	    -set:WithIOS=...               Include the iOS target platform (Default: true) 
-	    -set:WithTVOS=...              Include the tvOS target platform (Default: true) 
-	    -set:WithLinux=...             Include the Linux target platform (Default: true) 
-	    -set:WithHTML5=...             Include the HTML5 target platform (Default: true) 
-	    -set:WithPS4=...               Include the PS4 target platform (Default: false) 
-	    -set:WithXboxOne=...           Include the XboxOne target platform (Default: false) 
-	    -set:WithDDC=...               Build a stand-alone derived-data cache for the engine content and templates (Default: true) 
-	    -set:AnalyticsTypeOverride=... Identifier for analytic events to send 
-	
-	Graph:
-	    Trigger: None
-	        Agent: Editor Win64 (CompileWin64)
-	            Node: Update Version Files
-	            Node: Compile UnrealHeaderTool Win64
-	            Node: Compile UE4Editor Win64
-	        Agent: Editor Mac (CompileMac)
-	            Node: Compile UnrealHeaderTool Mac
-	        Agent: Target Platforms Win64 (CompileWin64)
-	            Node: Compile UE4Game Win64
-	            Node: Compile UE4Game Win32
-	            Node: Compile UE4Game Linux
-	            Node: Compile UE4Game Android
-	            Node: Compile UE4Game HTML5
-	        Agent: Target Platforms Mac (CompileMac)
-	            Node: Compile UE4Game IOS
-	            Node: Compile UE4Game TVOS
-	        Agent: Tools Group Win64 (CompileWin64)
-	            Node: Build Tools Win64
-	            Node: Build Tools CS
-	            Node: Make Feature Packs
-	        Agent: DDC Group Win64 (Win64)
-	            Node: Build DDC Win64
-	        Agent: Installed Build Group Win64 (Win64)
-	            Node: Make Installed Build Win64
-
-Please note that some portions of the node graph can be ignored because they're only relevant to Epic's internal build system. For example, 
-the Agent descriptions, namely `CompileWin64` and `CompileMac`, refer to the type of machine they're meant to run on; however, the script 
-can still be run from a single build machine. 
-
-### Some Additional Commands
-
-Depending on your targeted platform and available Software Development Kit (SDK), we're providing some useful commands that you should take note of.
-
-#### Making an Installed Build for your Current Platform
-
-If you want to make an Installed Build for your current platform, thereby avoiding the need to set all of the platform options to `false`, append 
-`-set:HostPlatformOnly=true` to the command line.
-
-##### IOS and TVOS
-
-If you want to include IOS and TVOS in your Installed Build, you'll have to have a Mac set up for remote building. Otherwise, append the following to your command line:
-`-set:WithIOS=false -setWithTVOS=false`.
-
-##### Playstation 4 and XboxOne
-
-If you have the required SDKs, you can add Playstation 4 (PS4) and XboxOne as target platforms (although, this is untested). To add PS4 and XboxOne 
-as target platforms, append `-set:WithPS4=true -set:WithXboxOne=true` to the command line.
-
-#### Excluding the Derived Data Cache
-
-Building a stand-alone [Derived Data Cache (DDC)](Engine/Basics/DerivedDataCache) for the Engine and Template content can be one of the slowest aspects of the build. 
-If you don't need a stand-alone DDC, you can skip this step by appending `-set:WithDDC=false` to the command line.
-
-#### Signing Executables
-
-If you have to set up your machine to sign executables as part of the build process, append `-set:SignExecutables=true` to the command line.
 
 ### Including Additional Files
 
 RuntimeDependencies (set in `build.cs` files for each module) automatically collate and determine which files are included in an Installed Build. 
 However, there are some required files that cannot be included this way, so they're defined in `[UE4Root]/Engine/Build/InstalledEngineFilters.xml`. 
-The `InstalledEngineFilters` XML file also lists patterns of files that should be excluded from the build, determining the types of files that need to be stripped or signed, and finally, deciding 
-what projects to build the DDC for. If you need to add any additional files to the Installed Build, the `InstalledEngineFilters` XML file is a great place to start.
\ No newline at end of file
+The `InstalledEngineFilters` XML file also lists patterns of files that should be excluded from the build, determining the types of files that need to be 
+stripped or signed, and finally, deciding what projects to build the DDC for. If you need to include any additional files with the Installed Build, the `InstalledEngineFilters` 
+XML file is a great place to start.
\ No newline at end of file
diff --git a/Engine/Documentation/Source/Programming/Rendering/FShaderCache/FShaderCache.INT.udn b/Engine/Documentation/Source/Programming/Rendering/FShaderCache/FShaderCache.INT.udn
--- a/Engine/Documentation/Source/Programming/Rendering/FShaderCache/FShaderCache.INT.udn
+++ b/Engine/Documentation/Source/Programming/Rendering/FShaderCache/FShaderCache.INT.udn
@@ -3,6 +3,7 @@ Title:FShaderCache
 Crumbs:%ROOT%, Programming, Programming/Rendering
 Description:The FShaderCache provides mechanisms for reducing shader hitching in-game.
 Version: 4.9
+tags:Rendering
 
 
 ## Overview
diff --git a/Engine/Documentation/Source/Programming/Rendering/MaterialExpressions/CreatingMaterialExpressions.INT.udn b/Engine/Documentation/Source/Programming/Rendering/MaterialExpressions/CreatingMaterialExpressions.INT.udn
--- a/Engine/Documentation/Source/Programming/Rendering/MaterialExpressions/CreatingMaterialExpressions.INT.udn
+++ b/Engine/Documentation/Source/Programming/Rendering/MaterialExpressions/CreatingMaterialExpressions.INT.udn
@@ -3,6 +3,7 @@ Title:Material Expressions Technical Guide
 Crumbs:%ROOT%, Programming, Programming/Rendering
 Description:Guide to creating new material expressions for use in the Material Editor.
 Version: 4.9
+tags:Rendering
 
 [TOC (start:2 end:3)]
 
diff --git a/Engine/Documentation/Source/Programming/Rendering/Programming.INT.udn b/Engine/Documentation/Source/Programming/Rendering/Programming.INT.udn
--- a/Engine/Documentation/Source/Programming/Rendering/Programming.INT.udn
+++ b/Engine/Documentation/Source/Programming/Rendering/Programming.INT.udn
@@ -4,6 +4,8 @@ Crumbs:%ROOT%, Programming
 Description:Information for graphics programmers working with the rendering systems and writing shaders.
 Version: 4.9
 Related:Engine/Basics/CoordinateSpace
+parent:Programming
+tags:Rendering
 
 The Renderer module in the engine manages and renders scenes, which hold the renderable information associated with each world. It includes the definition of all drawing policies and shader.
 The RHI module, which is the interface for rendering APIs, is another key module for graphics programming. The [](Programming/Rendering/Overview) covers many key classes, settings, and variables to investigate, while 
diff --git a/Engine/Documentation/Source/Programming/Tutorials/UMG/1/UMG_1.INT.udn b/Engine/Documentation/Source/Programming/Tutorials/UMG/1/UMG_1.INT.udn
--- a/Engine/Documentation/Source/Programming/Tutorials/UMG/1/UMG_1.INT.udn
+++ b/Engine/Documentation/Source/Programming/Tutorials/UMG/1/UMG_1.INT.udn
@@ -2,7 +2,7 @@ Availability: Public
 Title:1. Set Up Module Dependencies for UMG
 Crumbs:%ROOT%, Programming, Programming/Tutorials, Programming/Tutorials/UMG
 Description:Modifying our project to include the modules needed to build menus.
-version:4.9
+version:4.14
 
 [VAR:Steps]
 [OBJECT:Navigation]
diff --git a/Engine/Documentation/Source/Programming/Tutorials/UMG/2/UMG_2.INT.udn b/Engine/Documentation/Source/Programming/Tutorials/UMG/2/UMG_2.INT.udn
--- a/Engine/Documentation/Source/Programming/Tutorials/UMG/2/UMG_2.INT.udn
+++ b/Engine/Documentation/Source/Programming/Tutorials/UMG/2/UMG_2.INT.udn
@@ -2,7 +2,7 @@ Availability: Public
 Title:2. Extend Our Game Mode
 Crumbs:%ROOT%, Programming, Programming/Tutorials, Programming/Tutorials/UMG
 Description:Extending the project's Game Mode to add and remove menus from the screen.
-version:4.9
+version:4.14
 
 [VAR:Steps]
 [OBJECT:Navigation]
@@ -77,6 +77,24 @@ version:4.9
 
 	[REGION:note]This code will create instances of any **Widgets** we provide and put them on the screen. It also removes them, so that only one is active at a time, although **Unreal Engine** can handle displaying and interacting with many Widgets at once. We never need to destroy a Widget directly, because removing it from the viewport and clearing (or changing) all variables that reference it will cause it to be cleaned up by Unreal Engine's garbage collection system.[/REGION]
 
+1. Finally, we need to set the input mode on our **Player Controller** class. To do this, we'll add a new C++ class to our project, based on **Player Controller**. Within this class, we only need to call one extra function when the game starts to ensure that we can interact with UI elements.
+
+	[![](NewClass.png)](NewClass.png)
+	[![](PlayerController.png)](PlayerController.png)
+	
+	In HowTo_UMGPlayerController.h, we will add the following override to the class:
+	
+		public:
+			virtual void BeginPlay() override;
+
+	In HowTo_UMGPlayerController.cpp, we will add our overridden function:
+	
+		void AHowTo_UMGPlayerController::BeginPlay()
+		{
+			Super::BeginPlay();
+			SetInputMode(FInputModeGameAndUI());
+		}
+
 
 ----
 
@@ -181,6 +199,39 @@ We have built the code framework to create and display menus, and remove them wh
 					}
 				}
 			}
+			
+		**HowTo_UMGPlayerController.h**
+			// Copyright 1998-2016 Epic Games, Inc. All Rights Reserved.
+
+			#pragma once
+
+			#include "GameFramework/PlayerController.h"
+			#include "HowTo_UMGPlayerController.generated.h"
+
+			/**
+			 * 
+			 */
+			UCLASS()
+			class HOWTO_UMG_API AHowTo_UMGPlayerController : public APlayerController
+			{
+				GENERATED_BODY()
+				
+			public:
+				virtual void BeginPlay() override;
+			};
+
+		**HowTo_UMGPlayerController.cpp**
+			// Copyright 1998-2016 Epic Games, Inc. All Rights Reserved.
+
+			#include "HowTo_UMG.h"
+			#include "HowTo_UMGPlayerController.h"
+
+			void AHowTo_UMGPlayerController::BeginPlay()
+			{
+				Super::BeginPlay();
+				SetInputMode(FInputModeGameAndUI());
+			}
+			
     [/PARAM]
 [/OBJECT]
 
diff --git a/Engine/Documentation/Source/Programming/Tutorials/UMG/3/UMG_3.INT.udn b/Engine/Documentation/Source/Programming/Tutorials/UMG/3/UMG_3.INT.udn
--- a/Engine/Documentation/Source/Programming/Tutorials/UMG/3/UMG_3.INT.udn
+++ b/Engine/Documentation/Source/Programming/Tutorials/UMG/3/UMG_3.INT.udn
@@ -2,7 +2,7 @@ Availability: Public
 Title:3. Create Menu Widget Blueprints
 Crumbs:%ROOT%, Programming, Programming/Tutorials, Programming/Tutorials/UMG
 Description:Create widget blueprints to be used as menu screens.
-version:4.9
+version:4.14
 
 [VAR:Steps]
 [OBJECT:Navigation]
@@ -177,7 +177,40 @@ With our main menu built, we can set up a GameMode asset that will load it as so
 					}
 				}
 			}
-    [/PARAM]
+
+		**HowTo_UMGPlayerController.h**
+			// Copyright 1998-2016 Epic Games, Inc. All Rights Reserved.
+
+			#pragma once
+
+			#include "GameFramework/PlayerController.h"
+			#include "HowTo_UMGPlayerController.generated.h"
+
+			/**
+			 * 
+			 */
+			UCLASS()
+			class HOWTO_UMG_API AHowTo_UMGPlayerController : public APlayerController
+			{
+				GENERATED_BODY()
+				
+			public:
+				virtual void BeginPlay() override;
+			};
+
+		**HowTo_UMGPlayerController.cpp**
+			// Copyright 1998-2016 Epic Games, Inc. All Rights Reserved.
+
+			#include "HowTo_UMG.h"
+			#include "HowTo_UMGPlayerController.h"
+
+			void AHowTo_UMGPlayerController::BeginPlay()
+			{
+				Super::BeginPlay();
+				SetInputMode(FInputModeGameAndUI());
+			}
+			
+	[/PARAM]
 [/OBJECT]
 
 %Steps%
diff --git a/Engine/Documentation/Source/Programming/Tutorials/UMG/4/UMG_4.INT.udn b/Engine/Documentation/Source/Programming/Tutorials/UMG/4/UMG_4.INT.udn
--- a/Engine/Documentation/Source/Programming/Tutorials/UMG/4/UMG_4.INT.udn
+++ b/Engine/Documentation/Source/Programming/Tutorials/UMG/4/UMG_4.INT.udn
@@ -2,7 +2,7 @@ Availability: Public
 Title:4. Configure Our Game Mode
 Crumbs:%ROOT%, Programming, Programming/Tutorials, Programming/Tutorials/UMG
 Description:Set up our Game Mode to use our Main Menu Widget.
-version:4.9
+version:4.14
 
 [VAR:Steps]
 [OBJECT:Navigation]
@@ -173,6 +173,39 @@ Our custom GameMode asset is now in effect on our level, configured to load our
 					}
 				}
 			}
+			
+		**HowTo_UMGPlayerController.h**
+			// Copyright 1998-2016 Epic Games, Inc. All Rights Reserved.
+
+			#pragma once
+
+			#include "GameFramework/PlayerController.h"
+			#include "HowTo_UMGPlayerController.generated.h"
+
+			/**
+			 * 
+			 */
+			UCLASS()
+			class HOWTO_UMG_API AHowTo_UMGPlayerController : public APlayerController
+			{
+				GENERATED_BODY()
+				
+			public:
+				virtual void BeginPlay() override;
+			};
+
+		**HowTo_UMGPlayerController.cpp**
+			// Copyright 1998-2016 Epic Games, Inc. All Rights Reserved.
+
+			#include "HowTo_UMG.h"
+			#include "HowTo_UMGPlayerController.h"
+
+			void AHowTo_UMGPlayerController::BeginPlay()
+			{
+				Super::BeginPlay();
+				SetInputMode(FInputModeGameAndUI());
+			}
+
     [/PARAM]
 [/OBJECT]
 
diff --git a/Engine/Documentation/Source/Programming/Tutorials/UMG/5/UMG_5.INT.udn b/Engine/Documentation/Source/Programming/Tutorials/UMG/5/UMG_5.INT.udn
--- a/Engine/Documentation/Source/Programming/Tutorials/UMG/5/UMG_5.INT.udn
+++ b/Engine/Documentation/Source/Programming/Tutorials/UMG/5/UMG_5.INT.udn
@@ -2,7 +2,7 @@ Availability: Public
 Title:5. Build A Second Menu
 Crumbs:%ROOT%, Programming, Programming/Tutorials, Programming/Tutorials/UMG
 Description:Build a second menu with new functionality.
-version:4.9
+version:4.14
 
 [VAR:Steps]
 [OBJECT:Navigation]
@@ -173,6 +173,39 @@ version:4.9
 					}
 				}
 			}
+			
+		**HowTo_UMGPlayerController.h**
+			// Copyright 1998-2016 Epic Games, Inc. All Rights Reserved.
+
+			#pragma once
+
+			#include "GameFramework/PlayerController.h"
+			#include "HowTo_UMGPlayerController.generated.h"
+
+			/**
+			 * 
+			 */
+			UCLASS()
+			class HOWTO_UMG_API AHowTo_UMGPlayerController : public APlayerController
+			{
+				GENERATED_BODY()
+				
+			public:
+				virtual void BeginPlay() override;
+			};
+
+		**HowTo_UMGPlayerController.cpp**
+			// Copyright 1998-2016 Epic Games, Inc. All Rights Reserved.
+
+			#include "HowTo_UMG.h"
+			#include "HowTo_UMGPlayerController.h"
+
+			void AHowTo_UMGPlayerController::BeginPlay()
+			{
+				Super::BeginPlay();
+				SetInputMode(FInputModeGameAndUI());
+			}
+
     [/PARAM]
 [/OBJECT]
 
diff --git a/Engine/Documentation/Source/Programming/Tutorials/UMG/6/UMG_6.INT.udn b/Engine/Documentation/Source/Programming/Tutorials/UMG/6/UMG_6.INT.udn
--- a/Engine/Documentation/Source/Programming/Tutorials/UMG/6/UMG_6.INT.udn
+++ b/Engine/Documentation/Source/Programming/Tutorials/UMG/6/UMG_6.INT.udn
@@ -2,7 +2,7 @@ Availability: Public
 Title:6. On Your Own!
 Crumbs:%ROOT%, Programming, Programming/Tutorials, Programming/Tutorials/UMG
 Description:Ideas for further exploration!
-version:4.9
+version:4.14
 
 [VAR:Steps]
 [OBJECT:Navigation]
diff --git a/Engine/Documentation/Source/Programming/UnrealArchitecture/Delegates/Dynamic/DynamicDelegates.INT.udn b/Engine/Documentation/Source/Programming/UnrealArchitecture/Delegates/Dynamic/DynamicDelegates.INT.udn
--- a/Engine/Documentation/Source/Programming/UnrealArchitecture/Delegates/Dynamic/DynamicDelegates.INT.udn
+++ b/Engine/Documentation/Source/Programming/UnrealArchitecture/Delegates/Dynamic/DynamicDelegates.INT.udn
@@ -27,12 +27,3 @@ variations specific to dynamic delegates.
 ## Executing Dynamic Delegates
 
 [INCLUDE:Programming/UnrealArchitecture/Delegates#ExecutingDelegates]
-
-## Wrapped Dynamic Delegates
-
-Dynamic delegates can be declared with a wrapper proxy method for executing the delegate.
-
-| Declaration Macro | Description |
-| --- | --- |
-| `DECLARE_DYNAMIC_DELEGATE[_RetVal, etc.]( DelegateName )` | Creates a wrapped dynamic delegate. |
-| `DECLARE_DYNAMIC_MULTICAST_DELEGATE[_RetVal, etc.]( DelegateName )` | Creates a wrapped dynamic multi-cast delegate. |
diff --git a/Engine/Documentation/Source/Programming/UnrealArchitecture/Reference/Interfaces/Interfaces.INT.udn b/Engine/Documentation/Source/Programming/UnrealArchitecture/Reference/Interfaces/Interfaces.INT.udn
--- a/Engine/Documentation/Source/Programming/UnrealArchitecture/Reference/Interfaces/Interfaces.INT.udn
+++ b/Engine/Documentation/Source/Programming/UnrealArchitecture/Reference/Interfaces/Interfaces.INT.udn
@@ -23,11 +23,12 @@ In your .h file (e.g. `ItemInterface.h`):
     UINTERFACE(Blueprintable)
     class UItemInterface : public UInterface
     {
-        GENERATED_UINTERFACE_BODY()
+        GENERATED_BODY()
     };
 
     class IItemInterface
-    {    GENERATED_IINTERFACE_BODY()
+    {    
+		GENERATED_BODY()
 
     public:
 
diff --git a/Engine/Documentation/Source/Programming/UnrealArchitecture/TMap/TMap.INT.udn b/Engine/Documentation/Source/Programming/UnrealArchitecture/TMap/TMap.INT.udn
--- a/Engine/Documentation/Source/Programming/UnrealArchitecture/TMap/TMap.INT.udn
+++ b/Engine/Documentation/Source/Programming/UnrealArchitecture/TMap/TMap.INT.udn
@@ -1,8 +1,11 @@
 Title:TMap
 Crumbs:
-Description:
+Description:TMaps are primarily defined by two types - a key type and a value type - which are stored as associated pairs in the map. 
 Availability:Public
-Version:4.9
+Version:4.14
+Related:Engine/UI/LevelEditor/Details/Properties/Map
+
+[TOC]
 
 After TArray, the most commonly used container in **Unreal Engine 4** (UE4) is **TMap**.  This container is an associative container, which means that every key has an associated value, and you can efficiently look up a value object by its key.
 
@@ -93,7 +96,13 @@ It's also possible to insert all the elements from another map by using the Appe
 	// 	{ Key: 9, Value: "Melon"     }
 	// ]
 
-Here, the resulting map is equivalent to using Add/Emplace to add them individually, so that duplicate keys from the source map replace ones in the target.
+In the above example, the resulting map is equivalent to using Add/Emplace to add them individually, so that duplicate keys from the source map replace ones in the target.
+
+
+If you mark the TMap with the `UPROPERTY()` macro and one of the editable keywords (`EditAnywhere`, `EditDefaultsOnly`, `EditInstanceOnly`), you can [add and edit elements in Unreal Editor](Engine/UI/LevelEditor/Details/Properties/Map).
+
+	UPROPERTY(Category = MapsAndSets, EditAnywhere)
+	TMap<int32, FString> FruitMap;
 
 ## Iteration
 
@@ -183,7 +192,9 @@ The **FindOrAdd** function will search for a given key and return a reference to
 	// 	{ Key: 8, Value: ""          }
 	// ]
 
+[REGION:note]
 Note that here the Ref7 reference may have been invalidated by the call to FruitMap.FindOrAdd(8) if a reallocation occurred.
+[/REGION]
 
 Despite its name, the **FindRef** function searches for a key and returns a value rather than a reference.  If the key was found, a copy of the associated value is returned, otherwise a default-constructed value type is returned.  This results in similar behavior to FindOrAdd but, because FindRef function returns a value rather than a reference, the map will not be modified and so can be called on a const object:
 
@@ -234,7 +245,9 @@ Elements can be removed from maps by using the **Remove** function and providing
 	// 	{ Key: 9, Value: "Melon"     }
 	// ]
 
+[REGION:note]
 Removing elements actually leaves holes in the data structure, which you can see when visualizing the map in Visual Studio's watch window, but they have been omitted here for clarity.
+[/REGION]
 
 The **FindAndRemoveChecked** function can be used to remove an element from the map and return its value.  The "checked" part of the name means that the key will be checked for existence and assert if it doesn't exist:
 
@@ -366,7 +379,9 @@ TMap does not provide a way of checking how many elements are preallocated, like
 	// 	{ Key: 0, Value: "Fruit0" }
 	// ]
 
+[REGION:note]
 Note that the slack happens to have caused the new elements to be added in reverse order. This is an example of why element order in maps should not be relied upon in any way.
+[/REGION]
 
 The **Shrink** function also works like the TArray equivalent in that it removes any wasted slack from the end of the container.  However, as TMap allows holes in its data structure, this will only remove slack from holes left at the end of the structure:
 
@@ -487,45 +502,45 @@ An example of a custom KeyFuncs might look like this:
 
 Here, we have a type which has a unique identifier as part of its state, but also has some other state that does not contribute to its identity.  GetTypeHash and operator== would be inappropriate here, as it would be a lie for operator== to ignore part of the state and GetTypeHash needs to match operator==, which it couldn't if operator== was defined properly.  However, for the purpose of identifying this type in a map, we are happy to just use the UniqueID part of the state.
 
-First, we inherit **BaseKeyFuncs** as it helpfully defines some stuff for us, including KeyInitType and ElementInitType.  We pull these directly from Super into the derived class so that we can use them in the rest of the implementation.
+1. First, we inherit **BaseKeyFuncs** as it helpfully defines some types for us, including KeyInitType and ElementInitType.  We pull these directly from Super into the derived class so that we can use them in the rest of the implementation.
 
-BaseKeyFuncs takes two template parameters: the element type of the map and the type of our key: the object being returned from GetSetKey.  As with all maps, the element type is a TPair, taking FMyStruct as its KeyType and TMyStructMapKeyFuncs's template parameter as its ValueType.  We make our replacement KeyFuncs a template to allow the ValueType to be specified on a per-map basis, rather than needing to define a new KeyFuncs every time we want to create a TMap keyed on FMyStruct.  The second BaseKeyFuncs argument is the type of the key, not to be confused with TPair's 'KeyType', which is what is being stored in the Key field of the element.  As we want to use FMyStruct::UniqueID as our key, we specify FString here.
+	BaseKeyFuncs takes two template parameters: the element type of the map and the type of our key: the object being returned from GetSetKey.  As with all maps, the element type is a TPair, taking FMyStruct as its KeyType and TMyStructMapKeyFuncs's template parameter as its ValueType.  We make our replacement KeyFuncs a template to allow the ValueType to be specified on a per-map basis, rather than needing to define a new KeyFuncs every time we want to create a TMap keyed on FMyStruct.  The second BaseKeyFuncs argument is the type of the key, not to be confused with TPair's 'KeyType', which is what is being stored in the Key field of the element.  As we want to use FMyStruct::UniqueID as our key, we specify FString here.
 
-We then define our three required KeyFuncs static functions.  The first is GetSetKey which, given an element type, returns the key.  Our element type is TPair, and our key is UniqueID, so we simply return that directly.
+1. We then define our three required KeyFuncs static functions.  The first is GetSetKey which, given an element type, returns the key.  Our element type is TPair, and our key is UniqueID, so we simply return that directly.
 
-The second static function is Matches which takes the keys of two elements, having already been extracted from the element types using GetSetKey, and compares them to see if they are equivalent.  FString's operator== is case-insensitive, and we want a case sensitive search, so we use the FString::Compare function with an appropriate option.
+	The second static function is Matches which takes the keys of two elements, having already been extracted from the element types using GetSetKey, and compares them to see if they are equivalent.  FString's operator== is case-insensitive, and we want a case sensitive search, so we use the FString::Compare function with an appropriate option.
 
-Finally, the GetKeyHash static function takes an extracted key and returns a hashed value for it.  Again, the GetTypeHash behavior for FString is to ignore case, so we call a case-sensitive FCrc function to calculate our hash for us.
+1. Finally, the GetKeyHash static function takes an extracted key and returns a hashed value for it.  Again, the GetTypeHash behavior for FString is to ignore case, so we call a case-sensitive FCrc function to calculate our hash for us.
 
-Now we can create a TMap using our new KeyFuncs.  We also need to provide an allocator, because the KeyFuncs parameter is last, but we'll just substitute the default:
+1. Now we can create a TMap using our new KeyFuncs.  We also need to provide an allocator, because the KeyFuncs parameter is last, but we'll just substitute the default:
 
-	TMap<
-		FMyStruct,
-		int32,
-		FDefaultSetAllocator,
-		TMyStructMapKeyFuncs<int32>
-	> MyMapToInt32;
+		TMap<
+			FMyStruct,
+			int32,
+			FDefaultSetAllocator,
+			TMyStructMapKeyFuncs<int32>
+		> MyMapToInt32;
 	
-	// Add some elements
-	MyMapToInt32.Add(FMyStruct(3.14f), 5);
-	MyMapToInt32.Add(FMyStruct(1.23f), 2);
+		// Add some elements
+		MyMapToInt32.Add(FMyStruct(3.14f), 5);
+		MyMapToInt32.Add(FMyStruct(1.23f), 2);
 	
-	// MyMapToInt32 == [
-	// 	{
-	// 		Key: {
-	// 			UniqueID:  "D06AABBA466CAA4EB62D2F97936274E4",
-	// 			SomeFloat: 3.14f
-	// 		},
-	// 		Value: 5
-	//	},
-	// 	{
-	// 		Key: {
-	// 			UniqueID:  "0661218447650259FD4E33AD6C9C5DCB",
-	// 			SomeFloat: 1.23f
-	// 		},
-	// 		Value: 5
-	//	}
-	// ]
+		// MyMapToInt32 == [
+		// 	{
+		// 		Key: {
+		// 			UniqueID:  "D06AABBA466CAA4EB62D2F97936274E4",
+		// 			SomeFloat: 3.14f
+		// 		},
+		// 		Value: 5
+		//	},
+		// 	{
+		// 		Key: {
+		// 			UniqueID:  "0661218447650259FD4E33AD6C9C5DCB",
+		// 			SomeFloat: 1.23f
+		// 		},
+		// 		Value: 5
+		//	}
+		// ]
 
 A note about providing your own KeyFuncs: be aware that TMap assumes that two items that compare equal using KeyFuncs::Matches return the same value from KeyFuncs::GetKeyHash.  In addition, modifying the key of an existing map element in a way which will change the results from either of these functions is considered undefined behavior, as this will invalidate TMap's internal hash.  These rules also apply to the overloads of operator== and GetKeyHash when using the default KeyFuncs.
 
diff --git a/Engine/Documentation/Source/Programming/UnrealArchitecture/TSet/TSet.INT.udn b/Engine/Documentation/Source/Programming/UnrealArchitecture/TSet/TSet.INT.udn
--- a/Engine/Documentation/Source/Programming/UnrealArchitecture/TSet/TSet.INT.udn
+++ b/Engine/Documentation/Source/Programming/UnrealArchitecture/TSet/TSet.INT.udn
@@ -1,8 +1,11 @@
 Title:TSet
 Crumbs:
-Description:
-Availability:Docs
-Version:4.12
+Description:TSets are a fast container class to store (usually) unique elements in a context where the order is irrelevant. 
+Availability:Public
+Version:4.14
+Related:Engine/UI/LevelEditor/Details/Properties/Set
+
+[TOC]
 
 **TSet** is similar to **TMap** and **TMultiMap**, but with an important difference: Rather than associating elements with separately-provided keys, a TSet uses the element itself as the key, via an overridable function that evaluates the element. TSets are very fast (constant time) for adding, finding, and removing elements. By default, TSets do not permit duplicate keys, but this behavior is supported, should it be desired.
 
@@ -10,7 +13,7 @@ Version:4.12
 
 TSets are a fast container class to store (usually) unique elements in a context where the order is irrelevant. For most use cases, just one parameter - the element type - is needed. However, TSet can be set up with different template parameters to change its behavior and make it more versatile. a derived struct based on `DefaultKeyFuncs` can be specified to provide hashing functionality, as well as permitting multiple keys with the same value to exist in the set. Finally, like the other container classes, a custom allocator for the elements can be specified.
 
-Like TArray, TSet is a homogeneous container, meaning that all of its elements are strictly the same type.  TSet is also a value type and supports the usual copying, assignment, and destructor operations, as well as strong ownership of its elements, which will be destroyed when the set is destroyed. The key type must also be a value type.
+Like TArray, TSet is a homogeneous container, meaning that all of its elements are strictly the same type.  TSet is also a value type and supports the usual copying, assignment, and destructor operations, as well as strong ownership of its elements, which will be destroyed when the set is destroyed. The key type is also required to be a value type.
 
 TSet uses hashes, which means that the **KeyFuncs** template parameter if provided, tells the set how to determine the key from an element, how to compare two keys for equality, how to hash the key, and whether or not duplicate keys will be permitted.  These have defaults which will just return a reference to the key, use `operator==` for equality, and the non-member `GetTypeHash` function for hashing.  If your key type supports these functions, it is usable as a set key without the need to provide a custom KeyFuncs. To write a custom KeyFuncs, extend the `DefaultKeyFuncs` struct.
 
@@ -63,6 +66,15 @@ It's also possible to insert all the elements from another set by using the Appe
 
 Here, the resulting set is equivalent to using Add/Emplace to add them individually, so that duplicate keys from the source set replace ones in the target.
 
+### Editing UPROPERTY TSets
+
+If you mark the TSet with the `UPROPERTY()` macro and one of the editable keywords (`EditAnywhere`, `EditDefaultsOnly`, `EditInstanceOnly`), you can [add and edit elements in Unreal Editor](Engine/UI/LevelEditor/Details/Properties/Set).
+
+	UPROPERTY(Category = SetExample, EditAnywhere)
+	TSet<FString> FruitSet;
+
+
+
 ## Iteration
 
 Iteration over TSets is similar to TArrays.  You can use C++'s ranged-for feature:
@@ -95,7 +107,7 @@ We can ask the set how many elements it holds by using the **Num** function:
 	int32 Count = FruitSet.Num();
 	// Count == 8
 
-We can use the `FSetElementId` struct to find the index of a key within the set. This enables us to look at that element via the indexing operator[]. Calling operator[] on a non-const set will return a non-const reference and calling it on a const set will return a const reference.
+We can use the `FSetElementId` struct to find the index of a key within the set. This enables us to look at that element via the indexing **operator[]**. Calling operator[] on a non-const set will return a non-const reference and calling it on a const set will return a const reference.
 
 	FSetElementId BananaIndex = FruitSet.Index(TEXT("Banana"));
 	// BananaIndex is a value between 0 and (FruitSet.Num() - 1)
@@ -132,27 +144,27 @@ Most of the time, you will want to look up elements without knowing whether or n
 
 If called on a const set, the returned pointer will also be const.
 
-The **Array** function returns a TArray populated with a copy of all the elements in the TSet.  In both cases, the array being passed is emptied before population, so the resulting number of elements will always equal the number of elements in the set:
+The **Array** function returns a TArray populated with a copy of all the elements in the TSet.  The array being passed is emptied before population, so the resulting number of elements will always equal the number of elements in the set:
 
 	TArray<FString> FruitArray = FruitSet.Array();
 	// FruitArray == [ "Banana","Grapefruit","Pineapple","Pear","Orange","Kiwi","Melon","Mango" ] (order may vary)
 
 ## Removal
 
-Elements can be removed by key with the **Remove** function:
+Elements can be removed by index with the **Remove** function, though this is recommended only for use while iterating through the elements:
 
-	FruitSet.Remove(TEXT("Banana");
+	FruitSet.Remove(0);
 	// FruitSet == [ "Grapefruit","Pineapple","Pear","Orange","Kiwi","Melon","Mango" ]
 
 Removing elements will actually leave holes in the data structure, which you can see when visualizing the set in Visual Studio's watch window, but they have been omitted here for clarity. If a TSet supports duplicate keys, **Remove** will remove all keys that match the input parameter. The Remove function returns the number of elements removed, and will be 0 if the key provided was not contained in the set.
 
-	int32 RemovedAmountPineapple = FruitSet.FindAndRemoveChecked(TEXT("Pineapple"));
+	int32 RemovedAmountPineapple = FruitSet.Remove(TEXT("Pineapple"));
 	// RemovedAmountPineapple == 1
 	// FruitSet == [ "Grapefruit","Pear","Orange","Kiwi","Melon","Mango" ]
-	FString RemovedAmountLemon = FruitSet.FindAndRemoveChecked("Lemon"); // assert!
+	FString RemovedAmountLemon = FruitSet.Remove(TEXT("Lemon"));
 	// RemovedAmountLemon == 0
 
-Finally, all elements can be removed by using the **Empty** function:
+Finally, all elements can be removed by using the **Empty** or **Reset** function:
 
 	TSet<FString> FruitSetCopy = FruitSet;
 	// FruitSetCopy == [ "Grapefruit","Pear","Orange","Kiwi","Melon","Mango" ]
@@ -160,11 +172,11 @@ Finally, all elements can be removed by using the **Empty** function:
 	FruitSetCopy.Empty();
 	// FruitSetCopy == []
 
-Like TArray, Empty takes an optional slack value which can be used to optimize the case when you are about to repopulate the set with a given number of elements.
+Like the TArray equivalent, Empty takes an optional slack value (the default is zero) which is used to resize the internal storage array after emptying the set. This value is used as the array's new maximum size. No reallocation occurs if the array's current maximum size is the same as the slack argument.
 
 ## Sorting
 
-TSets can temporarily be sorted with the **Sort** function.  The next iteration over the set will present the elements in sorted order, though future modifications to the set will likely result in the set being unsorted again.  Sorting is unstable, so equivalent elements (in a TSet that permits duplicate keys) may appear in any order.
+TSets can temporarily be sorted with the **Sort** function.  The next iteration over the set will present the elements in sorted order, though future modifications to the set will likely result in the set being unsorted again. Sorting is unstable, so equivalent elements (in a TSet that permits duplicate keys) may appear in any order.
 
 The Sort function takes a binary predicate which specifies the sort order:
 
@@ -191,12 +203,12 @@ Like TArray, TSet is a regular value type and as such can be copied via the stan
 
 ## Slack
 
-TSets also have a notion of **slack** and this can be used to optimize the population of the set.  The **Reset** function acts like an Empty() call, but does not free the memory that was previously used by the elements:
+TSets have a notion of **slack**, and this can be used to optimize the population of the set.  The **Reset** function acts like an Empty() call, but does not free the memory that was previously used by the elements:
 
 	FruitSet.Reset();
 	// FruitSet == [ <invalid>, <invalid>, <invalid>, <invalid>, <invalid>, <invalid> ]
 
-Here the set is emptied in the same way as Empty would, but the memory used for storage is not freed and remains as slack.
+Here the set is emptied in the same way as Empty would, but the memory used for storage is not freed and remains as slack. If a slack value greater than the array's current maximum size is specified, memory will be reallocated to increase the array's capacity up to that slack value.
 
 TSet does not provide a way of checking how many elements are preallocated, like TArray::Max() does, but still supports preallocating slack.  The Reserve function can be used to preallocate slack for a particular number of elements before adding:
 
@@ -221,9 +233,9 @@ The **Shrink** function also works like the TArray equivalent in that it removes
 	FruitSet.Shrink();
 	// FruitSet == ["Fruit8", <invalid>, "Fruit6", <invalid>, "Fruit4", <invalid>, "Fruit2", <invalid>, "Fruit0" ]
 
-Note that only one invalid element has been removed by the Shrink call, because there was only one hole at the end.  The **Compact** function can be used to remove all the holes before shrinking:
+Note that only one invalid element has been removed by the Shrink call, because there was only one hole at the end.  The **Compact** function can be used to remove all the holes before shrinking. If preserving order is important, the **CompactStable** function can be used, but remember that many other TSet operations do not guarantee order stability:
 
-	FruitSet.Compact();
+	FruitSet.CompactStable();
 	// FruitSet == ["Fruit8", "Fruit6", "Fruit4", "Fruit2", "Fruit0", <invalid>, <invalid>, <invalid>, <invalid> ]
 	FruitSet.Shrink();
 	// FruitSet == ["Fruit8", "Fruit6", "Fruit4", "Fruit2", "Fruit0" ]
@@ -248,5 +260,6 @@ A note about providing your own DefaultKeyFuncs: be aware that TSet assumes that
 
 The **CountBytes** and **GetAllocatedSize** functions are used to estimate how much memory is currently being utilized by the internal array.  CountBytes takes an FArchive and GetAllocatedSize can be called directly.  They are typically used for stats reporting.
 
-The **Dump** function takes an FOutputDevice and writes out some implementation information about the contents of the set.  It's usually used for debugging.
+The **Dump** function takes an FOutputDevice and writes out some implementation information about the contents of the set. There is also a **DumpHashElements** function, which lists all elements from all hash entries. These are usually used for debugging.
 
+TSets can be tagged with UPROPERTY for the purposes of automatic serialization, garbage collection, .ini file settings, and editing via the Details panel or as Blueprint defaults. Currently, inline editing of sets is limited to treating the set as a list of values. For example, a TSet of int32s might look like **(1,2,3)**, while a TSet of FNames might appear as **("One","Two","Three")**. Similar to TMap, TSet properties do not currently work as replicated members, and they cannot be accessed by Blueprints.
diff --git a/Engine/Documentation/Source/SiteIndex/SiteIndex.INT.udn b/Engine/Documentation/Source/SiteIndex/SiteIndex.INT.udn
--- a/Engine/Documentation/Source/SiteIndex/SiteIndex.INT.udn
+++ b/Engine/Documentation/Source/SiteIndex/SiteIndex.INT.udn
@@ -2,6 +2,8 @@ Title:Site Map
 Availability:Public
 Crumbs:%ROOT%
 Description:Index of all pages in the Unreal Engine documentation
+parent:%ROOT%
+order:1000
 Template:DefaultSiteMap.html
 
 [OBJECT:Expander]
@@ -27,7 +29,7 @@ Template:DefaultSiteMap.html
 			[REGION:normaltitle]
 				By Engine Version:
 			[/REGION]
-			[tagfilters(type:"engine_version" filter:"used")]
+			[tagfilters(type:"engine_version")]
 		[/REGION]
 		[REGION:wrapper transparent]
 			[REGION:normaltitle]
@@ -35,7 +37,17 @@ Template:DefaultSiteMap.html
 			[/REGION]
 			[tagfilters(type:"tag" filter:"used")]
 		[/REGION]
+		[REGION:wrapper transparent]
+			[REGION:normaltitle]
+				By Type:
+			[/REGION]
+			[tagfilters(type:"type" filter:"used")]
+		[/REGION]
 	[/PARAM]
 [/OBJECT]
 
 [DIR(output:"listbutton" parent:"%ROOT%" org:"hierarchy" id:"sitemap" itemclass:"link")]
+
+<!--
+[DIR(output:"listbutton" type:"playlist" org:"hierarchy" id:"videomap" itemclass:"link")]
+-->
\ No newline at end of file
