Copying //UE4/Dev-Documentation to Samples-Main (//UE4/Samples-Main)

diff --git a/Engine/Documentation/Source/Support/Builds/ReleaseNotes/2016/4_14/ReleaseNotes_4_14.INT.udn b/Engine/Documentation/Source/Support/Builds/ReleaseNotes/2016/4_14/ReleaseNotes_4_14.INT.udn
new file mode 100644
--- /dev/null
+++ b/Engine/Documentation/Source/Support/Builds/ReleaseNotes/2016/4_14/ReleaseNotes_4_14.INT.udn
@@ -0,0 +1,3756 @@
+Availability:Public
+Title: Unreal Engine 4.14 Release Notes
+Description: Release Notes for Unreal Engine 4.14
+Crumbs:%ROOT%, Engine
+reference-image:image_6.gif
+reference-image:image_7.gif
+reference-image:image_9.gif
+reference-image:image_10.gif
+reference-image:image_16.gif
+reference-image:image_23.gif
+reference-image:image_27.gif
+reference-image:image_29.gif
+reference-image:image_30.gif
+reference-image:image_33.gif
+reference-image:image_39.gif
+reference-image:image_41.gif
+reference-image:image_45.gif
+reference-image:image_48.gif
+reference-image:image_49.gif
+reference-image:image_50.gif
+reference-image:image_57.gif
+reference-image:image_59.gif
+reference-image:image_60.gif
+reference-image:image_62.gif
+reference-image:image_63.gif
+reference-image:image_64.gif
+reference-image:image_65.gif
+reference-image:image_86.gif
+
+![image alt text](image_0.png)(h:304 w:936 convert:false)
+
+# What's New
+
+Unreal Engine 4.14 introduces a new **forward shading renderer** optimized for VR, enabling crisp **multi-sampled anti-aliasing** in your games. The new **Contact Shadows** feature renders beautifully detailed shadows for intricate objects. We've also introduced a new **automatic LOD generation** feature for static meshes that does not require a third-party library.
+
+We've **streamlined the animation tools** to help you be more productive, and added many new features to **Sequencer** (UE4's non-linear cinematic tool), as well as improvements to vehicles, clothing and animation Blueprints.
+
+For mobile developers, **Vulkan** support is ready to use on compatible Android devices! And, we've added various new **mobile rendering features** such as reading from scene color and depth, and the ability to draw 3D objects on top of your UI.
+
+On the Windows platform, C++ programmers can now use **Visual Studio "15"** for development.  Visual Studio 2015 is still supported.
+
+This release includes hundreds of updates from Epic and **71 improvements** submitted by the incredible community of Unreal Engine developers on GitHub!  Thanks to each of these contributors to Unreal Engine 4.14:
+
+[REDACTED_CONTRIBUTORS]
+
+# Major Features 
+
+## New: Forward Shading Renderer with MSAA
+
+The new forward shading renderer combines high-quality UE4 lighting features with Multisample Anti-Aliasing (MSAA) support!  MSAA and the option to enable per-material optimizations make the forward renderer well suited for VR.
+
+![image alt text](image_1.png)(w:929 h:664 convert:false)
+
+The forward renderer works by culling lights and reflection captures to a frustum-space grid. Each pixel in the forward pass then iterates over the lights and reflection captures affecting it, shading the material with them. Dynamic shadows for stationary lights are computed beforehand and packed into channels of a screen-space shadow mask allowing multiple shadowing features to be used efficiently.
+
+Enable 'Forward Shading' in the Rendering Project settings and restart the editor to use the forward renderer.
+
+![image alt text](image_2.png)(w:240 h:24 convert:false)
+
+Supported **forward rendering features** include:
+
+* Full support for stationary lights, including dynamic shadows from movable objects which blend together with precomputed environment shadows
+
+* Multiple reflection captures blended together with parallax correction
+
+* Planar reflections of a partial scene, composited into reflection captures
+
+* D-Buffer decals
+
+* Precomputed lighting and skylights
+
+* Unshadowed movable lights
+
+* Capsule shadows
+
+* Instanced stereo compatible
+
+Some features are **not yet supported** with Forward Shading:
+
+* Screen space techniques (SSR, SSAO, Contact Shadows)
+
+* Shadow casting Movable Lights
+
+* Dynamically shadowed translucency
+
+* Translucency receiving environment shadows from a stationary light
+
+* Light functions and IES profiles
+
+* Alpha to Coverage
+
+* MSAA on D-Buffer decals, motion blur, dynamic shadows and capsule shadows
+
+![image alt text](image_3.png)(w:929 h:664 convert:false)
+
+The forward renderer supports both multi sample anti-aliasing (MSAA) and temporal anti-aliasing (TAA). In most cases TAA is preferable because it removes both geometric aliasing and specular aliasing. In VR, the constant sub-pixel movement introduced by head tracking introduces unwanted blurriness, making MSAA a better choice.
+
+Projects that choose to use MSAA will want to build content to mitigate specular aliasing. The 'Normal to Roughness' feature can help reduce specular aliasing from detailed normal maps. Automatic LOD generation for static meshes can flatten features on distant meshes and help reduce aliasing from small triangles.
+
+In our tests, using MSAA instead of TAA increases GPU frame time by about 25%. Actual cost will depend on your content.
+
+To use MSAA, set the default Anti-Aliasing Method in the Rendering project settings:
+
+![image alt text](image_4.png)(w:336 h:53 convert:false)
+
+The console variable 'r.MSAACount' controls how many MSAA samples are computed for every pixel. 'r.MSAACount 1' has special meaning and falls back to Temporal AA, which allows for convenient toggling between anti-aliasing methods.  
+
+**Performance**
+
+The forward renderer can be faster than the deferred renderer for some content. Most of the performance improvement comes from features that can be disabled per material. By default, only the nearest reflection capture will be applied without parallax correction unless the material opts-in to High Quality Reflections, height fog is computed per-vertex, and planar reflections are only applied to materials that enable it.
+
+![image alt text](image_5.png)(w:228 h:73 convert:false)
+
+Leveraging these options in Epic's new VR game, Robo Recall, the forward renderer is about 22% faster than the deferred renderer on an NVIDIA 970 GTX.
+
+## New: Contact Shadows
+
+Contact shadows allow for highly detailed dynamic shadows on objects.
+
+[REGION:asyncgif]
+![image alt text](image_6.gif)(w:929 h:522)
+[/REGION]
+
+The ivy below is only a few flat cards but is able to self-shadow in a very convincing way due to outputting Pixel Depth Offset in the material.
+
+[REGION:asyncgif]
+![image alt text](image_7.gif)(w:929 h:606)
+[/REGION]
+
+The Contact Shadows feature adds a short ray cast in screen space against the depth buffer to know whether a pixel is occluded from a given light or not. This helps provide sharp detailed shadows at the contact point of geometry. There are a number of reasons why shadows through other algorithms may have missing or blurry contacts. Typically it is due to lack of resolution or a depth bias. Regardless of the reason, the new Contact Shadows feature can fill in the gap very well for a small cost.
+
+Contact shadows can be used by setting the Contact Shadow Length property on your light. This controls the length of the ray cast in screen space where 1 is all the way across the screen. Large values can degrade quality and performance so try and keep the length to the minimum that achieves your desired look.
+
+![image alt text](image_8.png)(w:434 h:118 convert:false)
+
+Another use case of contact shadows is to get self-shadowing from the parallax occlusion mapping from arbitrary lights. This requires outputting pixel depth offset in the material. This animation shows a parallax occlusion mapped surface with contact shadow length set to 0.1.
+
+[REGION:asyncgif]
+![image alt text](image_9.gif)(w:752 h:344)
+[/REGION]
+
+## New: Automatic LOD Generation
+
+Unreal Engine now automatically reduces the polygon count of your static meshes to create LODs!
+
+[REGION:asyncgif]
+![image alt text](image_10.gif)(w:768 h:628)
+[/REGION]
+
+The above animation shows five LODs that were generated automatically. Each is half the number of triangles as the previous.
+
+Automatic LOD generation uses what is called quadric mesh simplification. The mesh simplifier will calculate the amount of visual difference that collapsing an edge by merging two vertices would generate. It then picks the edge with the least amount of visual impact and collapses it. When it does, it picks the best place to put the newly merged vertex and removes any triangles which have also collapsed along with the edge. It will continue to collapse edges like this until it reaches the requested target number of triangles.
+
+This mesh simplifier maintains UVs including generated lightmap UVs, normals, tangents, and vertex colors. Because UVs are maintained the same materials can be used as well as all LODs can share the same lightmap. 
+
+The high level settings for controlling the generated LODs are in the static mesh viewer under LOD Settings.
+
+![image alt text](image_11.png)(w:425 h:184 convert:false)
+
+"LOD Group" provides a list of presets. These can be changed per project in BaseEngine.ini under [StaticMeshLODSettings]. We encourage you to set up good categories for your project and mostly use LOD groups instead of controlling the details of every LOD.
+
+An important setting to note is "Auto Compute LOD Distances". Because the algorithm knows how much visual difference every edge collapse is adding it can use this information to determine what distance that amount of error is acceptable. That means it will automatically calculate the screen size to use for each LOD as well.
+
+If you wish to muck with the details of auto generation for each LOD they can be found under Reduction Settings.  Note that this feature currently only works with static meshes and that mesh proxy LOD generation is not yet supported.
+
+## New: Precomputed Lighting Scenarios
+
+We now support precomputing lighting for multiple lighting setups with the same geometry! This is especially important for use cases such as VR and architectural visualization where you need the highest possible quality at the fastest possible performance.
+
+![image alt text](image_12.png)(w:929 h:540 convert:false)
+
+![image alt text](image_13.png)(w:929 h:541 convert:false)
+
+In the above example the directional light, sky light and skybox have been placed in a Lighting Scenario level called DayScenario.  The streetlights have been placed in NightScenario.
+
+![image alt text](image_14.png)(w:260 h:164 convert:false)
+
+To use Lighting Scenarios:
+
+* Right click on a sublevel in the Levels window and change it to Lighting Scenario.  When a Lighting Scenario level is made visible, its lightmaps will be applied to the world.
+
+* Change the level streaming method to Blueprint on the Lighting Scenario level
+
+* Place meshes and lights into this level and build lighting
+
+* In the BeginPlay of your persistent level's Level Blueprint, execute a Load Stream Level on the Lighting Scenario level that you want active.
+
+![image alt text](image_15.png)(w:552 h:196 convert:false)
+
+Limitations:
+
+* Only one Lighting Scenario level should be visible at a time in game.
+
+* When a Lighting Scenario level is present, lightmap data from all sublevels will be placed inside it so that only the DayScenario lightmaps are loaded when it's daytime.  As a result, lightmaps will no longer be streamed by sublevel.
+
+* A Reflection Capture updated is forced when making a Lighting Scenario level visible, which can increase load time.
+
+## New: Improved Per-Pixel Translucent Lighting
+
+In the deferred renderer, the new forward shading functionality can now be used on translucent surfaces to get specular highlights from multiple lights and image-based reflections from parallax-corrected reflection captures!
+
+[REGION:asyncgif]
+![image alt text](image_16.gif)(w:664 h:401)
+[/REGION]
+
+![image alt text](image_17.jpg)(w:329 h:118 convert:false)
+
+## New: Full Resolution Skin Shading
+
+UE4 now supports full resolution skin shading for the Subsurface Profile shading model. This provides high-fidelity lighting for surface details such as pores and wrinkles. 
+
+![image alt text](image_18.png)(w:312 h:334 convert:false) ![image alt text](image_19.png)(w:312 h:334 convert:false)
+
+*Checkerboard rendered skin (left), Full resolution skin (right)  *(Note: 3D head model by [REDACTED_PERSON])
+
+![image alt text](image_20.png)(w:320 h:320 convert:false) ![image alt text](image_21.png)(w:321 h:321 convert:false)
+
+*Surface detail - checkerboard (left), full resolution (right)*
+
+Previously, lighting on skin was represented using a checkerboard pattern, where half the pixels contained diffuse lighting and the other half, specular lighting. The lighting was recombined during a final subsurface profile fullscreen pass. That approach gave good results for subsurface lighting (which is low-frequency by nature), but it could result in lower fidelity lighting for surface details. 
+
+With the new approach, every pixel contains diffuse and specular lighting information, packed into an RGBA encoding. This allows us to reconstruct full-resolution lighting during the final subsurface profile fullscreen pass, giving better results for surface details and more stable behavior with temporal antialiasing. 
+
+**Compatibility**
+
+Full resolution skin shading requires at least a 64-bit scene color format with a full alpha channel. The default FloatRGBA scene color format works fine, but 32-bit representations such as FloatRGB are not supported. If the scene color format is not compatible with full resolution skin, we fall back to checkerboard-based lighting. 
+
+This behaviour can be overridden using the *r.SSS.Checkerboard *console variable. The possible values for this are:
+
+0: Checkerboard disabled (full resolution)
+
+1: Checkerboard enabled (old behavior)
+
+2: Automatic (default) - Full resolution lighting will be used if the scene color pixelformat supports it
+
+**Limitations**
+
+It's worth noting that the full-resolution skin shading is an approximation. It works well in the vast majority of cases, but certain material features can be problematic due to the encoding method. In particular:
+
+* Metallic materials
+
+* Emissive materials
+
+These features will work, but you may notice differences in output compared to checkerboard due to the packed RGBA diffuse/specular encoding. 
+
+It is possible to workaround particular issues when authoring materials by setting the opacity to 0 in areas where skin shading is not desirable. Pixels with an opacity of zero are treated as default lit for the purposes of shading. 
+
+*Note: *Masking non-opaque pixels in this way is also worthwhile for performance reasons, since these pixels are bypassed by the subsurface postprocess. 
+
+**Performance Considerations**
+
+If your title has a 64-bit scene color format then full resolution subsurface lighting will typically be faster than checkerboard due to the reduced number of texture fetches. However, if your title has a 32-bit scene color then the performance gain from the reduced texture bandwidth will likely outweigh the benefits (although this is hardware dependent). 
+
+## New: Reflection Capture Quality Improvements
+
+When you use Reflection Captures, the engine mixes the indirect specular from the Reflection Capture with indirect diffuse from lightmaps.  This helps to reduce leaking, since the reflection cubemap was only captured at one point in space, but the lightmaps were computed on all the receiver surfaces and contain local shadowing.
+
+![image alt text](image_22.png)(w:757 h:436 convert:false)
+
+*(**With lightmap mixing on the left, without on the right**)*
+
+Mixing works well for rough surfaces, but for smooth surfaces the reflections from Reflection Captures no longer match reflections from other methods, like Screen Space Reflections or Planar Reflections.  
+
+**Lightmap mixing is no longer done on very smooth surfaces**. ** **A surface with roughness .3 will get full lightmap mixing, fading out to no lightmap mixing by Roughness .1 and below.  This allows Reflection Captures and SSR to match much better and it's harder to spot transitions.  
+
+The below shot shows mirror surface reflections before and after.  Note the difference in the reflection of the wall between SSR and reflection captures.  The artifact is especially noticeable in motion, because it will move with your camera due to SSR limitations.
+
+[REGION:asyncgif]
+![image alt text](image_23.gif)(w:762 h:440)
+[/REGION]
+
+**This affects existing content - in cases where you had reflection leaking on smooth surfaces, that leaking will be much more apparent.** To solve this, place additional reflection probes to reduce the leaking. Levels should have one large spherical capture at a minimum. You can also revert to the old lightmap mixing behavior with a rendering project setting:
+
+![image alt text](image_24.png)(w:356 h:53 convert:false)
+
+## New: Visual Studio "15" Support
+
+Unreal Engine 4.14 now the upcoming Visual Studio "15" out of the box.  Visual Studio 2015 is still supported as well.  Visual Studio "15" is currently available in "Preview" from [Microsoft's Visual Studio web site](https://www.visualstudio.com/en-us/news/releasenotes/vs15-relnotes).
+
+![image alt text](image_25.png)(w:929 h:558 convert:false)
+
+If you have multiple versions of Visual Studio installed, you can select which to use through the 'Source Code' section in 'Editor Preferences.'
+
+![image alt text](image_26.png)(w:670 h:132 convert:false)
+
+## New: Create Static Mesh from Actors
+
+You can now right-click actor(s) in the level viewport and convert their current state to a new Static Mesh asset. This even works with skeletal meshes, so you can capture a mesh from posed characters.
+
+[REGION:asyncgif]
+![image alt text](image_27.gif)(w:929 h:598)
+[/REGION]
+
+## New: NVIDIA Ansel Support
+
+UE4 4.14 adds support for NVIDIA Ansel Photography!  
+
+Ansel is a new tool from NVIDIA that enables players to take in-game screenshots. While in Ansel mode the game will pause and players will have camera control to compose shots and apply various screen effects. It can also capture a variety of screenshots, from HDR to 360 stereo. See NVIDIA's website for more details: [http://www.geforce.com/hardware/technology/ansel](http://www.geforce.com/hardware/technology/ansel)
+
+![image alt text](image_28.jpg)(w:929 h:598 convert:false)
+
+Ansel support is now exposed as a new UE4 plugin.  After enabling the plugin in your project, you can access Ansel in a standalone game session. 
+
+[REGION:asyncgif]
+![image alt text](image_29.gif)(w:928 h:425)
+[/REGION]
+
+*(Viewing an Ansel 360 capture in a web browser)*
+
+We have also exposed functions on the Player Camera Manager class so your games can customize Ansel capture behavior. Games may wish to limit the distance of camera movement, disable UI elements, disable/enable certain lighting or post processing effects, etc. Thanks to [REDACTED_PERSON] and NVIDIA for providing the implementation.
+
+To get started using this feature, check out the [Ansel plugin documentation](Engine/Plugins/Ansel).
+
+## New: Improved Cable Component
+
+The Cable Component plugin has been updated with new features, including collision support and sockets for attaching objects or effects.
+
+[REGION:asyncgif]
+![image alt text](image_30.gif)(w:781 h:496)
+[/REGION]
+
+Cable Component now includes these new features:
+
+* Simple collision, including friction settings
+
+* Stiffness setting, which tries to reduce bending
+
+* Sockets at each end of the cable
+
+* Ability to set either end to 'free'
+
+## New: UI Font Outlines
+
+Fonts for UMG and Slate now have an optional outline that can be applied to them.  
+
+![image alt text](image_31.png)(w:929 h:534 convert:false)
+
+Any widget that specifies a font can change the outline setting, color, or material to be used with the outline.
+
+![image alt text](image_32.png)(w:397 h:148 convert:false)
+
+A font material on an outline can be used in the same way that any other font material is used except that a material specified for an outline only applies to the outline.  Font materials can be used on the outline to create lots of different effects.
+
+[REGION:asyncgif]
+![image alt text](image_33.gif)(w:645 h:194)
+[/REGION]
+
+## New: Editable Map and Set Properties
+
+We now support editing **Map and Set properties** from within the Details Panel! 
+
+**Sets **are similar to Arrays, but you can never have the same element in a set twice and the order of elements is not guaranteed.  However, it's extremely quick to lookup into a set to see whether it contains an element.
+
+![image alt text](image_34.png)(w:401 h:185 convert:false)
+
+**Maps **will have a key and a value and you can edit both within the details panel.  Like Sets, all keys must be unique, and the order of elements is not guaranteed to persist.  However, it's very quick to lookup an element's value as long as you know it's key.
+
+![image alt text](image_35.png)(w:401 h:254 convert:false) 
+
+## New: Vector Noise in Materials
+
+The **Noise **material graph node includes several functions useful for procedural shading that produce a single-valued (scalar) result. 
+
+[REGION:imagetable]
+| ![image alt text](image_83.jpg)(w:150 convert:false) | ![image alt text](image_84.jpg)(w:150 convert:false) | ![image alt text](image_85.jpg)(w:150 convert:false) | [REGION:asyncgif]![image alt text](image_86.gif)(w:150)[/REGION] | ![image alt text](image_87.jpg)(w:150 convert:false) |
+| --- | --- | --- | --- | --- |
+| Cellnoise | Vector Noise | Gradient | Curl | Voronoi |
+[/REGION]
+
+The new **Vector Noise** node adds several more with 3D or 4D vector results. Due to the run-time expense of these functions, it is recommended that once a look is developed with them, all or part of the computation be baked into a texture using the **Draw Material to Render Target** Blueprint feature introduced in 4.13. These material graph nodes allow procedural looks to be developed in engine on final assets, providing an alternative to creating procedurally generated textures with an external tool to apply to assets in the engine.
+
+The new functions are:
+
+1. **Cellnoise**: Returns a random color for each cell in a 3D grid (i.e. from the mathematical floor operation applied to the node input). The results are always consistent for a given position, so can provide a reliable way to add randomness to a material. This Vector Noise function is extremely cheap to compute, so it is not necessary to bake it into a texture for performance.
+
+2. **Perlin 3D Noise**: Computes a version of Perlin Simplex Noise with 3D vector output. Each output component is in the range -1 to 1. Computing three channels of noise output at once is cheaper than merging the results from three scalar noise functions.
+
+3. **Perlin Gradient**: Computes the analytical 3D gradient of a scalar Perlin Simplex Noise. The output is four channels, where the first three (RGB) are the gradient, and the fourth (A) is the scalar noise. This is useful for bumps and for flow maps on a surface
+
+4. **Perlin Curl**: Computes the analytical 3D curl of a vector Perlin Simplex Noise (aka Curl Noise). The output is a 3D signed curl vector. This is useful for fluid or particle flow.
+
+5. **Voronoi**: Computes the same Voronoi noise as the scalar Noise material node. The scalar Voronoi noise scatters seed points in 3D space and returns the distance to the closest one. The Vector Noise version returns the location of the closest seed point in RGB, and the distance to it in A. Especially coupled with Cellnoise, this can allow some randomized behavior per Voronoi cell.
+
+Below is a simple stone bed material using the distance component of the Vector Noise / Voronoi to modulate some surface bumps and blend in moss in the cracks, and the seed position together with Vector Noise / Cellnoise to change the color and bump height per rock.
+
+![image alt text](image_36.png)(w:626 h:304 convert:false)
+
+Perlin Curl and Perlin Gradient can be added together in octaves, just as regular Perlin noise can. For more complex expressions, it is necessary to compute the gradient of the result of the expression. To help with this, place the expression to compute into a material function and use it with the helper nodes **Prepare3DDeriv**, **Compute3DDeriv**, and either **GradFrom3DDeriv** or **CurlFrom3DDeriv**. These use four evaluations of the base expression spaced in a tetrahedral pattern to approximate these derivative-based operations. For example, this network uses the gradient to compute bump normals from a bump height function.
+
+![image alt text](image_37.png)(w:929 h:241 convert:false)
+
+## New: PhysX 3.4 Upgrade
+
+Unreal Engine now uses the latest version of NVIDIA PhysX, which is 3.4. This brings improved performance and memory usage for rigid bodies and scene queries (especially multi-core performance.) 
+
+![image alt text](image_38.png)(w:441 h:197 convert:false)
+
+This version of PhysX allows for **Continuous Collision Detection **(CCD) on **kinematic objects**, which allows for accurate collisions between very fast moving rigid bodies!  In the animation below from a Robo Recall test level, a player is swiping a weapon to impact an oncoming bullet!
+
+[REGION:asyncgif]
+![image alt text](image_39.gif)(w:674 h:378)
+[/REGION]
+
+New features available to use in UE4 right away:
+
+* Continuous Collision Detection (CCD) support for kinematic objects (shown in the animation above!)
+
+* Faster updating of kinematic objects 
+
+* Faster convex hull cooking
+
+In future releases, we'll expose more new physics features available in the latest version of PhysX.
+
+## New: Animation Editor Revamp
+
+Animation-related tools have been split into separate asset editors rather than using one editor with multiple modes.
+
+![image alt text](image_40.jpg)(w:929 h:560 convert:false)
+
+Many other improvements have been made as well. Functionality that is common to each of the editors is now generally found in the viewport and the improved Skeleton Tree.
+
+* The Skeletal Mesh editor has had modifications to its layout and to the asset details panel, specifically the materials and LOD sections have been overhauled.
+
+* The Skeleton editor has had its layout tweaked and the skeleton tree itself has been polished.
+
+* The Animation editor has had its layout tweaked and the asset browser has gained the ability to optionally add and remove its columns.
+
+* The Animation Blueprint editor has had its layout tweaked to more closely follow that of the standard Blueprint editor. The Anim Preview Editor can now optionally apply changes that are made to the preview's properties to the class defaults.
+
+[REGION:asyncgif]
+![image alt text](image_41.gif)(w:409 h:552)
+[/REGION]
+
+### Asset Shortcut Bar
+
+![image alt text](image_42.png)(w:460 h:60 convert:false)
+
+You can jump between related animation assets that share a skeleton using the improved Asset Shortcut Bar.
+
+### Recording Moved to Transport Controls
+
+![image alt text](image_43.png)(w:252 h:68 convert:false)
+
+Recording used to be performed via a button in the toolbar. Now it has been moved to a recording button in the transport controls, similar to Sequencer.
+
+### Preview Scene Setup
+
+![image alt text](image_44.png)(w:605 h:378 convert:false)
+
+The objects in the scene and their animation can be modified in each of the editors via the "Scene Setup" menu. This allows preview animations to be applied, different preview meshes to be set (this is either specified for the skeleton or for individual animations) and additional meshes to be attached. Additional meshes are now specified as separate editor-only assets that define a set of skeletal meshes that are driven as slaves of the main mesh.
+
+## New: Animation Curve Window
+
+You can now easily tweak Animation Curves using the new dedicated window for this in the Animation Editor.  Curves are previewed live as you edit them.
+
+[REGION:asyncgif]
+![image alt text](image_45.gif)(w:761 h:537)
+[/REGION]
+
+Previously you could only configure curves on the animation assets themselves, but now you'll set these for the skeleton instead.
+
+## New: Child Actor Templates
+
+Child Actor Components added to a Blueprint can have their properties customized via Child Actor Templates.
+
+![image alt text](image_46.png)( w:929 h:658convert:false)
+
+Once you add a Child Actor Component, you will see an expandable template in the Details panel of the owning Actor's Blueprint Editor. From here, you can access all the properties of the Child Actor, including public variables. 
+
+For example, if you have *Blueprint_A* containing a PointLight Component with a public variable driving its color, and then make that Blueprint a Child Actor Component within *Blueprint_B*, you can now adjust that color variable from within *Blueprint_B's* Details panel!
+
+This is a dramatic improvement over previous behavior, wherein users were restricted to the default properties of the Child Actor Component and could only make updates via gameplay script.
+
+## New: Default Animation Blueprint
+
+Allows you to assign an animation Blueprint to a skeletal mesh that will always be run after any animation Blueprint assigned in the component. This allows you to set up anim dynamics or other controllers that will always be applied, whether that mesh is viewed in the animation tools, a Sequencer cinematic or just placed in a level.
+
+This allows for dynamics, controllers, IK or any other anim Blueprint feature to be related to a mesh and not have to be duplicated in every animation Blueprint intended to be used on that mesh.
+
+'Post process' animation Blueprints also have their own native and Blueprint update step so parameters can be read or calculated for use in the animation graph.
+
+![image alt text](image_47.png)(w:356 h:69 convert:false)
+
+[REGION:asyncgif]
+![image alt text](image_48.gif)(w:740 h:441)
+[/REGION]
+
+## New: Landscape Editing in VR
+
+You can now create and sculpt terrain and paint landscape materials using motion controllers in VR!  
+
+[REGION:asyncgif]
+![image alt text](image_49.gif)(w:620 h:336)
+[/REGION]
+
+You can summon the Landscape Editing tools from the "Modes" panel on your Quick Menu.  Then choose a brush from the UI and start painting!  If you hold the "Modifier" button on the motion controller, you can erase instead of painting.
+
+## New: Improved Support for Vehicles
+
+We've changed where tire forces are applied.  Previously, tire forces were applied at the vehicle's center of mass. We now apply force at the **tire's center of mass** which makes it easier to achieve load sway in cars. 
+
+[REGION:asyncgif]
+![image alt text](image_50.gif)(w:736 h:372)
+[/REGION]
+
+We've also added **Simple Wheeled Vehicle Movement Component** which provides wheel suspension and tire friction without the complexities of engine and drivetrain simulation.  This component allows you to easily apply **torque to individual tires**. All components inheriting from Wheeled Vehicle Movement Component can now be used on arbitrary components, and you no longer have to rely on the Wheeled Vehicle actor.
+
+Existing content will automatically have **Deprecated Spring Offset Mode** set to true which will maintain the old behavior. You can tune this behavior further by changing Suspension Force Offset.
+
+## New: Improved Vulkan support on Android
+
+Unreal Engine 4.14 is ready for shipping games with Vulkan support!
+
+[REGION:raw]
+![image alt text](image_51.png)( w:837 h:409 convert:false)
+[/REGION]
+
+* UE4 supports Android 7 (Nougat) devices with Vulkan drivers as well as the Samsung Galaxy S7 running a recent OTA update.
+
+* Many rendering issues have been fixed with the UE4 Vulkan renderer on Android devices.
+
+* The renderer will automatically fall back to OpenGL ES when launched on Android devices that are not Vulkan-capable.
+
+* Vulkan support on specific devices and driver versions can now be enabled or disabled using device profiles, with fallback to ES 3.1 and ES 2. This allows UE4 games to disable Vulkan support and use OpenGL ES on phones with incomplete or broken Vulkan implementations.
+
+## New: Support for Custom Depth on Mobile
+
+Custom Depth is now supported in the mobile rendering path. Custom post-process materials can now sample from Scene Depth, Custom Depth as well as Scene Color. 
+
+[REGION:raw]
+![image alt text](image_52.png)(w:804 h:394 convert:false)
+[/REGION]
+
+As it requires post-processing, Mobile HDR must be enabled, and the feature does not currently work while Mobile MSAA is enabled.
+
+## New: Scene Capture Improvements on Mobile 
+
+When rendering scene captures, the Scene Capture Source settings that output Inverse Opacity and Depth values are now supported on mobile. 
+
+[REGION:raw]
+![image alt text](image_53.png)(w:760 h:373 convert:false)
+[/REGION]
+
+* The "SceneColor (HDR) in RGB, Inv Opacity in A" option can be used to render objects with translucency into a texture which can then be alpha-blended over a scene or widget blueprint. 
+
+* Similarly, the depth value can be used as a mask when using the resulting texture.
+
+* Generating the opacity data has some cost, so use "SceneColor (HDR) in RGB, 0 in A" for improved performance if you do not need opacity
+
+* Scene captures now work correctly on devices that do not support floating point targets, such as Galaxy S6 prior to Android 6.
+
+## New: Improved Cloth Skinning
+
+We have added the ability to calculate our own mesh-to-mesh skinning data for clothing within the engine, so rather than using the render data exported in an .apx or .apb file we now use the render data UE4 already has. We take the simulation mesh from the APEX-exported asset and reskin our render data onto that mesh. This means that the final data should look as good as the data you originally imported.
+
+This brings a few benefits. Normals could previously appear incorrect (see image below) and you were previously restricted to one UV channel. Both of these issues are solved with the new skinning system.
+
+![image alt text](image_54.jpg)(w:510 h:585 convert:false)
+
+## New: Material Attribute Nodes
+
+Working with material attributes is now easier to read and less error prone as part of an ongoing update to improving extensibility of material properties.
+
+* **GetMaterialAttributes **- This node is a compact replacement for BreakMaterialAttributes
+
+* **SetMaterialAttributes **- This node is a compact replacement for MakeMaterialAttributes
+
+* **BlendMaterialAttributes **- This is a new node to allow easier blending of Material Attributes structures.
+
+The main improvement for the Get and Set nodes is that pins are optionally added unlike the Break and Make nodes which expose all attributes by default. This allows graphs to avoid the old workflow that required manually connecting every attribute pin. Selecting a node shows the list of current pins in the details panel which can be expanded or removed. For an example, the material function below takes a set of attributes then blends the Base Color and Roughness to a shiny, red surface.
+
+![image alt text](image_55.png)(w:841 h:226 convert:false)
+
+As well as reducing clutter in the graphs, these nodes take advantage of many backend changes to be forward-compatible with any custom material attributes that a project may need to add. Sharing materials between projects is more viable as missing attributes are automatically  detected and users given a chance to handle the errors. Any attribute not explicitly listed on a node is passed through with the main Material Attributes pin, including any that are added after the material graph is created. With the Make and Break nodes a new pin would be added and all graphs would need manually updating.
+
+The new Blend node is intended to allow blending of multiple sets of attributes using a mask, a common operation when working with detailed layers of materials. The example below evenly blends Red and Green materials (defined as functions) then has a node that applies a clear-coat to the result:
+
+![image alt text](image_56.png)(w:929 h:249 convert:false)
+
+By default the Blend node performs a linear interpolation (lerp) for all material attributes using the Alpha input. The node has checkboxes to opt-out of blending on a per-vertex/pixel level to allow easier control when using vertex-only or pixel-only mask data. Similarly to the new Get and Set nodes above, the Blend node will automatically handle new attributes being added or removed and allows programmers to specify custom blending behavior when registering attributes.
+
+## New: Pre-Skinned Local Position in Materials
+
+Materials now have access to a skeletal mesh's reference pose position for use in per-vertex outputs. This allows localized effects on an animated character. The node can be shared for static meshes also for which it returns the standard local position. The example graph below creates a grid pattern in local-space which remains relative to the skeletal mesh during animation:
+
+[REGION:asyncgif]![image alt text](image_57.gif)(w:346 h:369)[/REGION] ![image alt text](image_58.png)(w:500 h:369 convert:false)
+
+## New: Improved Sequencer Shot Import/Export
+
+Movie recording with frame handles per shot.  Master sequences can now be rendered with extra frames at the start and end of each shot.  These extra frames are cut into and out of by an Edit Decision List (EDL), which can be used in an external video editing package to adjust the cuts between shots.![image alt text](image_59.gif)(h:233 w:929 convert:false)
+
+## New: Improved Camera Rig Crane
+
+We've tweaked the camera rig crane behavior so that it mimics the movement of a physical crane.
+
+[REGION:asyncgif]
+![image alt text](image_60.gif)(w:773 h:457)
+[/REGION]
+
+* Roll and yaw of the camera crane mount is 0.
+
+* Add toggles to lock the mount pitch/yaw for the crane. By default they are not locked so that the camera will stay level with the ground.
+
+## New: Sequencer Audio Recording
+
+You can now record audio from a microphone while recording into a sequence.![image alt text](image_61.jpg)(h:288 w:929 convert:false)
+
+## New: Pose Driver Improvements
+
+The Pose Driver node allows a bone to drive other aspects of animation, based on a set of 'example poses'. In this release, it can now drive bone transforms as well as morph targets, for example driving a shoulder pad bone based on arm rotation. We have also added an option to use the translation of the driving bone instead of its orientation. Debug drawing has been improved to show each 'target' pose and how close the bone is currently considered to it.
+
+[REGION:asyncgif]
+![image alt text](image_62.gif)(w:449 h:390)
+[/REGION]
+
+## New: Virtual Bones
+
+We've added the ability to add 'virtual bones' to a skeleton. Virtual bones are not skinnable, but constrained between two existing bones on the skeleton and automatically have data generated for them for each animation on the skeleton. For example, you could add a joint that is a child of a hand, but constrained to a palm joint.  Unlike socket, then this joint can be used in Animation Blueprint as a target - i.e. IK target or look at target - or you could modify them in AnimBP for later use. 
+
+[REGION:asyncgif]
+![image alt text](image_63.gif)(w:489 h:330)
+[/REGION]
+
+This helps to improve character iteration time. Previously, if you change your target joint hierarchy for IK or aim, you have to do this in outside of engine, DCC, and import back all the animations to fix the animation data with that new joint included, but this virtual bone will allow you to skip that and do all of work in engine. However this will require recompressing of the animation data to include that joint back to the animation data. 
+
+To see more practical usage of virtual bones, see "Animation Techniques used in Paragon" for more information. It can be used to make it easier to retarget or change reference frames for controllers and are used for orientation and slope warping in Paragon. 
+
+## New: Morph Target Debug View Mode
+
+The new Morph Target View Mode makes it easy to see which vertices are affected by each morph target.
+
+[REGION:asyncgif]
+![image alt text](image_64.gif)(w:545 h:396)
+[/REGION]
+
+## New: Child Animation Montages
+
+Create a Child Montage based on a parent Montage, allowing you to replace animation clips, whilst maintaining overall timing. Useful for adding variations to a move whilst guaranteeing it won't affect gameplay.
+
+[REGION:asyncgif]
+![image alt text](image_65.gif)(w:566 h:624)
+[/REGION]
+
+## New: MIDI Device Plugin 
+
+This release contains a new "MIDI Device" plugin for interaction with music hardware. This is a simple MIDI interface that allows you to receive MIDI events from devices connected to your computer.  Currently only input is supported.
+
+In Blueprints, here's how to use it:
+
+* Enable the "MIDI Device" plugin using the Plugins UI, then restart Unreal Editor.
+
+* Look for "MIDI Device Manager" in the Blueprint RMB menu.
+
+* Call "Find MIDI Devices" to choose your favorite device. 
+
+* Break the "Found MIDI Device" struct to see what's available.
+
+* Then call "Create MIDI Device Controller" for the device you want. Store that in a variable. (It's really important to store the reference to the object in a variable, otherwise it will be garbage collected and won't receive events!)
+
+* On your MIDI Device Controller, bind your own Event to the "On MIDI Event" event. This will be called every game Tick when there is at least one new MIDI event to receive.
+
+* Process the data passed into the Event to make your project do stuff!
+
+## New: Landscape Rotation Tool
+
+The landscape mirror tool can now flip the reflected geometry parallel to the mirror plane, to create diagonally-opposed multiplayer maps.
+
+![image alt text](image_67.jpg)(w:721 h:453 convert:false)
+
+## New: Improved Mesh Material Slot Importing 
+
+The material workflow has been changed in order to give more control and information on how every material is used static and skeletal meshes and to improve material ordering inconsistencies when reimporting meshes.
+
+![image alt text](image_68.jpg)(w:461 h:354 convert:false)
+
+Each element in the list is a material slot with the following information
+
+* Name of the slot
+
+    * The name of the slot is used to match up the material on reimport.  When a mesh is reimported it looks for this name in the FBX file to determine which sections should match up to existing materials.  Previously this relied on index ordering which was easy to become out of order.  
+
+    * Meshes that were imported before this change will have their material slot set to none.  Meshes imported after this change will have their material slot set to the imported material name by default.
+
+* Material asset reference
+
+* The original imported material name (in the tooltip)
+
+In Blueprints and C++ it is now possible to use the material slot name instead of using a hard coded index to retrieve a material slot. Call **Set Material By Name** to set a dynamic material on a skeletal mesh or static mesh component.  Using a name lookup instead of an index ensures gameplay code still works properly if the order of materials on a mesh changes.  
+
+## New: Platform SDK Upgrades
+
+In every release, we update the engine to support the latest SDK releases from platform partners.
+
+![image alt text](image_69.png)(w:929 h:397 convert:false)
+
+* Xbox One:  Upgraded to August 2016 QFE 2
+
+* Playstation 4: Upgraded to PSR SDK 4.008.061
+
+* HTML5: Upgraded to Emscripten 1.36.13
+
+* macOS: Now supports 10.12 Sierra, Xcode 8.1
+
+* iOS/tvOS: Now supports iOS10/tvOS10, Xcode 8.1
+
+## New: Blueprint Library for Mobile Downloading/Patching
+
+The new **Mobile Patch Utilities Blueprint library** contains all the functionality required to allow a mobile game to **download and install game contents and patches** from a cloud website instead of being distributed as part of the initial download from the App Store.
+
+![image alt text](image_70.png)(w:929 h:486 convert:false)
+
+There is functionality to determine if updated game content is available, initiate the download, track progress, handle any errors and finally install the content paks that are downloaded successfully. Functionality to check for sufficient storage space and WiFi connectivity is also available, so the blueprint can warn the user in such cases. Both Android and iOS are supported.
+
+## New: Amazon GameCircle Plugin for Kindle Fire
+
+A new Online Subsystem GameCircle plugin is now included! 
+
+![image alt text](image_71.jpg)(w:929 h:293 convert:false)
+
+GameCircle Achievements, Leaderboards, and Friends is supports as well as Amazon In-App Purchases. Enabling the plugin will provide access to a new Amazon GameCircle project settings panel under the Plugins category. Changes to the AndroidManifest.xml for Fire TV may be enabled here.
+
+![image alt text](image_72.jpg)(w:929 h:334 convert:false)
+
+## New: Live GPU Profiler
+
+UE 4.14 includes a real-time GPU profiler which provides per-frame stats for the major rendering categories. To use it, enter the console command 'stat gpu*.*'* *You can also bring these up in the editor via the 'Stat' submenu in the Viewport Options dropdown.
+
+![image alt text](image_73.png)(w:929 h:293 convert:false)
+
+The stats are cumulative and non-hierarchical, so you can see the major categories without having to dig down through a tree of events. For example, shadow projection is the sum of all the shadow projections for all lights (across all the views). 
+
+The on-screen GPU stats provide a simple visual breakdown of the GPU load when your title is running. They are also useful to measure the impact of changes instantaneously; for example when changing console variables, modifying materials in the editor or modifying and recompiling shaders on the fly (with *recompileshaders* changed). 
+
+The GPU stats can be recorded to a file when the title is running for analysis later. As with existing stats, you can use the console commands 'stat startfile'* *and 'stat stopfile' to record the stats to a ue4stats file, and then visualize them by opening the file in the Unreal Frontend tool.
+
+![image alt text](image_74.png)(w:929 h:402 convert:false)
+
+*Profiling the GPU with UnrealFrontend. Total, postprocessing and basepass times are shown*
+
+Stats are declared in code as float counters, e.g:
+
+DECLARE_FLOAT_COUNTER_STAT(TEXT("Postprocessing"), Stat_GPU_Postprocessing, STATGROUP_GPU);
+
+Code blocks on the rendering thread can then be instrumented with SCOPED_GPU_STAT macros which reference those stat names. These work similarly to SCOPED_DRAW_EVENT. For example:
+
+SCOPED_GPU_STAT(RHICmdList, Stat_GPU_Postprocessing);
+
+GPU work that isn't explicitly instrumented will be included in a catch-all [unaccounted]* *stat. If that gets too high, it indicates that some additional SCOPED_GPU_STAT events are needed to account for the missing work. It's worth noting that unlike the draw events, GPU stats are cumulative. You can add multiple entries for the same stat and these are aggregated across the frame.
+
+Certain CPU-bound cases the GPU timings can be affected by CPU bottlenecks (bubbles) where the GPU is waiting for the CPU to catch up, so please consider that if you see unexpected results in cases where draw thread time is high. On PlayStation 4 we correct those bubbles by excluding the time between command list submissions from the timings. In future releases we will be extending that functionality to other modern rendering APIs.
+
+## New: Improved Merge Actor Texture Atlas (experimental)
+
+We've improved how **texture space is utilized** when merging actors together and combining materials, by introducing a new option to generate a **weighted (binned) atlas texture**, instead of having a atlas-texture in which each material is equally weighted. 
+
+![image alt text](image_75.png)(w:286 h:286 convert:false)![image alt text](image_76.png)(w:286 h:286 convert:false)
+
+*(Left: Equal-weighted materials.  Right: Binned method)*
+
+The new functionality first calculates the importance of an individual material according to the *maximum *sized texture it samples.  These values are then used to calculate the amount of space the material should occupy and to iteratively add each texture to the atlas-texture.  This makes the atlas-texture more representative for the input materials because it takes the original texture data size into account for the resulting space it occupies in the atlas texture.
+
+## New: Android Support on Linux
+
+Thanks to pull requests from the community we now have Android support for Linux. CodeWorks for Android from NVIDIA for Linux is the easiest way to set up the NDK and SDK tools needed. In addition, OpenJDK 1.8 should be installed; set JAVA_HOME to point to your install. Please note Android Vulkan on Linux is not supported at this time.
+
+## New: Media Player Editor
+
+You can now drag and drop files from your computer into the Media Player's viewport, thus allowing you to preview video files without having to create a FileMediaSource asset first
+
+![image alt text](image_77.jpg)(w:664 h:561 convert:false)
+
+*  A tab for decoder performance statistics has been added. The output depends on the player plug-in being used for playback
+
+	![image alt text](image_78.png)(w:313 h:365 convert:false)
+
+## New: VR Multiview Support for Mobile
+
+You can now use the mobile multiview path on supported devices!  Mobile multiview is similar to instanced stereo on the desktop, and provides an optimized path for stereo rendering on the CPU.  
+
+To use this feature, enable it in your Project Settings, under the VR section.  For the feature to work, Android build settings should be set to OpenGL ES2, Mobile HDR should be disabled, and instanced stereo should be disabled.  Currently, the feature is compatible with modern Mali-based GPUs.  If you package with the feature on, but don't have a compatible GPU, it will be disabled at runtime. 
+
+This feature is still considered experimental as we verify compatibility with more features and devices.
+
+![image alt text](image_79.jpg)(w:929 h:84 convert:false)
+
+Enable Mobile Multiview in your Project Settings menu as seen above, then restart your editor for the changes to take effect.
+
+## New: Layer Support for SteamVR and PSVR
+
+Layer support has now been added for both SteamVR and PSVR!  It works exactly like it does for the Oculus Rift plugin, using the Stereo Layer component.
+
+## New: VR Loading Movies
+
+The engine now supports loading movies on Oculus, GearVR, SteamVR, and PSVR.  These run on the rendering thread, and can mask framerate hiccups as you load up your content.
+
+To use the splash screen, you can set a texture using the "Set Splash Screen" node, or choose to automatically have it appear when you load a map with the "Enable Auto Loading Splash Screen" node.
+
+## New: PSVR Support for Multiple Framerate Targets
+
+We now support native 90Hz to 90Hz reprojection and 120Hz to 120Hz reprojection on the PSVR!  This means you can opt in to running at a higher framerate to minimize latency and reprojection artifacts.  The engine will limit your framerate to your selected option, but it's still your responsibility to make sure you consistently maintain that framerate!
+
+![image alt text](image_80.jpg)(w:929 h:158 convert:false)
+
+## New: Multitouch Support in Windows
+
+Touch events will now be generated in Windows 7, 8, and 10 when using a touch screen.  This will enable touch enabled games and experiences on new Windows tablets and also enables testing touch controls for mobile games without having to deploy to a target device.
+
+## New: Fast Asynchronous Loading System (experimental)
+
+Cooked builds can now use a completely new Event Driven Loader which is far more efficient than the old streaming code. Games using the EDL should see the load times drop by about 50%, but in many cases it may be dramatically faster.
+
+The Event Driven Loader comes with an unified code path for loading assets. This means that all packages will be loaded using the new async path instead of the old blocking path.
+
+EDL is currently an experimental feature and disabled by default, but can easily be enabled through Project Settings.
+
+![image alt text](image_81.jpg)(w:929 h:318 convert:false)
+
+## New: Simplified Game Mode and Game State Classes
+
+We've added the new Game Mode Base and Game State Base classes as parents of existing classes Game Mode and Game State. Core features needed by all games are now in the Base classes, while legacy and shooter-specific features are in the Game Mode/Game State. Newly-started projects will inherit from the Base classes, while existing projects will default to using their legacy counterparts. Additionally, some new functions have been exposed to Blueprints.
+
+This change is part of an ongoing effort to update older gameplay classes to be easier to understand and subclass for projects of all types. All of the samples other than Shooter Game have been updated to use the Base classes, while Shooter Game shows how to use some of the more shooter-specific features in the legacy classes. Game Mode/Game State will continue to be supported, and you should subclass the version which is more appropriate for your game.
+
+## New: Faster Network Replication
+
+An internal re-factor was done to how we replicate properties from the server to connected clients. 
+
+We have modified the code to more efficiently share property replication work across many connections. Before this change was made, we used to only share the work of unconditional properties (properties registered with DOREPLIFETIME rather than DOREPLIFETIME_CONDITION). We now share the work for all types of properties. What this means, is the work we do to check when properties have changed (and need to be sent), happens way less often.
+
+This has shown to have an increase internally by as much as 40%. We have more work to do here, but wanted to share this good news and progress with you in the meantime!
+
+# Release Notes
+
+### AI
+
+* New: Added an "Enable AI System" flag to the World Settings that controls whether the AI System will be created for the given map.
+
+* Bugfix: Fixed an issue with AI Sense Sight's time slicing unintentionally skipping queue-aging if the given time limit was reached.
+
+* Bugfix: Fixed an issue basing whether the AI should turn or not by comparing the current desired rotation vs the last update's desired rotation. The AI should be turning based on whether or not the current desired rotation is different from the rotation of the Pawn.
+
+* Improvements for the consistency of loudness usage in AI Sense Hearing.
+
+* Minor cosmetic fix to EQS test scoring function's preview.
+
+* Selecting Geometry trace mode will no longer project onto the Nav Mesh first in the Projected Points EQS generators.
+
+#### Behavior Tree
+
+* New: Expanded the BT Decorator Is At Location function with an option to use AI Data Provider.
+
+* New: Multi-node comments are now supported in the Behavior Tree editor. This is similar functionality to the Blueprint editor.
+
+* Bugfix: Fixed an issue with BT Decorator Compare BB Entries function so that it works regardless of the order of the supplied BB keys.
+
+* Bugfix: Blackboard key synchronization mechanics correctly handle more edge cases and complex inheritance scenarios.
+
+* Bugfix: Fixed an issue in propagation of persistent Blackboard keys resulting in these keys missing or being doubled in child Blackboard assets.
+
+* Bugfix: Fixed an issue for derived Blackboard sometimes getting post-load mechanics run before their parents. 
+
+#### Navigation
+
+* Bugfix: Aborting previous movement requests as part of requesting a new one no longer resets the agent's current velocity. 
+
+* Bugfix: Blueprint "AI Move To" function no longer uses the default AI Controller's navigation filter. 
+
+* Bugfix: Fixed navigation export for landscape not handling mirroring and async gathering correctly. 
+
+* Bugfix: Fixed an issue with navigation export for the capsule component.  
+
+* Bugfix: Navlink's "Snap to cheapest area" mode now works correctly with dynamic navmeshes.
+
+* Bugfix: Navmesh generation no longer gets stuck in an infinite loop.
+
+* Bugfix: Fixed an issue with parameters of vehicle RVO avoidance.
+
+* Bugfix: Fixed an issue with Path Following Component for "Has Reached" not using the Goal's Radius when "Use Nav Agent Goal Location" is set to false. When "Use Nav Agent Goal Location" is set to true, we want to avoid using the Goal's location on the navmesh, but instead the acceptability of a radius should be based on the Goal's radius. 
+
+* Path Following Component distinguishes between partial and full paths in terms of acceptance radius used. This will make partial paths behave as expected with non-trivial big acceptance radii. 
+
+* Layer limit for navmesh generation has been increased. 
+
+* Update Move Focus will only clear the focus if the path following component is idle. This keeps the AI rotated in the correct direction when movements get paused. 
+
+### Animation
+
+* New: Added "Strip Animation Data On Dedicated Server" option to Animation Settings. This will remove all compressed data from cooked server data. Disabled by Default.
+
+* New: Added "Aim Offset Look At" node. Aim Offset node that drives its inputs automatically from a Target Location and a Source Socket.
+
+* New: Added Conversion settings to the Alembic import process. This enables the following:
+
+    * Inverting UV Channels.
+
+    * Applying a conversion matrix to the Alembic File.
+
+    * Presets for Autodesk Maya and 3DS Max. 
+
+* New: Added "LegIK" Anim node, which supports two or more bones per limb with min compression rotational constraint.
+
+* New: Added logic to separate overlapping notifies onto separate tracks. Notifies can no longer end up hidden from the user by occupying the same time on the same track. 
+
+* New: Added the option to copy selected Morph Target names to clipboard.
+
+* New: Added a start time to Montage Play node and Play Slot Animation as Dynamic Montage.
+
+* New: Added the ability to skip empty "preroll" frames when importing an Alembic file and improved the frame information for the Import Options UI. 
+
+* New: The animation "Sequence Length" property can now be read in Blueprints. 
+
+* New: Added the ability to open Anim Montage to open asset from the context menu of the segment. 
+
+* New: Added ability for selected Morph Target to view the modified vertices while in the Persona.
+
+* New: Added option to now insert current pose to the Pose Asset with Insert Pose.
+
+* New: Added SmartName Deterministic Cooking:
+
+    * SmartName discards GUID during cooking and DisplayName becomes the identifier during cooking. 
+
+    * The GUID created nondeterministic cooking causing patch size to become huge without any actual data change. 
+
+* Bugfix: Fixed a crash when attaching slave components with differing bone counts.
+
+* Bugfix: Fixed a crash with retargeting Additive Animation Montage.
+
+* Bugfix: Fixed a crash when routing a wire through multiple re-route nodes in an Animation Blueprint.
+
+* Bugfix: Fixed a crash in "Get Post Curve" when importing a pose to pose asset.
+
+* Bugfix: Fixed a crash when retargeting Anim Blueprint with "Allow remapping to existing assets" was enabled.
+
+* Bugfix: Fixed a crash when calling "Set Animation Mode" on a component with no Skeletal Mesh. 
+
+* Bugfix: Fixed an issue when re-creating a skeleton due to a merge conflict it will add all Skeletal Meshes back to the skeleton so that it recreates all hierarchy. 
+
+* Bugfix: Fixed an issue with invalid bound calculation in Calculate LOD Count. 
+
+* Bugfix: Ensure the runaway loop counter gets reset when processing parallel animation. 
+
+* Bugfix: Fixed an issue for fast path struct copy being broken for Vectors.
+
+* Bugfix: Fixed an issue for warning triggered by calling "Get Bone Transform" before the master pose component has been registered
+
+* Bugfix: Fixed an issue where layered blend per bone odd/even connection counts alternately working/not working
+
+* Bugfix: Fixed an issue with offset of the floor mesh in Persona when Auto-Alignment was enabled. 
+
+* Bugfix: Fixed an issue with pose flickering on LOD change when using "Layered Blend by Bone" node.
+
+* Bugfix: Fixed an issue to stop empty states being created when dragging montages into State Machines. 
+
+* Bugfix: Fixed an issue with Anim Instance "Is Running Parallel Evaluation" check to correctly test if the Skeletal Mesh Component still references the Anim Instance being used. 
+
+* Bugfix: Fixed an issue with "generated" string appearing for a Skeletal Mesh LOD that was previously auto-generated but had been reimported from an FBX. 
+
+* Bugfix: Fixed an issue with displaying one more frame for an animation.
+
+    * For example, for 30 frames, it was displaying 0-30, which is incorrect and should be 0-29.
+
+    * To avoid confusion, now it doesn't display key count. It will display the last frame number, so it will look different, but the data is still the same. 
+
+* Bugfix: Fixed an issue with a number of cases where sub-instance animation instance functions weren't being correctly passed through the main instance.
+
+* Bugfix: Fixed an issue with "auto" check box in the Morph Target panel.
+
+* Bugfix: Fixed an issue with functionality of "show uncompressed animation" option in Persona viewports.
+
+* Bugfix: Fixed an issue with entries in the Anim Curves tab losing "auto" state when hidden/reshown.
+
+* Bugfix: Removed errant asterisks in the tooltips on the "Rotation Multiplier" node. 
+
+* Bugfix: Fixed incorrect tooltips on the "Hand IK Retargeting" node.
+
+* Bugfix: Fixed an issue in montages where the wrong section was updated with changes from a Details panel when clicking on a different section.
+
+* Bugfix: Fixed an issue when duplicating an Anim Sequence that has been layer-edited would apply the layer twice. 
+
+* Bugfix: Fixed an issue where Geometry Cache objects would not be visible in Standalone (packaged) builds.
+
+* Bugfix: Fixed an issue with keyboard shortcuts for Pose Asset editor panel.
+
+* Bugfix: Fixed an issue with keyboard shortcuts in the Curve Viewer panel.
+
+* Bugfix: Fixed an issue with notify states not ending when the Animation Instance was cleared. 
+
+* Bugfix: Fixed an issue with orphaned transition nodes being created when they were duplicated or pasted without two valid pin links. 
+
+* Bugfix: Fixed an issue with Override Materials hanging around when setting Skeletal Meshes.
+
+* Bugfix: Fixed an issue with pin defaults on sub-instance nodes overwriting the defaults on the inner class.
+
+    * The class defaults will now be copied out to the pin.
+
+* Bugfix: Fixed an issue with pose asset blending.
+
+    * Normalization has been removed for weights that are less than one so the poses do not instantly snap.
+
+* Bugfix: Fixed an issue where Pose Asset was evaluating the curve in the wrong frame. 
+
+* Bugfix: Fixed an issue where the Pose Asset editor list was not updating when changing the source animation.
+
+* Bugfix: Fixed an issue with "Remove Linear Key" for Additive Animations since DDC has been added. 
+
+* Bugfix: Fixed an issue with "Root Motion from everything" calculating incorrect root motion when animations haven't been ticking. 
+
+* Bugfix: Fixed an issue with sample rate not applied for non-default sample rate when importing an animation.
+
+* Bugfix: Fixed some inconsistencies in IK editing in the viewports inside Persona when changing effector spaces. 
+
+* Bugfix: Fixed an issue with split struct pins on animation nodes sometimes losing their links on editor restart. 
+
+* Bugfix: Fixed an issue where sub-instance automatic properties would appear in the Details panel. 
+
+* Bugfix: Fixed a typo in "Root motion from everything" accumulation code where we calculated root motion weight using the wrong variable.
+
+* Bugfix: Fixed an issue where Undoing and Redoing would force the Skeletal Mesh into T-pose.
+
+* Bugfix: Fixed an issue when importing Morph Targets with fixed frame count.
+
+* Bugfix: Fixed an issue with Blendspace triangulation.
+
+* Improved editor response during dragging on additive frame index (stop animation compression happening for each mouse move).
+
+* Non-POD structs can now be copied using the 'fast path' in animation blueprints.
+
+* Prevented invalid GUIDs from being saved into SmartName containers.
+
+* Removed save warning about empty animations from assets that contain curve data.
+
+* Animation Curve Cached UID list has now been moved to Bone Container from Anim Instance in preparation of supporting curve linking to bone feature, so that you could cull out for a certain LOD.
+
+* Changed "Create Physics Asset" in the context menu for Skeletal Meshes to enable creation of a Physics Asset without having to assign it to the Skeletal Mesh. This makes it easier to make new assets for Capsule Shadows. 
+
+* Changed inline curve name editing to only change the name of that specific curve, instead of renaming the SmartName itself.
+
+#### Animation Blueprint
+
+* Bugfix: Fixed a crash using anim debug with sub-instances that are preceded by branching nodes.
+
+* Bugfix: Fixed a crash when forcing ref pose during Anim Blueprint initialization when using sub anim instances.
+
+* Bugfix: Fixed a crash when using sub-instances in non-entry states.
+
+#### Tools
+
+* New: Added Skeletal LOD workflow refactor. 
+
+    * Now we track source files for LODs to save time when reimporting LODs often.
+
+### Audio
+
+* New: Added a debug command to solo specific sound waves, sound cues, and sound classes.
+
+* New: Audio thread is now enabled by default.
+
+* New: Gameplay Statics' "Create Sound 2D" and "Spawn Sound 2D" now have a boolean parameter that specifies the sound should persist across level transitions.
+
+    * Note that the returned component must continue to be referenced by something (Game Instance is a good option) or else it will be garbage collected and the sound will stop.
+
+* New: Looping sounds are now searchable by the asset registry.
+
+* New: Making ability to suppress subtitles a UPROPERTY.
+
+* New: Added support for parsing .wav files with a WAVE_FORMAT_EXTENSIBLE format chunk.
+
+* New: Added support for pausing individual audio components.
+
+* Bugfix: Fixed a crash which would occur when copying and pasting an undefined custom sound node from one project to a project which doesn't have the custom node.
+
+* Bugfix: Fixed an issue for one-shot sound wave instances triggering unintentionally from active sounds when volume mix is 0.
+
+* Bugfix: Fixed an issue for sound cue delay not consuming startime correctly.
+
+* Bugfix: Fixed an issue with attenuation on 2D multi-channel sound sources. 
+
+* Bugfix: Fixed an issue with concurrency resolution rule stop quietest with looping sound waves.
+
+* Bugfix: Fixed an issue with sound node distance cross fade for case of looping sounds.
+
+* Bugfix: Fixed an issue where Undo in Sound Cue editor does not undo looping changes. 
+
+* Bugfix: Fixed various thread safety issues with occlusion on the audio thread. 
+
+* Hooked up compression quality settings to AT9 files.
+
+* Prevent procedural sound waves from stopping if they generate 0 bytes of data.
+
+* Read Compressed Info calculates duration for ADPCM audio.
+
+### Automation
+
+* New: Added a json summary report file generation as part of all automated test summaries.
+
+* New: Functional tests now support warnings as errors on a per test basis.
+
+* New: Included .tps files in the installed build to make it easier for users to identify third party software in use. Also included source for additional programs to Installed Build. 
+
+* New: Added Documentation/Extras in Installed Build.
+
+* Bugfix: Fixed an issue with Test Equal Add Error Message in Automation Test Base
+
+* Bugfix: Fixed issues with Mobile Device Interface in Installed Builds. 
+
+* Bugfix: Fixed an issue with Android sessions not appearing in the Session Frontend.
+
+* Bugfix: Fixed support for PS4 in Installed Builds. 
+
+* Bugfix: Fixed various connection reliability issues with TCP Message Transport used for automation testing on mobile devices. 
+
+* Bugfix: Fixed various issues for PS4 deployment.
+
+* Bugfix: Fixed support for Xbox for Installed Builds created via Build Graph script. 
+
+* Automation - Presets are now stored in json files stored in Config so they can be shared, and human readable. Working on screenshot automation, getting it where it needs to be to permit us to have repeatable tests for comparison.
+
+* Each session launched from the Session Frontend Launch tab now has a unique GUID and they will correctly appear as separate sessions in the Automation tab.
+
+* Excluded some more intermediates from the Installed Build.
+
+* Functional Tests now run correctly in non-editor build.
+
+* Functional tests will no longer move the camera if the user is using a debug camera.
+
+### Blueprints
+
+* New: Added a new "Is Standalone" node which returns true if the game is running in Standalone.
+
+* New: Added ability to change curves used by a Timeline at runtime.
+
+* New: Added an option in Editor Preferences Blueprint section to "Display Unique Names for Blueprint Nodes" rather than the node type.
+
+* New: Blueprint Nativization: Reduced size of the executable file.
+
+* New: Added ability to flag Components added via Blueprints as editor only.
+
+* New: Exposed "Get Component By Class" to Blueprints.
+
+* New: Exposed the method to "Is Any Rigid Body Awake" for primitive components to Blueprints. This enables a Blueprint to determine if Physics Simulation is being done on the object and fulfills a request to determine if an object is "awake" or "asleep".
+
+* New: Replication condition property flags are now exposed to Blueprints. This can be set in the Properties Details panel.
+
+* New: Single layout Blueprint is no longer an experimental feature. The option to disable it has been removed. 
+
+* New: Widget and Animation rich diff tool is no longer experimental.
+
+    * Some changes are not detectable by the diff tool, but they are mentioned in a warning to the user. 
+
+* New: Enabled instances of Hierarchical Instanced Mesh Components to be moved around with the transform widget in the Blueprint Editor, just like Instanced Mesh Components!
+
+* Bugfix: Fixed a crash in the Blueprint Profiler when compiling or profiling that had the following: 
+
+    * Contains a loop.
+
+    * Nested Macros.
+
+    * Collapsed Nodes.
+
+    * Called parent functions or events. 
+
+* Bugfix: Fixed a crash caused when compiling a Blueprint after hot-reloading its base class.
+
+* Bugfix: Fixed a crash when compiling particularly cross-dependent Blueprints. 
+
+* Bugfix: Fixed a crash when trying to pick a debug world for the Blueprint debugger. 
+
+* Bugfix: Fixed a crash when implementing an interface to a child Blueprint and then implementing it with a parent Blueprint. 
+
+* Bugfix: Fixed a crash that could occur when reparenting a child Blueprint class that had inherited its default scene root component from the parent class. 
+
+* Bugfix: Fixed a crash that would occur on a right-click in the Blueprint editor's component tree view after having copied one or more BSP actors to the clipboard. 
+
+* Bugfix: Fixed a crash that could happen when pasting a copied "Select" node. 
+
+* Bugfix: Fixed a crash on load in "Remove Node and Promote Children" when removing a corrupted SCS node if it has no parent link (the children are moved to the root node instead).
+
+* Bugfix: Fixed a rare crash in "Get Default Value Visibility" after removing a pin from a Blueprint node. 
+
+* Bugfix: Fixed a crash when one attempts to generate an abstract class from a Blueprint. 
+
+* Bugfix: Fixed a crash in FCDODiffControl when CDOs being compared had different numbers of properties. 
+
+* Bugfix: Fixed a crash that could occur when pasting graph nodes with orphaned connections to a different graph. 
+
+* Bugfix: Fixed a crash when deleting many Blueprint assets.
+
+* Bugfix: Fixed an issue where a function, created from collapsed nodes, cannot have a parameter of a weak pointer type. 
+
+* Bugfix: Fixed an issue with assets referenced by a function's local variables won't be garbage-collected. 
+
+* Bugfix: Fixed an issue where the Blueprint Editor viewport now correctly displays the result of an undo action that changes the default transformation of the Blueprint Actor. 
+
+* Bugfix: Fixed an issue where changes in User Defined Structure will be properly applied in local variables, without resetting the values. 
+
+* Bugfix: Fixed circular dependencies introduced by default values on User Defined Structs that are now loaded correctly. This fixes on cause of data resetting to default on load. 
+
+* Bugfix: Fixed an issue where compiling a Blueprint would trigger the World Outliner to sometimes detach. 
+
+* Bugfix: Fixed an issue where entering "foreach" in the Blueprint Editor context menu will not result in "For Each Loop" being the first selection.
+
+* Bugfix: Fixed an issue where moving a variable from a child class to a parent class would break variable Blueprint nodes. 
+
+* Bugfix: Fixed an issue in the Blueprint context menu where options would be unavailable when you hid a sub-level.
+
+* Bugfix: Fixed an issue so that the "Create Save Game Object" node will now change its output type to match the class you select. 
+
+* Bugfix: Fixed an issue that caused comment bubbles to display incorrectly when using Play-In-Editor.
+
+* Bugfix: Fixed an issue where functions could lose inputs after adopting them from an interface. 
+
+* Bugfix: Fixed an issue that would cause split-pins on latent "Gameplay Task" nodes to lose connection.
+
+* Bugfix: Fixed an issue that could hang the editor indefinitely during Blueprint Compilation. 
+
+* Bugfix: Fixed an issue that prevented the option to split a struct pin from appearing on a "Get Variable" node.
+
+* Bugfix: Fixed an issue that prevented you from ctrl-dragging pin links from macro nodes. 
+
+* Bugfix: Fixed an issue that would cause subsequent latent nodes to fail to execute once one was running. 
+
+* Bugfix: Fixed an issue that would prevent saving a Blueprint after compiling a component that it used. 
+
+* Bugfix: Fixed an issue where Blueprints were not properly inheriting their associated config file. 
+
+* Bugfix: Fixed an issue with rotation struct input pins where splitting them would cause each of the values to shift to the wrong axis. 
+
+* Bugfix: Fixed an issue in which data loss could occur for inherited component default value overrides within a child Blueprint class after recompiling a Blueprinted component class. 
+
+* Bugfix: Fixed an issue that could lead to data loss at runtime when choosing to enable the optimized Blueprint component instancing path for a cooked build. 
+
+* Bugfix: Fixed an issue that could lead to per-instance data loss in Blueprint-constructed components after renaming component variables in the components tree view. 
+
+* Bugfix: Fixed an issue that would lead to Blueprint class compile errors caused by missing client-only component class reference properties in an uncooked server-only context (e.g. UE4Editor.exe -server) 
+
+* Bugfix: Fixed an issue where "Custom Event" was no longer the top option when searching the Blueprint menu for "custom."
+
+* Bugfix: Fixed the directions of parameters in functions like "Get Overlap Infos."
+
+* Bugfix: Fixed an error when output references parameters of Blueprint Native Event are not in a specific order. 
+
+* Bugfix: Fixed failed "ensures" when pasting linked Anim Blueprint nodes. 
+
+* Bugfix: Fixed issues that could lead to per-instance component data loss after reconstruction of a Blueprint class containing one or more "Add Component" nodes in the Construction Script.
+
+* Bugfix: Fixed issues with Object Library. It now stores "Blueprint Generated Graph" object instead of "Blueprint" object. 
+
+* Bugfix: Fixed an out-of-order registration issue with native and/or non-scene components and Blueprint-constructed components.
+
+* Bugfix: Fixed an issue with "Time Seconds To String" format so it matches the comment. 
+
+* Bugfix: Fixed a validation of incompatible pines issues. It will now generate an error. 
+
+* Graphs in a Blueprint can properly be deleted even if they were once part of an interface on the Blueprint.
+
+* Inherited components in a child Blueprint class will no longer trigger a warning message when the original component has been removed from its parent class.
+
+* Moving a child component in a child Blueprint doesn't force a parent to become dirty anymore.
+
+* Nested Default Sub-objects no longer trigger an assertion failure in "Static Allocate Object" when used by a Blueprint type.
+
+* Reduced PIE startup time significantly when auto-recompiling a Blueprint class with multiple dependencies.
+
+* Removed warning that would complain about "Make Struct" nodes that didn't have any properties exposed, enabling Blueprint users to return default struct data.
+
+* "Show World Context Object" markup is now fully inherited. Previously only inherited for one level of a hierarchy.
+
+* Stale breakpoints saved in Blueprint assets no longer trigger a warning on load. Message moved to 'Display' channel.
+
+* Struct Ref Pins in "Set Fields In Struct" node can no longer be split when optional pins as this created ambiguous behavior.
+
+* Blueprint function "Get Direction Vector" has been renamed "Get Direction Unit Vector" to reflect details of its result.
+
+* Can no longer drag child actor from World Outliner to a Level Blueprint and create an invalid actor reference.
+
+* Default event nodes in child Blueprints now wire up parent call inputs so you can use them as is without breaking functionality.
+
+* Duplicating levels that have Blueprints with "Knot" nodes attached to events will no longer fail with a "Graph is linked to external private object" error.
+
+* EaseIn/EaseOut descriptions corrected to match actual behavior.
+
+### Core
+
+* New: Added a config based system for specifying which classes should be stripped on dedicated server and client builds.
+
+    * Class names can be specified in the project cooking settings in the editor
+
+    * All engine classes that relied on "Needs Load For Server" or "Needs Load For Client" have been moved over to the config based system.
+
+    * "Needs Load For Server" and "Needs Load For Client" still function as they previously did.
+
+* New: Added a new file into the ini hierarchy to begin fixing the Engine/Base > Project/Default > Engine/Platform > Project/Platform confusion.
+
+    *  We now have Engine/Base -> Engine/BasePlatform -> Project/Default -> Engine/Platform -> Project/Platform. However, Engine/Platform will soon be deprecated as we move things over to Engine/BasePlatform, that are safe to move.
+
+* New: Added ability for arrays of structs to specify a property to be the key. 
+
+    * So, with LODGroups, the Name key inside the struct can be the unique key, so when you have multiple .ini files in the hierarchy overriding the same LODGroup by name, it will replace the first with the second, instead of adding two entries with the same name. Set by @ArrayName=KeyPropertyName. Per Object Config sections need a little different handling, which uses * (see Base
+
+* New: Added new Module Types: "ClientOnly" and "ServerOnly"enabling better control of when modules are compiled and loaded. 
+
+    * Can be used to slim down server / client builds.
+
+* New: Added some extra information to an assert when "Cpp Structs Ops" has not been initialized.
+
+* New: Added the ability for plugins to reside outside of the Engine/Plugins and Game/Plugins directory. 
+
+    * This is useful for plugins that are shared across game teams, but don't want to be placed in Engine/Plugins.
+
+* New: Accepting comma in milliseconds separator when parsing Timespan.
+
+* New: Minor changes to pak file security delegate. Now all files that don't exist in the pak file will be passed through the security delegate, regardless of whether they exist on disc or not.
+
+* New: Only accept "log=<filename>" and "abslog=<filename>" command line values if the filename has a ".log" or ".txt" extension.
+
+* New: Refactored handling of some UI and rendering code when running headless (i.e. with -nullrhi switch).
+
+* New: Upgrading an existing UProperty from FVector to FVector4 is now seamless and preserves data properly.
+
+* Bugfix: Fixed a rare crash when async loading objects caused by linker being detached too early, before other package's import has been fully processed.
+
+* Bugfix: Fixed a crash on cooker exit due to unpredictable destruction order of static tickable objects and their collection.
+
+* Bugfix: Fixed a crash with potential memory corruption when loading Data Tables. It sometimes would allocate less memory than was required.
+
+* Bugfix: Fixed a crash in PhysX cleanup if an exception occurred during shutdown.
+
+* Bugfix: Fixed a thread safety issue in the load module code that could crash if multiple threads were trying to load the same module at the same time.
+
+* Bugfix: Fixed an issue where "Get All Assets" would not return data about live objects when passing bIncludeOnlyOnDiskAssets=false.
+
+* Bugfix: Fixed an issue with task graph memory leaks.
+
+* Bugfix: Fixed an issue with TMap deserialization in struct serializer and updated unit tests.
+
+* Bugfix: Fixed an issue to display a meaningful error and shutdown cleanly if Core modules fail to load on engine startup.
+
+* Bugfix: Fixed an issue so that dynamic delegates can no longer get double bound when copying an Actor.
+
+* Bugfix: Fixed a potential race condition and minor memory leak in FMonitoredProcess. This could result in the editor hanging on startup.
+
+* Bugfix: Fixed an issue where "Force Garbage Collection" might not force a GC immediately if run on a server with no clients connected. 
+
+* Bugfix: Fixed a memory corruption when instancing subobjects of a TArray property.
+
+    * Fixed a memory overrun when reading an SSL certificate. 
+
+    * Fixed an invalid memory access when shutting down FInputBindingEditorModule.
+
+* Bugfix: Fixed an issue where duplicate log lines appear when commandlet "Log To Console" is set.
+
+* Bugfix: Fixed an issue with FAssetPtr remapping issues for duplicated packages, which was a case with operations like Level Save As and Level Duplicate.
+
+* Bugfix: Fixed an issue with Malloc Binned "Get Allocation Size" for aligned allocations, which caused crashes with the poison proxy.
+
+* Bugfix: Fixed an issue with Parse "Value" when there is whitespace between the command and the argument. 
+
+* Bugfix: Fixed an issue with "Is A" race condition during async loading.
+
+* Bugfix: Fixed an issue where Unreal Build Tool would fail when trying to delete a module's manifest file in a non-existent directory.
+
+* Bugfix: Fixed an issue where newline escape sequence wasn't written to the pipe.
+
+* FScopedCreateImportCounter will now always store the current linker and restore the previous one when it exits. 
+
+    * This should fix some of the linker warning messages reporting wrong package name.
+
+* Invalidate string asset reference tags after finishing up loading of an async package. 
+
+    * Fixes some problems with asset refs not updating correctly after a package was loaded asynchronously
+
+* Lines of output text from FMonitoredProcess are no longer truncated at the default pipe buffer size, and are no longer returned with trailing CR characters on Windows.
+
+* Made sure Skeleton assets are loaded before "PostLoad" is called on them to avoid crashes when async loading animations.
+
+* Silently skip creating exports from a package where the outer is also an export and has been filtered at runtime during loading.
+
+* Thread safety fixes for the async log writer + made the async log writer flush its archive more often.
+
+* Thread safety fixes when using the async loading thread.
+
+* We now abort a directory watch if we lose access to the directory in question to prevent an infinite loop.
+
+### Editor and Tools 
+
+* New: Added a warning message to indicate that HLOD clusters do not have a valid proxy mesh, similar to the "lighting unbuilt" message.
+
+* New: Added a commandlet to enable command line importing of any asset type that the engine supports. 
+
+* New: Added an ability to vsync the editor. This is disabled by default. Set "r.VSyncEditor" to 1 to enable it.
+
+* New: Added an option in the FBX importer to automatically convert the unit in the FBX file to Unreal Units. 
+
+* New: Added asset cooking path length in their tooltips so that you can detect cooking path issue before you are cooking. 
+
+* New Added checkbox for "Show All Advanced Details" config variable so that it will auto-expand all advanced property categories. 
+
+* New added command line option to enable movie capture burnin.Use the commandline -UseBurnIn=yes.
+
+* New Added mesh simplification plugin picker to the Project Setting under Editor > Mesh Simplification. The menu to pick simplification plugins also contains a link to find other plugins in the Epic Games Launcher Marketplace.
+
+* New: Added metadata to remove "reset to default" button for certain properties.
+
+* New: added new shortcut "Ctrl + Shift + B" to build current level geometry. 
+
+* New Added the ability to import grayscale Targa (.tga) files in the editor. 
+
+* New Added the ability to search through all Editor Preferences and Project Settings at once. Use the selection in these windows at the top for "All Settings." 
+
+* New: Added the possibility to preview a wave sound in the Content Browser from the thumbnail. 
+
+* New: Added Undo/Redo functionality to the Editor Preferences and Project Settings windows. 
+
+* New: Changed and unified saved shortcuts across all editors. 
+
+    * Save As is now Ctrl + Alt + S.
+
+    * Save All is now Ctrl + Shift + S.
+
+    * Save Current is now Ctrl + S.
+
+* New: Added shortcut Ctrl+Shift+Space to quickly rotate through different viewports (left, up, front, back ,etc).
+
+* New: Changed the FBX batch import workflow to let the user fix all missing source files before starting the FBX batch import process. 
+
+* New: The Content Browser now supports hiding/toggling specific asset detail columns. 
+
+* New: Added the ability for Ctrl + Middle Mouse drag will give options to switch to different viewports. 
+
+    * Drag to Left, Right, Up, and Down will match the viewport intuitively.
+
+    * Drag to Up-Left is Front, Up-Right is Back, Bottom-Left and Bottom-Right are both for the Perspective viewports.
+
+    * Dragging exceeding a certain distance will jump to that viewport directly whereas dragging a little will bring up a popup menu showing which viewport it's going to switch to. 
+
+* New: Added a new view option for "Show C++ Classes" in the Content Browser to enable you to hide the C++ classes you may not want to see. 
+
+* New: The Engine version including branch name will now be displayed in the Title bar by default. 
+
+* New: Extended "Select all input nodes" function to general Blueprint editor. 
+
+* New: FBX Exporter now exports all LODs.
+
+* New: The Date Time fields in the editor are now displayed and entered in the local user's timezone. 
+
+    * The time entered is saved into assets using the value converted in the UTC timezone. 
+
+    * The timezone that is used is configurable in the Editor Settings  under Region & Language for Internationalization settings. 
+
+* New: There's now support for generating a "Content-Only" plugin from the new Plugin Wizard. 
+
+* New: Added Show Flags for Widget Components. 
+
+* New: Reference viewer now provides options to show/hide soft and hard references and dependencies with different line color.
+
+* New: Added Asset Picker support to the Script Interface properties. 
+
+* New: Added unique asset icons for Sound Concurrency assets.
+
+* New: The Widget Reflector now has a configurable delay between clicking to take a snapshot and taking the snapshot, which is helpful for reflecting widgets in a state that involves the cursor or other timing factors.
+
+* New: Added the ability for users to specify a Base material to create Material Instances instead of just creating materials when importing an FBX. 
+
+* New: Added project file to source control when creating a new project.
+
+* Bugfix: Fixed a crash when moving Instanced Static Mesh Component in editor when it had no mesh set, but had instances. 
+
+* Bugfix: Fixed a crash on undoing a deletion of a mesh in mesh painting mode in the VR Editor. 
+
+* Bugfix: Fixed an occasional crash on editor shutdown due to a race condition in FTcpMessageTransportConnection. 
+
+* Bugfix: Fixed several crashes in the Session Frontend when viewing profiles. 
+
+* Bugfix: Fixed a crash issue where making changes to components in the Blueprint editor when there is an active component in the component visualizer. 
+
+* Bugfix: Fixed an issue where batching optimization for dragging components onto Blueprints. Adding thousands of static mesh components now takes seconds instead of up to an hour.
+
+* Bugfix: Fixed an issue so that comments will now be preserved on "Parameter" nodes when converting them to "Constant" nodes within the Material Editor.
+
+* Bugfix: Fixed an issue that prevented a material from being saved when copying a comment node from a Material Function to a Material.
+
+* Bugfix: Fixed an issue so that the Crash Report Client can now correctly restart Editor sessions that were run from the Launcher. 
+
+* Bugfix: Fixed an issue where the cursor would remain hidden after left-clicking in the Orthographic view. 
+
+* Bugfix: Fixed an issue that would cause Custom keyboard shortcuts to no longer switch between the Command and Control keys on Mac when the editor was restarted. 
+
+* Bugfix: Fixed an issue so that debug lines will now render correctly in orthographic views if their thickness has been specified.
+
+* Bugfix: Fixed an issue where FBX and OBJ files containing a UV channel with no name were not imported correctly into the editor. 
+
+* Bugfix: Fixed an issue so that the FBX importer will correctly re-import a mesh when the combined mesh checkbox was set to true. 
+
+* Bugfix: Fixed an issue with "On Mouse Button Up" not resetting SButton's Brush state from "Pressed" to "Hovered."
+
+* Bugfix: Fixed some issues with the LODGroup workflow:
+
+    * During FBX Import, LODs are now recomputed with the correct settings. 
+
+    * In the Static Mesh editor, Extra LODs are removed when changing the LODGroup. 
+
+* Bugfix: Fixed various crashes when manipulating BSP geometry in the editor. 
+
+* Bugfix: Fixed a crash while importing FBX scenes if objects contain procedural textures, which is not supported. 
+
+* Bugfix: Fixed a crash when setting a closed loop on a Spline component which has no spline points defined. 
+
+* Bugfix: Fixed an issue where the dialogue "Actors are referenced, are you sure you want to delete?" from appearing when you are copying a lot of actors to another level.
+
+* Bugfix: Fixed a warning that would spam the log when opening the Editor's Blueprint menu. 
+
+* Bugfix: Fixed an issue where an actor deletion while in Mesh Paint mode would not be undoable. 
+
+* Bugfix: Fixed an issue where all collision would be disabled if you don't auto-generate a simple hull when importing an FBX. 
+
+* Bugfix: Fixed an issue where importing a Texture asset while an editor window was open for that asset would cause the Details Panel to disappear.
+
+* Bugfix: Fixed an issue where occasionally changing tabs in the Editor would cause the UI to freeze briefly and undock the selected tab.
+
+* Bugfix: Fixed an issue where painting override vertex colors on a SpeedTree mesh would cause its wind animation to cease or look incorrect.
+
+* Bugfix: Fixed an issue where structs containing enum values with whitespace characters would not correctly serialize when copied from the Details Panel.
+
+* Bugfix: Fixed an issue where Unreal Frontend and other standalone tools would not display tooltips on dropdown menu items.
+
+* Bugfix: Fixed an issue with assigning of materials on components would not work if you multi-select multiple components of a Blueprint and assign the Material. 
+
+* Bugfix: Fixed an issue when choosing "Yes to All" was not functioning properly when asked to delete many actors. 
+
+* Bugfix: Fixed an issue  with in-game console command parameter display to include all parameters. 
+
+* Bugfix: Fixed an issue with loctext related warnings when opening the Blueprints submenus on the Level Editor toolbar. 
+
+* Bugfix: Fixed an issue with missing widgets when using Mobile Preview mode with a Mobile HDR disabled project. 
+
+* Bugfix: Fixed an issue with single player PIE so that the window position is correctly fetched and saved when running a dedicated server. This is held separately to stored positions for multiple client PIE.
+
+* Bugfix: Fixed an issue in the output log filter not updating when text was committed. 
+
+* Bugfix: Fixed an issue where Mobile Packaging Wizard will not find any project maps if UE4 and the Project are installed on different drives. 
+
+* Device button text in "Device Output Log" window should be more readable now. 
+
+* Bugfix: Fixed an issue so that .exe icon is correctly set for packaged Windows games. 
+
+* Bugfix: Fixed an issue so that String Asset Reference properties will now only open an Asset picker (not an actor picker) as the goal is to reference an asset.
+
+* Bugfix: Fixed an issue where "find culprit" dialog did not honor the user's choice. 
+
+* Bugfix: Fixed an issue with Unreal Game Sync where database events were being returned out of order could cause builds to show as never completing. 
+
+* Bugfix: Fixed an issue with Unreal Game Sync showing a build failure/warning dialog immediately if you submit while the build is yellow/red. This now ignores builds before your last submitted CL. 
+
+* Bugfix: "Play in Standalone Game" does not assume -window as a launch parameter therefore you can now specify -window or -fullscreen.
+
+* "Select all input nodes" now have a corresponding "Select all output nodes" for nodes that are output nodes.
+
+* Deleting many assets at once will now batch them as a single source control command instead of one for each asset.
+
+* Directional Light icon in the viewport changed (the two arrows within the icon were removed), as the asset itself provide another arrow indicating the light direction.
+
+* "HighResShot" console command now supports the same functionality as the Editor UI version and improved the help with usage examples. 
+
+* Improved editor hitching while using launch on by delaying package save until all data is cached. 
+
+* Minor update to the layout of the User Defined Struct editor. 
+
+* Improvements have been made to the Unreal Frontend Project Launcher Progress page:
+
+    * The custom profile name is now visible above the progress bar when launching a profile. 
+
+    * The Task and Status views can now be resized. 
+
+    * Columns for "Errors" and "Warnings" have been added to the Task view to show which tasks have problems. 
+
+    * Previously completed tasks are no longer cleared when the launch is cancelled.
+
+    * The message log will now only scroll to the newest message when the log is scrolled to the very bottom. 
+
+* Removed the Experimental "Live Editor" feature from UE4
+
+    * Live Editor was an early attempt at automation of editor features using external hardware, such as MIDI.
+
+    * This feature was never fully completed and we have no plans to continue work on it.
+
+    * If you'd like to use MIDI devices in your editor plugins or games, please check out the new "MIDI Device" plugin available in this release!
+
+* Unreal Game Sync: Improvements to "Clean Workspace" dialog.
+
+    * Empty folders are now set to be deleted by default.
+
+    * Added a context menu to select all/none/empty/default for a subtree, as well as to open Windows Explorer at that location. 
+
+* We now export the collision mesh when exporting to FBX. Disabled by default. 
+
+    * Added an Editor Preference option to enable or disable the export of collisions. Editor Preferences > General > FBX Export Collision Mesh. 
+
+* Focusing an Actor that has 0 size bounds now works.
+
+* Play in editor will no longer auto resume a game on focus received if it was previously paused.
+
+* Rebuild lightmaps commandlet now doesn't apply world transforms twice.
+
+* Removed duplicate platforms from deploy list in Unreal Frontend.
+
+* Scalability settings will no longer auto-apply if the low performance warning message is ignored.
+
+* Source file paths are now normalized when caching their state.
+
+* The Auto Distance Error value for Static Mesh LODs can now be set to a value smaller than 1.
+
+* The Hue value displayed in tooltips for colors within the Theme bar will now match the hue value displayed for that color in the Color Picker.
+
+* The orbit camera should no longer reverse zoom direction when zoomed with the mouse wheel.
+
+* The performance when selecting thousands of actors has increased and should no longer render the editor unusable when lots of actors are selected.
+
+* The pin-type filter now filters on both the localized and source type descriptions.
+
+* The placement mode now filters on both the localized and source item descriptions.
+
+* Undo/redo of mobility changes will also undo/redo the mobility changes on ancestors/descendants that were changed along with it.
+
+* Visual logger icon updated.
+
+* When importing a FBX file containing multiple animations, user can re-import any animation without overriding it with the first animation in the FBX file.
+
+* Updated Map and Set Properties so they can now be edited from within the Details Panel. Sets ensure that all of its elements are unique, and Maps ensure that all of its keys are unique.
+
+#### Content Browser
+
+* New: Deleting folders in the Content Browser now removes the folder from disk.
+
+* Bugfix: Fixed a crash renaming a collection with multiple collections selected.
+
+* Bugfix: Fixed a crash when force deleting a blueprint with a child actor component from the Content Browser.
+
+* Bugfix: Fixed an issue with dynamic collections not returning anything.
+
+* Added ~ to the list of invalid characters for object/package names.
+
+#### Landscape
+
+* New: Added 180° rotation support to the Landscape Mirror tool.
+
+* New: The Landscape LOD visualization mode now renders material tessellation while in the Viewport is in Wireframe mode.
+
+* Bugfix: Fixed a crash entering Landscape mode when there is a Landscape in a sublevel that is loaded but hidden.
+
+* Bugfix: Fixed a crash editing Landscapes with the "r.LightPropagationVolume" console variable enabled.
+
+* Bugfix: Fixed a crash when applying a material with tessellation enabled to a landscape or when painting a material with tessellation on a newly added component.
+
+* Bugfix: Fixed a crash when duplicating a level containing a Landscape through the Content Browser.
+
+* Bugfix: Fixed a crash when loading duplicated Landscape levels.
+
+* Bugfix: Fixed a crash when merging levels containing a Landscape.
+
+* Bugfix: Fixed a crash with Landscapes caused by using undo after using the Move-To-Level tool and then selecting a Landscape Proxy in a sub-level.
+
+* Bugfix: Fixed a crash if a Landscape Layer Info object is force deleted while it is in use.
+
+* Bugfix: Fixed an inability to import Landscape Layer weight maps when creating a Landscape.
+
+* Bugfix: Fixed an issue that would cause Landscape Components to be invisible when using tessellation on a Landscape.
+
+* Bugfix: Fixed an issue with Landscape rendering errors after using the "Change Component Size" tool.
+
+    * Also affects initial import of a Landscape with weight maps.
+
+* Bugfix: Fixed a regression with performance when importing Landscape heightmaps.
+
+* Bugfix: Fixed some small typos in the Landscape tutorial.
+
+* Bugfix: Reduced material update log spam when creating and editing Landscapes. Some will remain since Landscape fundamentally manipulates a lot of material instances. 
+
+* Removed "Cmd: UpdateLandscapeSetup" log file spam during while editing World Settings on maps with landscape.
+
+#### Material Editor
+
+* New: Added the ability to clear the preview mesh on a material instance.
+
+* New: Added the ability to highlight all nodes connected to an input.
+
+* New: Added tooltips to the inputs of the Material Final Result node.
+
+* Bugfix: Fixed a crash when loading cooked Materials in the editor.
+
+* Bugfix: Fixed an issue where "Select all linked nodes" function would not show up in Material Editor for output nodes.
+
+#### Matinee
+
+* Bugfix: Fixed an issue with incorrect subtitle timing and rendering when scrubbing.
+
+#### Persona
+
+* New: Added better bulk import reporting so current progress through a bulk animation import can be seen, rather than only showing progress through the current animation.
+
+* Bugfix: Fixed a crash when changing scene settings and hitting undo in Persona with complex cloth active.
+
+* Bugfix: Fixed bounds calculations in Persona that include bones to respect LOD (and other required bones).
+
+#### World Outliner
+
+* New: Added the ability to collapse and expand folders in the World Outliner.
+
+#### Static Mesh Editor
+
+* New: Added Console Command "FORCELODGROUP" that will apply a Static Mesh LODGroup to selected groupless meshes in the editor. This can batch mesh-reduce (ie, Simplygon) all meshes in a level.
+
+#### Source Control
+
+* Bugfix: Fixed an issue where history and merging do not work if the user has "format.pretty" setting in their gitconfig.
+
+* Silenced some redundant P4 errors that could be generated when running a "stat update" on a file. Some of the options produced errors when working with newly added files. These errors are now downgraded to infos like they are for the main stat command.
+
+#### World Browser
+
+* Bugfix: Fixed a crash on Changing sub-level visibility under certain conditions. 
+
+* Disabled use of World Composition with multiplayer PIE using separate process option, as it not properly supported now.
+
+#### Foliage
+
+* New: Added console variables to discard grass and/or scalable foliage data on load.
+
+    * If set in the scalability settings, it reduces the memory requirements of the "off" setting at the expense of having to reload the level when changing the setting from "off" to a higher level.
+
+* New: Implemented filter to enable painting foliage on other foliage (e.g. vines on trees).
+
+* Bugfix: Fixed an issue where painting foliage on blocking "query" actors was not working. Actors which are not themselves moving only need to be set to "Query" to be blocking, but the foliage tool was checking for both "Query & Physics."
+
+* Bugfix: Fixed an issue where disabling 'Use Landscape Lightmap' option was skewing Procedural Foliage Instances.
+
+* Prevent landscape grass generated in the editor viewport from being duplicated for Play In Editor (PIE) sessions.
+
+#### Cooker
+
+* New: Added -mallocbinned commandline option to use it while cooking. Binned allocator is faster to cook in some cases.
+
+* New: Added experimental Partial Garbage Collection feature for cooker. Added "-partialgc" commandline option for cooker. Instead of cleaning all packages when GC, this cleans up packages which will not be required for cooking later.
+
+* New: Fixed an issue with deterministic cooking for bake distributions.
+
+* New: Improved peak memory usage by the cooker by only resolving string asset references for the current package being saved.
+
+* New: Now cooker processes async shader compilation while blocked doing other work. This gives a minor improvement in cook time.
+
+* Bugfix: Fixed an issue for partial LOD level deterministic cooking issue. Spawn Module now has its postload called before being used.
+
+* Cooker forces a shader compilation flush when it detects that it has passed the max memory budget. This helps to avoid out of memory problems when cooking large projects. 
+
+* "Directories To Always Stage As UFS" now filters .uasset and .umap files. Cooked assets will no longer be replaced by their uncooked version during staging UFS.
+
+* Launch on in the editor will save required packages before cooking so that changes are correctly reflected.
+
+#### Sequencer
+
+* New: Added an option to map the transform motion of a Skeletal mesh to the root bone on FBX export. This option is in the General Settings in the Miscellaneous section.
+
+* New: Added "Get Cine Camera Component" function from Cine Camera Actor.
+
+* New: Added option to convert a spawnable track to a possessable track.
+
+* New: Added option to create sub sequences for each master sequence shot.
+
+* New: Added the ability to keyframe Geometry Cache Actor properties in order to use it in sequencer.
+
+* New: Enable "Actor to Track" on the Cine Camera Actor as a keyable property.
+
+* New: Enabled subtitles for audio tracks.
+
+* New: Force the roll and yaw of the camera crane mount to be 0 so that the camera stays level to the ground. The yaw is controlled by the crane yaw control.
+
+* New: Sequencer can now playback Anim Montages.
+
+* New: Set Skeletal meshes to always tick pose and refresh bones when there's an Animation track in Sequencer.
+
+* New: Tweaked default Sequencer colors.
+
+    * Audio track is brighter.
+
+    * Transform, bool, event tracks are less saturated.
+
+    * Recording subsection is more saturated.
+
+    * Fade track now shows a gradient.
+
+* New: Added icons for Cine Camera, Camera Rig Rail, and Camera Rig Crane Actors.
+
+* Bugfix: Fixed a crash when using sequence recorder with -game. Added GEditor check.
+
+* Bugfix: Fixed an issue where it was impossible to animate a single component of a structure property while still enabling the original value for non-animated components to be set externally. 
+
+* Bugfix: Fixed an issue where the values in the key editors in Sequencer would become out of sync with the property values. 
+
+* Bugfix: Fixed an issue with the Fade track so that is only affects the current player and not all worlds. 
+
+* Bugfix: Fixed an issue with Game Mode Override option in movie capture so that it works when capturing the Editor. 
+
+* Bugfix: Fixed issues where properties were not getting bound when converting from Matinee to Level Sequence tracks. 
+
+* Bugfix: Fixed an issue with Anim Notifies not working when playing an Animation on Blueprint-driven Skeletal meshes. 
+
+* Bugfix: Fixed an issue with Anim Trails not playing in full with a Sequencer-driven animation. 
+
+* Bugfix: Fixed an issue with byte and integer properties not being able to be exposed to Cinematics on Blueprints. 
+
+* Bugfix: Fixed an issue where the cursor would jump around inconsistently when ending a drag. 
+
+* Changed curve editor visibility so that it's no longer a toggleable saved config. 
+
+    * It's now just a toggleable state because it was unexpected to have Sequencer restart with the curve editor visible.
+
+* Changed FMovieScene3DTransformTrackInstance "Update" to properly update "Component Velocity."
+
+* Disabled auto-possess player for recorded pawns. 
+
+    * This fixes a bug where if you record a Third Person Template character when you open the Sequence the recorded character will possess the viewport.
+
+* Restore state when stopping the level sequence player so that objects controlled by Sequencer return to their positions.
+
+* Spawnables now use deferred spawning.
+
+#### VR Editor
+
+* New: You can now instantly reset the world's scale to its default by tapping the touchpad button while gripping the world. This makes it easy to get back to your "Player's Size" while editing. 
+
+* New: You can now disable VR Editor Tutorial in **Editor Preferences > Experimental > VR**.
+
+	![image alt text](image_82.png)(w:544 h:116 convert:false)
+
+* Bugfix: Fixed an issue with blurry toolbar icon when "Use Small Toolbar Icons" was enabled in the Editor Preferences. 
+
+* Bugfix: Fixed various issues and regressions with Mesh Paint in VR:
+
+    * Brush preview not being rendered for lasers/mouse when not actively painting.
+
+    * "Full press" over actors that was preventing the UI from being clickable.
+
+    * Brush cursor displayed when hovering over UI (including selection bar/close button).
+
+    * VR transform gizmo getting in the way of everything while painting, it is now hidden while in Mesh Paint mode.
+
+    * Not being able to fully interact with UIs after messing around with Mesh Paint
+
+### Gameplay Framework
+
+* New: Added a cvar "t.FPSChart.OpenFolderOnDump" to control whether or not FPS charts automatically open the profiling folder when stopfpschart is executed. 
+
+    * This can be useful to avoid a bunch of open windows while doing automated testing.
+
+* New: Added "Suppress Stacking Cues" flag to Gameplay Effect options to avoid sending a Gameplay Cue RPC for each instance of a stacking Gameplay Effect.
+
+* New: Added "GameModeBase" and "GameStateBase" classes as parents of existing classes "GameMode" and "GameState." The core functionality needed by all games is in the Base classes, legacy and match-specific features are in the child classes.
+
+    * Blueprint access has been added to Game State Base for "Has Begun Play", "Has Match Started", "Get Player Start Time", "Get Player Respawn Delay", "Game Mode Class", and "Spectator Class."
+
+    * Game Mode Base "Must Spectate" and "On Change Name" are now called properly.
+
+    * Added "Restart Player At Player Start" and "Restart PLayer At Transform" to manually control of where player pawns are spawned.
+
+    * Updated the GameMode documentation to reflect these changes and add extra context.
+
+* New: Added access to both the game's "Unpaused Time Seconds" and "Time Seconds" from Blueprints.
+
+* New: Added the ability to set the mouse position from Blueprints via the Player Controller.
+
+* New: Components can now specify whether they replicate from the Blueprint Property editor.
+
+* New: Components that are referenced by an "Edit Anywhere" property but are not created as part of object construction can now be edited via the in-world component editing panel.
+
+* New: Cull Distance Volumes no longer globally update on every Actor property change. This should reduce hitches after changing properties entirely unrelated to cull distance in levels with large Actor counts.
+
+* New: More Cheat Manager functions have been made available to call and to override from Blueprints.
+
+* New: Added "Use Client Side Level Streaming" property added to World Settings. 
+
+    * When enabled, the server doesn't evaluate streaming volumes. Instead, each client independently decides what volumes to stream in/out, depending on its players' cameras.
+
+* New: Particle System and Audio Components now route "Activate" and "Deactivate" events to Blueprints.
+
+* The "Delta Seconds" passed to Tick functions when using a "Tick Interval" will now represent the amount of world time that has passed since the last tick of the function.
+
+* The levels of a world are now categorized into different collections: dynamic, static, or duplicated. Levels default to dynamic, and this preserves existing behavior.
+
+    * Streaming levels can be marked as static in the Levels Details panel. This distinction has no effect on how the levels are rendered, but it will cause the levels to be placed into the corresponding collection.
+
+    * Games may opt-in to duplicating the dynamic levels when the world is loaded, and they will be maintained separately from the original levels. The game may use these duplicated levels at runtime, if desired.
+
+* New: The logic for selecting which Game Mode to spawn for a URL can now be overridden in "Game Instance", implemented in "Preload Content For URL", "Create Game Mode For URL", and "Override Game Mode" Class.
+
+    * The rules for selecting GameMode based on map prefix or alias is now modifiable in the Project Settings > Maps and Modes under the advanced options.
+
+* Bugfix: Fixed a crash undoing "replace actor with" when replaced object constructs child actors.
+
+* Bugfix: Fixed an issue where after duplicating a level you would have two instances of a Child Actor. 
+
+* Bugfix: Fixed an issue where Child Actors would sometimes be created at the wrong location.
+
+* Bugfix: Client-side only components attached to a replicated component will no longer be removed from the "Attach Children" array when a change to attachment is made on the server.
+
+* Bugfix: Fixed a check in Character Movement's "Step Up" to properly account for distance the component is above the floor.
+
+* Bugfix: Fixed an out of date with Radial Force Component's "Collision Object Query Params" by adding a "Begin Play" event callback.
+
+* Bugfix: Fixed an issue with Gameplay Statistic's "Spawn Emitter Attached" using wrong scale when "Snap To Target (Keep World Scale)" option is used. It was being used as a relative scale. Also, improved the comments for "Spawn Emitter Attached."
+
+* Bugfix: Fixed an issue with bounds calculations for local camera animations. 
+
+    * Correctly calculated bounds as local to the initial transform in the track. 
+
+* Bugfix: Fixed issues with "Can Jump" inconsistencies with previous versions. 
+
+* Bugfix: Fixed a crash during garbage collection after deleting Audio Volume.
+
+* Bugfix: Fixed a crash in GetPredictionData_Server() when a player is destroyed and the server Player Controller checks to see if it needs to force a network update. This, also, fixes similar crashes when calling a Character's "Set Replicate Movement" when not on the server. 
+
+* Bugfix: Fixed an error icon indicating lighting issues not going away after resolving the issue on a light that was created by copy/pasting a light that had a lighting issue. 
+
+* Bugfix: Fixed some typos in "Only Trace While Ascending" and "Deproject Screen To World."
+
+* Bugfix: Fixed UK2Node_GetInputAxisValue using default flag for "Override Parent" value on the created binding, causing  previous bindings to be squashed.
+
+* Bugfix: Fixed an issue with network replication for actors in "Always Loaded" sub-levels not working after seamless travel.
+
+* Bugfix: Fixed an issue with FPreviewScene-based worlds so that they now render correctly in-game. 
+
+    * Added new "Editor Preview" and "Game Preview" world types and deprecated old "Preview" world type. 
+
+* Bugfix: Fixed an issue to prevent Camera Anims without a "Field of View" track from making any modifications to the camera's FOV.
+
+* Bugfix: Fixed "Set Game Paused" now correctly returns whether the paused state was changed when setting it to false.
+
+* Component "Activate" functions now fire for "auto-activate" components if the component belongs to an actor placed in a level.
+
+* Component's are now correctly removed from their World's needs end of frame update list. This fixes a rare crash when a component is moved between worlds as part of seamless travel.
+
+* Consolidated rotation settings (Rotation Rate, Use Controller Desired Rotation flag, Orient Rotation To Movement flag) in a new "Rotation Settings" category. The Character Movement Component's flag for "Use Controller Desired Rotation is no longer an advanced property. This makes it more easily discoverable. 
+
+* Direct access to Static Mesh in "Static Mesh Component" has been deprecated. Please use "Get Static Mesh" and "Set Static Mesh" accessors.
+
+* Disabled some NaN checks in shipping builds.
+
+* Improved logic for when the transform of a static scene component can be modified. 
+
+* Load game from slot will no longer sometimes fail if using a Blueprint defined class. 
+
+* Mouse smoothing is now based on the application frame time rather than the dilated world time. 
+
+* Native classes will now properly run construction script on PostEditMove.
+
+* Native Components created by "Create Default Subobject" that have been removed will no longer still exist in existing instances of that class.
+
+* Scene Component Attach Children will no longer periodically get duplicate entries on the client for some replicated components.
+
+* The FOV is now reset on the Player Camera Manager's Camera Actor when it's initialized. 
+
+    * This fixes cases of stale FOV values after playing Camera Anims that don't end with the FOV at its base value.
+
+    * Base FOV can now be edited in the Camera Anim properties. This enables you to specify what the FOV keys are relative to. Previously it was always using a base FOV of 90 degrees.
+
+* The position of the (null) bounding volume of an empty Text Render component is now correct.
+
+* Wind source components now correctly "Activate" and "Deactivate."
+
+### Localization
+
+* New: Added a way to mark text in text properties as culture invariant.
+
+    * This enables you to flag properties containing text that doesn't need to be gathered.
+
+* New: Added "As Time", "As Date Time", and "As Date" overrides to Blueprints to let you format a UTC time in a given timezone. This defaults to your local timezone. Previously you could only format times using the "invariant" timezone, which assumed that the time was already specified in the correct timezone for display.
+
+* New: FText properties/pins will now preserve their existing identity (namespace + key) where possible when changing the source text.
+
+* New: Translations are now exported/imported per-identity rather than collapsing identical strings within the same namespace.
+
+    * This enables translators to translate each instance of a piece of text based upon its context, rather than requiring a content producer to go back and give the entry a unique namespace.
+
+    * It also enables us to optionally compile out-of-date translations, as they are now mapped to their source identity (namespace + key) rather than their source text.
+
+    * This change also adds FLocTextHelper as a high-level API for dealing with uncompiled localized text. This should be preferred for all new code rather than dealing with manifests/archives directly.
+
+* Bugfix: Fixed a crash when dealing with code-points outside the basic multilingual plane on platforms with UTF-32 FStrings.
+
+* Bugfix: Fixed a crash that could occur in the source code gatherer where it would sometimes underflow an array.
+
+* Bugfix: Fixed an issue where localized assets were not working with "Launch On" in the editor. 
+
+* Bugfix: Fixed an issue where the PO exporter writing out invalid escape sequences. 
+
+* FText can now format pre-Gregorian dates correctly.
+
+* Removed validity check on source cultures when remapping, as platforms may use invalid cultures that need to be remapped.
+
+* The internationalization system now correctly detects whether the local machine is in a timezone with daylight savings active
+
+* We now favor Traditional Chinese for Hong Kong and Macau.
+
+### Networking
+
+* New: World Origin rebasing has been added as an experimental feature!
+
+    * This means that the World Origin can be different on both the client and the server.
+
+    * This feature will continue to be improved over the next releases. 
+
+* Bugfix: Fixed a crash that could occur with FCurlHttpRequest's "Debug Callback".
+
+* Bugfix: Fixed an assert that could occur when playing back a replay that was recorded on a client that uses torn-off actors.
+
+* Bugfix: Fixed an assert in channel cleanup code that could occur if the connection was cleaned up and there were "Keep Processing Actor Channel Bunches Map" in-flight still.
+
+* Bugfix: Fixed an issue where gameplay tags will work better with backwards compatibility in replays. 
+
+* Bugfix: Fixed various issues with remapping objects that lose/gain relevancy. 
+
+* Bugfix: Fixed an issue with replays that could cause scrubbing to a specific time to go to the wrong time. 
+
+* Bugfix: Fixed an issue that could cause memory-constrained platforms (such as consoles) to run out of memory when scrubbing in replays in rapid succession.
+
+* Bugfix: Fixed an issue where replicated references that would remain null indefinitely due to becoming non relevant, and then becoming relevant again
+
+* Disable hot reloading when using single process MP PIE. This fixes some conflicts that can otherwise occur.
+
+* Don't strip map/package name from commandline on shipping config on dedicated server in the Game Instance's "Start Game Instance". Don't ever parse "-replay" on dedicated servers. This enables specifying map name on servers via commandline.
+
+* We no longer record dormant actors in replays for better performance
+
+### Physics
+
+* New: Added "Simple Wheeled Vehicle Movement" Component which allows for n wheels with no motor simulation. 
+
+    * Exposed "Wheeled Vehicle Movement Component 4W" so that it can be added to any Blueprint.
+
+* New: Added "Get Bone Mass" and "Get Skeletal Center Of Mass" functions for Skeletal meshes.
+
+* New: Added more vehicle stats for the profiler.
+
+* New: BSP now supports returning texture coordinates from line traces.
+
+* New: Exposed "Grab Rotation" for physics handle component. Also introduce hard vs soft physics handle.
+
+* New: Exposed the ability to separate reverse and brake on vehicles.
+
+* New: Improved cook times for "Procedural Mesh" Component collision. Fixed issue where "Update Mesh Section" on Procedural Mesh may update wrong collision verts due to mesh being cleaned in cooking.
+
+* New: Improved warning when trying to create a constraint between different scenes.
+
+* New: Improvements to data caching for Destructible meshes on save. This gives an 8x-10x improvement in load times for more complex meshes.
+
+* Bugfix: Fixed a crash when a physical animation component's mesh is changed out from under it. 
+
+* Bugfix: Fixed an issue in PhAT so that multiple constraints can be selected and edited properly at the same time.
+
+* Bugfix: Fixed an issue where constraint component not using root body when no bone name is provided.
+
+* Bugfix: Fixed an issue where the constraint index was not being assigned correctly. Also, exposed "Find Constraint Bone Name" to Blueprints.
+
+* Bugfix: Fixed an issue with FConstraintingBaseParam's "Contact Distance" clamping.
+
+* Bugfix: Fixed an issue with PhAT "Snap Constraint To Bone."
+
+* Bugfix: Fixed an issue with physics handle component not accounting for bone rotation/offset. 
+
+* Bugfix: Fixed an issue with "Set Root Body Index" so that it uses world space transform instead of reference skeleton. 
+
+    * Fixes a few issues when using ragdolls on Skeletal mesh assets that were imported with offset/rotation.
+
+    * Also, fixes a crash when opening up old physics asset and changing its mesh. 
+
+* Bugfix: Fixed an issue with vehicle suspension force offset to work from the spring location (in the z-axis) instead of the vehicle COM.
+
+* Bugfix: Fixed a crash caused by changing the collision response parameters of a Destructible mesh after it had been fractured.
+
+* Bugfix: Fixed a crash spawning complex Destructible meshes in multiplayer games. 
+
+* Bugfix: Fixed a crash when creating a Destructible mesh from a static mesh that uses SpeedTree materials.
+
+    * Added a warning message when attempting to do so now.
+
+* Bugfix: Fixed an issue with APEX Cloth disappearing if an actor containing a component with clothing had its custom time dilation set to 0.0.
+
+* Bugfix: Fixed an issue with box collision shapes not scaling correctly after using Merge Actor tool.
+
+* Bugfix: Fixed an issue with Destructible asset importing failing if one of the submeshes in the asset file had no geometry. 
+
+* Bugfix: Fixed incorrect tooltips for collision channels.
+
+* Bugfix: Fixed an issue with line traces against a component not returning face index, where appropriate.
+
+* Bugfix: Fixed an issue with line traces against a specific component not returning closest hit in triangle mesh case.
+
+* Bugfix: Fixed an issue where mesh corruption in clothing assets would happen when importing a new mesh that requires a uint32 index buffer over a clothing asset that previously required uint16 indices.
+
+* Bugfix: Fixed an issue where scales like (1,0,1) not being clamped to something valid in physics code.
+
+* Bugfix: Fixed an issue where "Support UV From Hit Results" not informing user an editor restart is required.
+
+* Altered Skeletal mesh bounds calculations to ignore clothing actors if they aren't simulating on the currently active LOD.
+
+* Changed some exposed properties on Destructible meshes to be hidden as they are only supported for non-Destructible Skeletal meshes. This addresses an issue where Physics Shadow Assets were properties were visible.
+
+* Enforced block size of 16K for PhysX scratch memory when simulating.
+
+* Ensure that physics handle automatically wakes up any object it's grabbing on release.
+
+* Improved comments for Hit callbacks.
+
+### Platforms
+
+* New: Add HarfBuzz support for proper text shaping in languages the require it such as Arabic.
+
+* New: Add support to Windows and Linux for swapping allocators with -ansimalloc, -binnedmalloc, and -binnedmalloc2 to enable easy comparison at runtime.
+
+* New: Added "Is Gamepad Attached" to Android so users can determine when a gamepad is attached to the device.
+
+* New: Added support for Android local notifications using the existing local notification Blueprint mechanisms. Using these Blueprint nodes will work on both Android and iOS now.
+
+* New: Added support for making the game window borderless (no system border or title bar) on desktop platforms. Disabled by default. 
+
+    * Enabling requires adding "UseBorderlessWindow=True" to [/Script/EngineSettings.GeneralProjectSettings] in DefaultGame.ini. 
+
+    * The game using this is responsible for adding a "Window Title Bar Area" widget to its UI, as well as window minimize/maximize/close buttons.
+
+* New: Adding interstitial ad support for Android using AdMob, including Blueprint functions.
+
+    * Four functions have been added: "Load Interstitial Ad", "Is Interstitial Ad Available", "Is Interstitial Ad Requested", and "Show Interstitial Ad."
+
+* New: AndroidManifest.xml will now include the texture format used in packaging. The appropriate "supports-gl-texture" tags will be added for the texture format chosen at packaging.
+
+* New: Engine will select the first shader platform in [/Script/WindowsTargetPlatform.WindowsTargetSettings] as the current RHI.
+
+* New: Support Metal shader compiling when cross compiling on Windows. 
+
+    * This must be enabled in the IOS Build Settings page since it can significantly increase build time.
+
+* New: Enable the ability to change ports for a firewall to pass them through rsync. This addresses UE-24679.
+
+* New: Upgrade emsdk toolchain to 1.36.11 - emscripten.
+
+* New: When launching the app on iOS within the editor, the device console logs are piped into the Editor's Console Log window.
+
+* Bugfix: Fixed an issue with cooking for compressed content where the bit window used for compression was consumed by the wrong parameter. 
+
+* Bugfix: Fixed a couple of problems with Mac app bundles modifying their content that made them incompatible with the Mac App Store.
+
+* Bugfix: Fixed an issue that was causing Sound Mixes to not properly finish at the end of certain sound effects when played on a device. Now, their behavior now matches PIE.
+
+* Bugfix: Fixed an issue that was causing the module Android_ASTCTargetPlatform to fail to load on Mac.
+
+* Bugfix: Fixed an issue where setting a new Touch Interface to none would cause any input present on the Virtual Joysticks to continue on indefinitely. Input now clears before setting a new Touch Interface.
+
+* Bugfix: Fixed an issue where shader compression used the cooking host platform instead of the target platform when selecting shader compression settings.
+
+* Bugfix: Fixed a race condition in threading code that affected pthread-based platforms, such as Linux, Mac, Android, iOS, etc. Threads are now properly joined instead of busy waiting. 
+
+* Bugfix: Fixed an issue where iOS Movie Player can't handle videos at resolutions that aren't multiple of 16. Now an assert is thrown when attempting to play a non-power of 16 resolution. 
+
+* Bugfix: Fixed an issue where occasionally a movie isn't played on iOS.
+
+* Give a more friendly error message when UHT fails with an invalid error code. This will help when trying to build without the prerequisites installed.
+
+* Make the "Running <Projectname> on <Device>" Toast stay active when launching a game to an Android device until the game has finished running on the device. Previously the toast was closing prematurely.
+
+* Pass down any extras used to launch SplashActivity to GameActivity. This enables code that relies on intent data (such as GCM) to function as expected. The GameActivity will receive the info and not just SplashActivity.
+
+* The cooking process now respects the per-platform chunking manifest configuration in the game ini. If you specify -manifests it forces all platforms to generate a manifest.
+
+* Improved the error message that displays as Unknown Error when failing to supply a Remove Build server for iOS on Windows.
+
+* Changed the "OOM Backup Memory Pool" to enable each platform to set how much memory to allocate. See the "Get Back Memory Pool Size" function in "Platform Memory." Defaults to 0, which was the previous behavior with the now removed "Support Backup Memory Pool" function in "Platform Memory," which was only true in Windows and PS4.
+
+#### Android
+
+* New: Updated libpng to 1.5.27 for Android to fix Google security notifications
+
+* New: Updated to CodeWorks for Android 1R5.
+
+* New: Enabled Windows types in vulkan.h from later SDKs.
+
+* New: ES 3.1 detection now handles 3.2 in the version string from newer devices. 
+
+* Bugfix: Fixed a crash that would occur on Kindle Fire when trying to load a project that had enabled the Oculus plugins. 
+
+* Bugfix: Fixed a crash on devices, such as Huawei Honor G750, which do not support the EXT_debug_lable and EXT_debug_marker extensions.
+
+* Bugfix: Fixed a crash that could occur when a Vulkan based project is resumed after losing focus.
+
+* Bugfix: Fixed an issue where some Android apps would crash if the device orientation changed when the launch splash screen was active.
+
+* Bugfix: Fixed an issue where OpenGL ES 3.1 detection was failing on devices which support OpenGL ES 3.2.
+
+* Bugfix: Fixed an issue with out of spec GLSL operations. 
+
+* Bugfix: Fixed the Vulkan include path in the NDK check. 
+
+* Disabled Google Play Games for ARM64 because native library causes crash on launch.
+
+* Hiding the virtual keyboard on Android properly closes input dialog.
+
+* Only use alternative event flow for Daydream packaged applications.
+
+* Pass through the intent action from Splash Activity to Game Activity.
+
+* We now properly handle seeking in Android Media Player, which did not work past 999ms before.
+
+* Read Android environment variables from .bashrc on Linux.
+
+* Restore texture filtering mode properly when movie played on Android.
+
+* Update cached length when file is written to on Android.
+
+* Use .sh extension for Android install scripts on Linux.
+
+#### iOS
+
+* Bugfix: Fixed a crash with reflection capturing in Metal.
+
+* Bugfix: Fixed an issue with Metal Dynamic RHI's "RHIReadSurfaceData" for shared textures on iOS.
+
+* No longer allow ES2 shaders to be compiled for tvOS.
+
+* Worked around some bad Xcode headers to enable compiling with Xcode8 / iOS10.
+
+#### Linux
+
+* New: Updated installation scripts to run on Ubuntu 16.10.
+
+* New: Added "offscreen" video driver to SDL so the engine can now create GL context on headless machines not running any display servers (EGL is enough).
+
+* New: Added support for creating Linux installed builds.
+
+* New: Bundled libc++ with the engine to guard against STL incompatibility problems. 
+
+    * Third party code that uses STL needs to to compiled against the bundled library to benefit from this.
+
+* New: Bundled SDL updated to trunk as of late September 2016.
+
+* New: Choose OpenGL version based on the first targeted RHI, which enables setting the version through the .ini.
+
+* New: Crash error message will print more details about erroneous memory access on Linux.
+
+* New: Last page of the signal handler alt stack is made read-only to guard against accidental smashing of it.
+
+* New: Added precompiled WebRTC (xmpp) libs.
+
+* New: Packaged non-code projects now include a 'bootstrap' shell script for convenience.
+
+* New: The engine will profile and choose the fastest available clock source, which also, fixes running on Windows 10 Linux subsystem.
+
+* New: Updated cross-toolchain to target CentOS7 and multiple architectures.
+
+* New: Enabled use of xgConsole in Unreal Automation Tool.
+
+* Bugfix: Fixed a crash when editing Widget Blueprints. 
+
+* Bugfix: Fixed an issue with client and server hanging when decoding voice chat. 
+
+* Bugfix: Fixed an issue with compilation of ARM (32-bit) server builds. 
+
+* Bugfix: Fixed an issue with copy/paste failing on Linux when the clipboard was empty. 
+
+* Bugfix: Fixed an issue with the heartbeat thread not being stopped during handling of a crash. 
+
+* Bugfix: Fixed an issue with QtCreator and Makefile generators, which missed some include paths and had others duplicated.
+
+* Bugfix: Fixed an issue with SDL EGL API binding. 
+
+* Bugfix: Fixed an issue with UBT hanging after Ctrl + C when mono thread creation fails. 
+
+* Bugfix: Fixed an unnecessary execution of some binding steps even when nothing has changed. 
+
+* Bugfix: Treat abort() / SIGABRT as a crash.
+
+    * This fixes some abnormal exits not being reported.
+
+* Disabled rebuilding hlslcc during installation given that it is now built against the bundled C++ library and will always match the engine.
+
+* Equals character is now enabled as part of commandline parameter value.
+
+* Minor robustness improvement in Linux crash handler.
+
+* More robust check of installed packages on Ubuntu in installation scripts.
+
+* Reduce amount of logging when launching externals procs.
+
+* ShaderCompileWorker will be told to dump core on crash in Debug editor builds to aid debugging.
+
+* Startup time of modular builds (e.g. editor) are 10x faster due to exclusion of unnecessary symbols.
+
+#### Mac
+
+* New: OpenGL support is now deprecated and Metal SM4 was replaced by SM5 in 4.13.
+
+    * We no longer compile OpenGL and Metal SM4 shaders by default to save cook time and decrease cooked data size. 
+
+    * Running under OpenGL will now display a dialog warning that Mac OpenGL is no longer supported and may not run correctly. 
+
+* New: Added a commandline switch for Mac "-RedirectNSLog" that will capture the output of NSLog using an NSPipe and write it into our log instead when not running under the debugger.
+
+    * This will capture the output from underlying libraries that would otherwise be lost when outside a debugger.
+
+* New: Added a cvar named "r.Mac.OpenGLDisabled" to disable OpenGL support on Mac. 
+
+    * If a Mac that does not support Metal launches while this cvar is set to 1 then they will get a dialog box describing that they do not support Metal and the process will close.
+
+* New: Added more statistics for Metal to track how many textures are allocated a frame and how much buffer memory is allocated per frame.
+
+* New: Added support for compiling Vulkan shaders for Android on Mac.
+
+* New: Added optional cache inside Metal RHI for disposed texture objects so we may reuse them. This is controlled by the console variable "rhi.Metal.TextureCacheMode".
+
+    * 0 = off.
+
+    * 1 (default) = will attempt to reuse private memory texture objects within the frame they are released otherwise they are disposed of as before.
+
+    * 2 = extends the caching to all textures - though Managed/Shared textures cannot be reused until after the frame in which they were released has been processed on the GPU. In this mode id<MTLTexture> objects are never returned to the OS so in order to conserve VRAM calls to setPurgeableState are made to enable the driver to reclaim unused memory if required.
+
+    * This setting can't be changed at runtime.
+
+* New: On MacOS, tvOS and iOS, the RHI and shader platform initialised on startup will now be decided by the options specified in the project's target settings rather than assumed in code.
+
+    * When the initial selection in the target settings is not available on the current device the code will gracefully fallback, if possible or alert the user if not.
+
+    * On Mac: Order of initialisation is the order listed in TargetedRHIs specification.
+
+    * On iOS/tvOS: Order is explicit: Metal MRT > Metal ES 3.1 > OpenGL ES 2.
+
+* Bugfix: Fixed a crash under the validation layer and Nvidia's El Capitan (10.11) drivers when Distance Field Particle Collisions are used without any scene Distance Fields available. We now bind the black volume texture when that is the case to avoid bad access on the GPU. 
+
+* Bugfix: Fixed a crash caused by trying to restore Help menu item on Mac if MenuBlock didn't contain it.
+
+* Bugfix: Fixed a crash caused by Metal context releasing its contents in some cases on MacOS 10.12.1.
+
+* Bugfix: Fixed a crash in CoreAudio on unpause. 
+
+* Bugfix: Fixed a crash in the editor on displaying a new window on MacOS 10.12+.
+
+* Bugfix: Fixed an issue to always bind buffers to all particle shader inputs to avoid a validation error in Metal.
+
+* Bugfix: Fixed an issue with AV Foundation video playback causing validation errors in Apple's debug tools, as the textures that it returns to us can't be used as render-targets and are stored in CPU-accessible memory.
+
+* Bugfix: Fixed an issue with AVFMedia plugin failing when a movie has captions.
+
+* Bugfix: Fixed an issue in MetalRHI that caused the first render query result to return the wrong result, which manifested as paused particle animations.
+
+* Bugfix: Fixed a focus problem in the asset picker that forced users to click on the search text field to be able to type.
+
+* Bugfix: Fixed an issue with game window not rendering its contents in some cases on macOS 10.12.1.
+
+* Bugfix: Fixed an issue with HLSL->Metal conversion for calls to asuint/asfloat where the input type matches the output type.
+
+* Bugfix: Fixed an issue MetalRHI's handling of calling "Set Render Targets Info" without any render targets.
+
+* Bugfix: Fixed an issue with TBasePassVertexShaderPolicyParamType not setting values for all shader parameters that would then cause an error in Metal's validation layer.
+
+* Bugfix: Fixed an issue for Metal shader translation retain more precision for float constants -1.0f >< 1.0f by emitting them in scientific notation.
+
+    * This prevents Hammersley constant amongst others from being flushed to 0.
+
+* Bugfix: Fixed an issue with FMac Metal not natively supporting PF_G8 + sRGB as Apple claims that not all Mac GPUs have single-channel sRGB formats. 
+
+* Changed the ShaderCache not to cache resource bindings in the draw states for shader platforms that don't use them.
+
+* Disabled the lazy-encoder construction in Metal for AMD. There is a situation that causes the lazy construction to perform a clear that isn't wanted.
+
+* Enabled more concurrent ShaderCompileWorker instances on Mac now that they don't consume such excessive amounts of memory.
+
+* Ensure that command-buffer failure in Metal is always fatal as there is no way to recover and make the default number of command-buffers sufficiently large to handle texture streaming in Mac games.
+
+* Rework Metal's handling of RHI Set Stream Source calls that override the stride of vertex declarations to be much more efficient.
+
+* Update the GPU identification code in Metal to handle changes made in macOS Sierra.
+
+* Use a read/write mutex to protect access to Metal's internal shader pipeline caches so that parallel threads can progress in the common case where new shaders are not being compiled.
+
+* Use Metal's internal setBytes API to upload buffer data that is less than a single CPU page as this is much faster and avoids fragmenting Metal's internal heaps. On Nvidia GPUs this is always used but for iOS versions without the feature or Macs running El Capitan on Intel GPUs an extra ring-buffer is used instead.
+
+* Use private memory for the Metal stencil SRV workaround needed on OS X El Capitan for a small performance improvement.
+
+#### PlayStation 4
+
+* New: Added a PS4 Device Profile when connected to a 4K TV.
+
+* New: "Get Num Of Tracking Sensors" and "Get Tracking Sensor Properties" are now supported by PSVR. This will return the properties for the PS4 tracking camera.
+
+* New: When a game supporting PlayStation Pro is connected to a 4K TV, the game may choose to automatically create 4K backbuffers or to continue with 1080p backbuffers. 
+
+    * This choice should be made on a per-game basis taking memory cost and performance into account. 
+
+    * Use "r.PS4Allow4kOutput" to enable. Disabled by default.
+
+* Bugfix: Fixed a crash that could occur with audio threading when the Kernel took over core 7 and starved the audio thread.
+
+* Bugfix: Fixed a possible GPU hang on PS4 when using DrawIndirect or DispatchIndirect. This hang could be seen in projects using DistanceField AO or DistanceField Shadowing.
+
+* Bugfix: Fixed an issue with render target alignment bugs on PS4 Pro. 
+
+* Bugfix: Fixed some audio drops on PS4 media playback.
+
+* PS4 Time's "SystemTime" now returns local machine time instead of UTC.
+
+* Stopped PS4 from always staging all ICU data files.
+
+* Upgradeable apps will now have the app_type field in Sony Project Files (.gp4) generated with the spelling 'upgradable' which is consistent with what Sony expects.
+
+* Removed the mapping for the Play controller's "Select" button, as it is reserved by the system.
+
+#### Xbox One
+
+* New: The default Xbox One XDK is now August 2016 QFE 2.
+
+* New: Added support for multiple binaries in Xbox One deployments. Removed the need for binary renaming during deploy.
+
+* New: Added support for using the virtual keyboard on Xbox One.
+
+* New: D3D12 Movie player updated to a fast SIMD implementation. This will be replaced in a future SDK when the hardware decoder is supported.
+
+* New: Improved Xbox One deployment performance and iteration times.
+
+    * Fixed an issue where staging could be performed twice per build from Visual Studio.
+
+    * Ensured files that have not changed are not recopied for staging or deployment.
+
+    * Modified manifest generation to avoid writing a new manifest if the contents are unchanged.
+
+* New: Refactored Xbox One controller handling to reduce Cross-OS calls to improve performance and to reduce memory usage.
+
+* New: Support for the media foundation player added for Xbox One.
+
+* New: Added better compression buffer sizing during pak creation, which will result in more files being compressed for use with hardware decompression.
+
+* New: Added support for staging era.xvd from the XDK, which is required by newer XDKs for application launch.
+
+* Bugfix: Fixed an assert when saving/loading games. 
+
+* Bugfix: Fixed an ESRAM memory leak caused by Volume Textures. 
+
+* Bugfix: Fixed an issue that would cause the leading character in complex AppxManifest settings to be ignored.
+
+* Bugfix: Fixed an issue where the shift key on a physical keyboard connected to an Xbox One kit could be ignored.
+
+* Bugfix: Fixed an issue with corrupted screenshots. 
+
+* Bugfix: Fixed an issue where DirectX 12 support was not being marked as experimental for the Xbox One. 
+
+* Bugfix: Fixed an issue with hardware compression and decompression to always fall back to the software implementation if the hardware implementation fails.
+
+    * This could happen if data built before hardware compression support was added is decompressed through the hardware decompression path or if compression or decompression occurred before the RHI was initialized.
+
+* Disabled engine analytics on Xbox One shipping games.
+
+#### HTML5
+
+* New: upgrade emsdk toolchain to 1.36.11
+
+    * OSX binaries: clang.
+
+    * Rebuilt ThirdPartySoftware libs for HTML5.
+
+    * Windows binaries: clang.
+
+#### Windows
+
+* New: Add support for multitouch in Windows 7-10!
+
+* New: When loading a UE4 module DLL, the engine now manually scans through project and plugin binaries and pre-loads dependent DLLs. This removes limitations caused by a maximum length of the PATH environment variable.
+
+* Bugfix: Fixed an issue where maximized borderless game windows would render their contents offset by 8 pixels from the top and left sides.
+
+#### All Mobile
+
+* New: Provided access to project directory in UPL (Unreal Plugin Language) using $S(ProjectDir).
+
+### Programming
+
+* New: Generated project files now include intellisense definitions and include paths for all UE4 modules, rather than just those in the standard UE4 editor target.
+
+* New: Unreal Build Tool now supports the Intel C++ compiler on Windows.
+
+* New: Added a -nolink command line option for Unreal Build Tool, which enables compiling a target without linking it. This is useful for automated build tests, such as non-unity compiles.
+
+* Bugfix: Makefile is now invalidated if any generated headers are deleted that fixes an issue where files are moved from one module to another, and the original module no longer contains any generated headers; its include path needs to be removed from the compile environment.
+
+* The Unreal Build Tool now deletes any DLLs in output folders which share names with build products. The Windows loader reads DLLs from the first location it finds a file with a matching name, so we need to ensure that it doesn't load stale DLLs when output directories are changed (moving a module into a plugin, for example).
+
+* Updated the Unreal Build Tool so that programs can now be built as part of editor/client/server solution configurations.
+
+* Updated UnrealVS so that open documents are now saved before running a single file compile.
+
+* The GUBP build tools have been removed. Scripts to create binary UE4 distributions are now implemented in Build Graph.
+
+### Rendering
+
+* New: Added a new "Cinematic" Scalability level to the scalability system. 
+
+    * This level is intended for offline rendering of cinematics only.
+
+* New: The PS4 defrag pool size is now controlled by a new variable "r.PS4DefragPoolSize", and is not bound anymore to "r.Streaming.PoolSize." It will need to be configured to fit the project needs, in term of graphic resources.
+
+* New: Added 1000 nit HDR output option.
+
+* New: Added a cvar "r.DumpSCWQueuedJobs" to help tracking ShaderCompileWorker crashes.
+
+* New: Added -emitdrawevents command line parameter to enable draw events from startup. This can be useful when capturing with third party graphics debuggers.
+
+* New: Added material nodes "Atmospheric Fog Sun Direction" and "Atmospheric Fog Sun Color" so that materials can very easily get access to the direction and color of a directional light that has been marked as the atmospheric fog light. This can reduce the need to use MIDs or MPCs.
+
+* New: Added per-material UV density data for Static Meshes and Skeletal Meshes to improve texture streaming metrics.
+
+* New: Added Pre-Tonemapper HDR color to buffer visualization and frame capture.
+
+* New: Added Project Setting to enable Skincache Shader Permutations
+
+* New: Added shader resource compression. This will compress shader resources in memory. It is on by default for all platforms except for Xbox One.
+
+* New: DirectX 12 RHI implementations for PC and Xbox One are now mostly merged. 
+
+    * Some platform specific implementations remain, but optimizations and fixes will now affect both platforms at the same time.
+
+* New: Hierarchical LOD transition distances can now be adjusted based on Scalability settings. Use "r.HLOD.DistanceScale."
+
+    * Defaults to 1, but higher values will make the transition happen farther away. 
+
+* New: Make shadow radius culling take field of view into account. 
+
+    * Note that this may affect shadow depth rendering performance positively or negatively, depending on the camera's FOV.
+
+* New: New scalability setting "r.Streaming.LimitPoolSizeToVRAM" to limit the texture streaming pool to how much video memory is available.
+
+    * Improved the behavior when using adaptative mip bias ("r.Streaming.UsePerTextureBias") so that the scalability bias is only used to generate over-budget warnings.
+
+* New: New texture streaming build status available through "r.Streaming.CheckBuildStatus." This is used to warn the user when a texture streaming build is required.
+
+* New: Particle cutout optimization no longer requires SubUV animation. 
+
+    * It is now on the "Required" module and there is now a Project Setting under Rendering > Optimization > Particle Cutouts by Default that will automatically set up the cutout optimization when materials change on an emitter.
+
+* New: Console variable "r.SSR.MaxRoughness" can now be used in a Shipping build rather than as a development-only feature.
+
+* New: Added support for "highresshot" on Vulkan.
+
+* New: Texture streaming MipBias, from Scalability Settings, does not affect the wanted mips computation anymore, to prevent low texture quality on characters.
+
+* New: The command "ListTextures" now outputs stats in terms of "in memory" and "on disk" to improve readability. Now, results are also more consistent between cooked builds and non cooked builds.
+
+* New: Vulkan shaders are now using glslang from SDK 1.0.24.0.
+
+* New: When running with the D3D11 RHI on Windows under RenderDoc, draw events are now enabled by default.
+
+* New: When running with the OpenGL RHI on Windows under RenderDoc, draw events are now enabled by default.
+
+* New: Added a fix for Vulkan drivers returning infinite number of swapchains.
+
+* Bugfix: Fixed a crash exiting VR Preview when using OpenGL 4 on Windows. 
+
+* Bugfix: Fixed a crash when using High Quality Particle lights. 
+
+* Bugfix: Fixed a potential crash when painting in the Landscape editor.
+
+* Bugfix: Fixed a crash when updating hierarchical instanced static meshes on the fly.
+
+* Bugfix: Fixed a crash in Procedural Mesh slicing when no verts are generated.
+
+* Bugfix: Fixed a crash when using tessellated materials with Blueprint "Draw Material To Render Target" nodes.
+
+* Bugfix: Fixed a crash with Post Process Volume when using Planar Reflections with Instanced Stereo.
+
+* Bugfix: Fixed an issue to always bind a value for "Is Stereo Parameter" in the Planar Reflection shader parameters to avoid an invalid GPU access on Metal in shaders that perform an if-test on this variable.
+
+* Bugfix: Engine no longer attempts to release Static Mesh resources if the mesh was never rendered, and therefore never initialized the resources. This also fixes some incorrect stats related to static mesh memory usage.
+
+* Bugfix: Fixed an issue for large spotlight culling issues.
+
+* Bugfix: Fixed an issue for "Auto Compute LOD Distances" option in static mesh LOD settings. Previously the calculation was incorrect and would give somewhat arbitrary distances.
+
+* Bugfix: Fixed an issue where Skylight Occlusion Tint would not properly apply the material's diffuse color.
+
+* Bugfix: Fixed a fencing bug with the async compute translucency lighting volume clear. This could cause a flickering or a rare GPU hang on PS4.
+
+* Bugfix: Fixed an issue with divide-by-zero in certain cases in the Motion Blur Post Process. This resulted in artifacts in certain cases.
+
+* Bugfix: Fixed an issue with Foliage where lighting would always needs to be rebuilt.
+
+* Bugfix: Fixed an issue where changing Detail Mode not propagating to components.
+
+* Bugfix: Fixed an edge-case when replacing texture references in a material that caused sampler type updates to fail.
+
+* Bugfix: Fixed an issue where gamma was being incorrectly applied to vertex colors of Procedural Meshes. Fixed Bounds Scale not being applied to Procedural Mesh Component.
+
+* Bugfix: Fixed an issue with "Get Capsule Mesh" not working correctly when oriented. 
+
+* Bugfix: Fixed an issue with normals on cap section when slicing Procedural Mesh Components. 
+
+* Bugfix: Fixed an issue with possible low resolution textures when they were used within UMG.
+
+* Bugfix: Fixed the Primitive filter in the Statistics window of the editor.
+
+* Bugfix: Scaled instanced static meshes will now correctly compute world normals.
+
+* Bugfix: For FRCPassPostProcessVelocityScatter use the Empty vertex declaration not the Filter vertex declaration as the vertex attributes aren't read in the shader, which fixed a Metal validation error.
+
+* Removed the console variable "r.Vulkan.UseGLSL", now we use SPIR-V everywhere.
+
+* Don't create material resources if we are in a build that can never render - Saves a few MB of memory
+
+* Made shadow radius culling take field of view into account. Note that this may affect shadow depth rendering performance positively or negatively, depending on the camera's FOV.
+
+* When using the GPU skin cache experimental feature, less shader permutations are now needed.
+
+* On XboxOne, the Vertex and Index buffers are now allocated with the right nextwriteoffset to prevent stomping old data on future writes.
+
+#### FX
+
+* Bugfix: Fixed a rare crash in Particle Trail Emitter instances.
+
+#### Lighting
+
+* New: Added "Capsule Indirect Shadow Min Visibility" to Skeletal Meshes, so that artists have control over indirect Capsule Shadow darkness without changing cvars.
+
+* New: Added flag for "Visible In Reflection Captures" to Primitive Component, which is useful for hiding objects too close to the capture point
+
+* New: Added "Distance Field Bias" to Static Mesh Build Settings for mesh Distance Fields, which is useful for reducing self shadowing on meshes that have ambient animation.
+
+* New: Changed the default Planar Reflection "Distance From Plane Fadeout End" from 600 to 100, which reduces artifacts and is a more intuitive initial setting.
+
+* New: Clamp Indirect Lighting Quality to 1 in preview builds to keep previews useful even with Indirect Lighting Quality jacked up to 10.
+
+* New: Distance Field Ambient Occlusion and Ray Traced Distance Field shadows are now enabled in High and Epic scalability settings by default.
+
+* New: Planar Reflection Actor show flags can now be edited.
+
+* New: Static light source shapes will no longer draw black in Reflection Captures. They will now fade out instead.
+
+* New: Static lights with Min Roughness = 1.0 don't get their source shapes drawn into reflection captures, since they are being used as virtual area lights.
+
+* Bugfix: Fixed an issue where lightmaps would not be applied because of how component transforms were computed.
+
+* Bugfix: Fixed an issue with double lighting of clear coat, cloth, and eye shading models. 
+
+* Bugfix: Fixed an issue where the Global Distance Field was not dirtying previous object position on Update Transform. This left behind a phantom shadow on teleports. 
+
+* Bugfix: Fixed an issue with Instanced Static Meshes being unbuilt after a lighting build if you ever cancelled a previous lighting build. 
+
+* Bugfix: Fixed an issue with pre-shadow artifacts with multiple scenes, like the character receiving Stationary Light shadows from static environment. 
+
+* Bugfix: Fixed an issue with Ray Traced Distance Field Shadows on PS4. 
+
+* Bugfix: Fixed the "r.AllowStaticLighting" flag in the pixel inspector.
+
+* Bugfix: Fixed an issue Screen Space Reflections disappearing when the camera direction is near perpendicular to the pixel's normal.
+
+* Bugfix: Resolved HLOD self-shadowing interactions in Lightmass.
+
+* Bugfix: Scene captures and Planar Reflections force a scene color alpha channel to be used when they are capturing. This fixes Planar Reflections with "r.SceneColorFormat=3", which happens due to scalability.
+
+* Lighting channels can now be edited on components with static mobility, since dynamic lights can still affect them.
+
+* Lightmass area shadows only mark texels as mapped inside the light's influence, which fixes multiple Stationary Lights with the flag 
+"Use Area Shadows For Stationary Light" interfering.
+
+* Now clamping light Min Roughness to .04 to avoid NaNs from Vis_SmithJointApprox on materials with Roughness 0
+
+* Shortened lightmass export filenames to prevent path length issues.
+
+#### Materials
+
+* New: Added "pre-skinned local position" material graph node. Available in the vertex shader only but can be passed through Custom UV interpolators.
+
+* New: Added new material functions to engine\content:
+
+    * Interior Cubemaps: Maps specified cubemaps to a virtual box. Useful for displaying insides of buildings without using polygons.
+
+    * Ray March Heightmap: Performs a ray march through a heightmap as if it were a volume. Enables more volumetric lighting for heightmaps and flipbooks.
+
+    * High Pass Function: Performs a high pass operation on a specified Function. Requires multiple offset samples of the function to be hooked up.
+
+    * High Pass Texture: Performs a high pass operation on a texture object.
+
+    * Unsharp Mask Function: Performs an unsharp mask operation on a specified function. Requires multiple offset samples of the function to be hooked up.
+
+    * Unsharp Mask Texture: Performs an unsharp mask operation on a texture object.
+
+    * Box Intersection: Returns the ray entry and exit positions and distance traveled inside a box.Box is axis aligned. To support arbitrary rotations, transform the view ray into the local space of the new coordinates.
+
+    * Cylinder Intersection: Returns the ray entry and exit positions and distance traveled inside a cylinder.
+
+    * Quadratic Formula: A simple helper function providing pins to solve the quadratic formula, providing T0 and T1.
+
+    * Smooth Ceil: Like regular ceil but with a variable soft ramp before each integer cutoff.
+
+    * Axis Aligned Fresnel: Performs fresnel restricted to a single axis. This is useful to restrict a rim light to only the edges of an object, even at glancing angles.
+
+* New: All noise material functions except "Fast Gradient Noise" now work on mobile.
+
+* New: Noise material functions have been made a little faster (up to 1.7x) by changing the GPU random number generators used.
+
+* New: Stereo planar reflection support.
+
+* Bugfix: Fixed a case where material functions would incorrectly flag as a re-entrant expression.
+
+* Bugfix: Fixed a case where material nodes containing a dot product or length calculation could optimize and ignore the actual vector sizes.
+
+* Bugfix: Fixed a number of issues with subsurface profile shading model.
+
+    * Emissive lighting was being applied twice when rendering checkerboard skin.
+
+    * Emissive lighting was modulated by basecolor in the recombine when rendering checkerboard skin.
+
+    * Metallic materials were contributing specular lighting to the diffuse channel when rendering checkerboard skin.
+
+    * Optimization: Fall back to default lit lighting for pixels where the opacity is 0.
+
+* Bugfix: Fixed an infinite loop crash when chaining multiple material attribute functions together.
+
+* Bugfix: Fixed texture streaming for instanced components like foliage.
+
+* Bugfix: Fixed world position in downsampled separate translucency.
+
+* Bugfix: Fixed cases where chained calls to the same function will overwrite input connections on previous calls in the chain by linking function inputs into the caller graph before re-entering after compiling an input.
+
+* Updated the Actor Position node to now be set to the position of the actor you are attached to instead of your owning actor.
+
+#### Postprocessing
+
+* New: Added Alpha output support for postprocess materials (optional as a property of the material).
+
+* Anti Aliasing method has been removed from Post Process volume settings, and can now only be set from Project Settings under the Rendering section.
+
+#### Optimizations
+
+* New: Added shadow stencil clear optimization - avoids clearing stencil for each projected shadow by zeroing the stencil buffer during testing.
+
+    * This can give significant performance gains in cases with large numbers of per-object shadows (e.g. many characters with a single stationary light).
+
+    * Note: Writing to stencil during testing negates the use of high stencil on many GPUs, but this change is still a clear win an all the cases we tested.
+
+* Updated material resources to not be created if we are in a build that can never render - This saves a few MB of memory.
+
+#### Mobile Rendering
+
+* New: Improved Mobile Scalability options.  Each mobile material quality level can now specify a CSM shadow filtering quality, and the Fog calculation can now be enabled and disabled for both HDR and non-HDR mobile games.
+
+* New: Added on screen warning for invalid reflection captures. It can be seen only in a game running with mobile renderer.
+
+* New: Added option to mobile quality level settings to control the quality of filtering applied when rendering with CSM shaders.
+
+* New: Android device profiles can now elect to disable support for Vulkan and fallback to OpenGLES.
+
+    * Use the console command r.Android.DisableVulkanSupport to disable support for Vulkan.
+
+* New: Mobile vertex fog is now a project option.
+
+    * It can now be enabled or disabled regardless of Mobile HDR project setting.
+
+* New: Added basic mobile support for global clip plane project setting.
+
+    * This fixes rendering issues with planar reflections on mobile.
+
+* New: Added more stringent checks for Android ES3.1 compatibility.
+
+    * Android ES3.1 devices must also support floating point render targets and shader IO blocks.
+
+    * Improved logging when ES3.1 support has failed.
+
+* Bugfix: Fixed a crash on startup when using Vulkan on Android devices with Mali-based GPUs.
+
+* Bugfix: Fixed Average Brightness property being applied to reflections in gamma space in the mobile base pass. This would cause ES2 reflections to be overbright.
+
+* Bugfix: Fixed case where Sky Light could be rendered in green color on some Android devices.
+
+* Bugfix: Fixed an issue where the editor ES2 Feature Level Preview and Mobile Preview PIE modes were not correctly limiting materials to only the 8 textures available on ES2.
+
+* Bugfix: Fixed issues with modulated shadows and decals on some PowerVR54x based android devices (such as Galaxy S4).
+
+* Bugfix: Fixed issues with transparent depth reading materials reading invalid depths on some android devices.
+
+* Bugfix: Fixed planar reflections when running Vulkan on Android devices with Adreno-based GPUs.
+
+* Bugfix: Fixed possible rendering issues with gpu skinning on OpenGL ES 3.1 platforms.
+
+* Bugfix: Fixed Android to use single channel texture for Shadow Maps instead of 4 channel.
+
+* Bugfix: Fixed Black HDR compressed textures on devices with Mali GPUs.
+
+* Bugfix: Fixed Black rendering artifacts on Nvidia Shield Tablet device when project is built with Deferred OpenGL ES rendering support.
+
+* Bugfix: Fixed ES3+ devices to properly detect support for HALF_FLOAT and UNSIGNED_INT_2_10_10_10_REV vertex formats.
+
+* Bugfix: Fixed the Help text for SortBasePass console variable to now properly describe each option.
+
+* Bugfix: Fixed iPad Mini 4 devices to now correctly render refraction effect.
+
+* Bugfix: Fixed Particle Cutout crashes on certain devices (Samsung Galaxy Note 2).
+
+* Bugfix: Fixed Shader compile errors on NVIDIA Shield devices when project uses dynamic Point Lights.
+
+* Updated planar reflections on mobile to no longer cause GPU particles to be updated twice per frame.
+
+* Updated scene captures to now correctly decode scene color on devices that do not support floating point targets (such as the pre-updated Galaxy S6).
+
+* Previously, Mobile HDR setting implied vertex fog. If your project does not use Mobile HDR, you may wish to disable Vertex Fog in the Project Settings.
+
+### UI
+
+* New: Added delegate to the Game Viewport Client that notifies when the Game's platform specific window is being closed, and can be used to prevent the game from being exited (and the window from being closed).
+
+* New: Textures, Materials, and Paper2D Sprites can now be added to Widget Blueprints directly from the Content Browser.
+
+* New: Updated the Text Render Component to support multiple font pages.
+
+* Updated Circular Throbber widgets to no longer support Period values less than zero and they will no longer throw an error when set to zero.
+
+* Updated the viewport to no longer lose mouse focus if a visible cursor clicks inside of it when the cursor should be visible.
+
+* Updated drag operations to only be completed or canceled by the finger that started the operation, when using touch input.
+
+#### Slate
+
+* New: Added support for underlined text rendering.
+
+    * The Text Block Style can now specify a brush to use to draw an underline for text (a suitable default would be "DefaultTextUnderline" from the Core Style class).
+
+    * When a brush is specified here, we inject highlights into the text layout to draw the underline under the relevant pieces of text, using the correct color, position, and thickness.
+
+* New: Added the Window Title Bar Area widget (and UMG wrapper for it), which can be used to specify which regions of the UI should enable the user to drag the window on desktop platforms.
+
+* New: You can now use the Content/SlateDebug folder to store any Slate resources that shouldn't be used in a shipping build.
+
+    * LastResort.ttf has been moved to this folder.
+
+* Bugfix: Fixed jittering that could occur when using a fixed wrapping width with an editable text.
+
+* Bugfix: Fixed an issue preventing Editable Text and Multi Line Editable Text from setting the correct foreground color when painting.
+
+* Bugfix: Fixed volatility and invalidation issues for text widgets.
+
+* Bugfix: Fixed Editor freezes and memory leaks when slate remote is enabled.
+
+* Bugfix: Fixed Scroll Box right mouse/touch grab scrolling functionality.
+
+* Bugfix: Fixed engine compilation with defined LOG_SLATE_EVENTS.
+
+* Updated material blend states to match what is expected of Slate rendering, which differs a lot from the scene renderer with the way it treats alpha.
+
+    * This fixes translucent rendering with the retainer widget, users will need to set their materials to Alpha Composite though for it to behave as expected.
+
+* Updated Menu Anchor and Scroll Box to now work with invalidation panels.
+
+#### UMG
+
+* New: Added controller support to Combo Box widget.
+
+* New: Added support for outline fonts.
+
+* New: Updated the Copy, Cut & Paste animation tracks to be fully supported in UMG now, though not fully supported for Sequencer in general (outside of UMG).
+
+* New: Updated the safe zone widget in both Slate and UMG to have the ability to control which sides are accounted for.
+
+    * This can enable the nesting of these widgets to solve situations where you want parts of your layouts to extend to the edge of the viewport.
+
+* New: Updated the selection of animation tracks in UMG to automatically update the selected widgets now.
+
+* New: Added an option to exclude all UMG widgets and slots from dedicated server builds.
+
+    * Set bLoadWidgetsOnDedicatedServer=false for this behavior.
+
+* New: Exposed a trace channel for the Widget Interaction Component.
+
+    * Defaults to Visibility.
+
+* New: Took the Widget Component and Widget Interaction Components out of experimental.
+
+    * Removed old importing support for upgrading ancient versions of widget components.
+
+    * Removed parabola distortion, as users can now do whatever they want in their custom MID they can override the widget with.
+
+* New: Updated the designer so that you can disable and enable the dashed lines around containers by using 'G' to toggle game mode. The option in the settings is now used as the default when you startup a designer.
+
+* New: Updated the UMG "Replace with" function to now support "Replace with *Selected Widget*".
+
+* New: Updated UMG Project settings to now show/hide different widget classes and categories in Palette view.
+
+    * Set this option under Project Settings->Editor->UMG Editor.
+
+* Bugfix: Fixed an issue where renaming a widget would cause the slot content animation track to be invalid.
+
+* Bugfix: Fixed an issue where buttons inside a scroll list with a click method set to Precise Click would activate if the mouse was released when they are out of view.
+
+* Bugfix: Fixed an issue where removing a Widget from Panel Widget did not invalidate layout, which could leave it still visible accidentally.
+
+* Bugfix: Added missing base includes and forward declarations for Progress Bar and Text Block.
+
+* Bugfix: Fixed 3D widgets getting culled when using materials that disable depth testing.
+
+* Updated Combo Box (String) widgets to now correctly clear the currently selected option if that option is removed from a Selection Changed callback.
+
+* Prevented widget animation from being named the same as one of the override functions in a User Widget.
+
+* Updated the Experimental Class warning to now have a more detailed description.
+
+* Updated minimized windows to no longer behave like they are visible when determining what widgets are located under the mouse cursor.
+
+* Updated the grid panel to now properly set padding on slots as it is intended to.
+
+* Updated the wrap box so that the second row now properly receives "inner" vertical padding.
+
+* Documented the meta parameters allowed on widgets, like we do for regular UObjects.
+
+    * For binding widgets from blueprints you can now do BindWidget (unchanged), and to simplify binding widgets optionally, you can now just use BindWidgetOptional, rather than the combination of BindWidget and setting OptionalWidget to true.
+
+* Improved setting focus for users on widgets.
+
+    * If we are unable to set the focus immediately, possibly because the user is setting focus in the Construct callback before the widget is in the tree, we now update the Slate Operations FReply on Local Player to set focus in the next frame when it is more likely the widget will become focus-able.
+
+* Updated retainer widgets to now contain the Virtual Window directly in the hierarchy. This should now make it possible to reliably focus elements inside the retainer widget.
+
+* Updated Button and Widget to no longer multiply color and opacity twice.
+
+### VR
+
+* New: Added support for getting the HMD Device name (e.g. "OculusRift," "SteamVR," "PSVR," etc) through Blueprints and C++.
+
+* New: Added the "Start in VR" project setting, which enables projects to specify that they would like to automatically attempt to run in VR mode, regardless of whether or not -vr is passed on the commandline.
+
+* New: Added the ability to re-enter VR preview in the Editor after exiting from it using the quit option inside the SteamVR overlay. Previously one had to quit and restart the Editor to re-enable VR.
+
+    * It is now possible to reliably exit a game from the Steam VR overlay exit button, even when previewing in the Editor.
+
+* New: Added the ability to select which HMD module to use on startup by passing "hmd=<name>" on the command line.
+
+    * This is useful if you have more than one type of device attached to your machine. If no device matches the name passed in, a warning will be printed to the log and VR will be disabled.
+
+* New: Added support to the Oculus Rift for the "r.ScreenPercentage" console variable if it is set manually.
+
+    * The obsolete console command "HMD SCREENPERCENTAGE" will still work, but assigning to "r.ScreenPercentage" is now the preferred way.
+
+* New: Updated the PSVR 2DVR texture to now be a UTexture, so it can be a RenderTarget, etc. Also improved error reporting for this and some other PSVR features.
+
+* New: Added three options to the PSVR HMD Setup Dialog Cancel reactions.
+
+    * Disallow Cancel - Put the HMD Setup Dialog up again. Appropriate for titles that only support VR output.
+
+    * Switch To 2D - Switch to 2D output mode and send the HMD Connect Canceled Delegate. Appropriate for apps where 2D play is supported.
+
+    * Delegate Defined - Fire the HMD Connect Canceled Delegate. The project needs to handle this delegate and take appropriate action while being mindful of TRC requirements.
+
+    * Be sure to consider the case where the HMD is disconnected as the program starts up.
+
+* New: Added delegates for when the PSVR headset is taken off the head or put on the head.
+
+* New: Updated GoogleVR SDK to v1.01, which adds proper support for splash screens.
+
+* New: Updated OSVR plugin.
+
+    * Fixes multiple crashes, and improves tracking.
+
+* New: Added VR splash screen for level transitions.
+
+* New: Updated UI elements that composite after reprojection to preserve quality.
+
+* Bugfix: Applied fix for Adreno GearVR devices rendering black on startup.
+
+* Bugfix: Fixed Add Controller Yaw Input with PSVR. The yaw was being reset every frame.
+
+* Bugfix: Fixed an issue where the main editor window was not restored after exiting VR preview in some cases.
+
+* Bugfix: Fixed crash when exiting PIE while a panoramic screenshot was in the process of being taken.
+
+* Bugfix: Fixed motion controller attachment movement when the game is paused.
+
+* Bugfix: Fixed PSVR reprojection prediction time. HMD motion now runs off the edge of the rendered area less easily.
+
+* Bugfix: Fixed saving the Oculus Rift INI setting "bPixelDensityAdaptive". It was previously saved as a floating point value instead of a boolean.
+
+* Bugfix: Fixed some static analysis warnings in Oculus Rift plugin code.
+
+* Bugfix: Fixed the "Get bounds node" for SteamVR to return the correct bounds if they have changed since initialization. This could happen if the HMD was outside of tracking area during start up.
+
+* Worked around a bug in SteamVR that caused VR to fail to initialize if SteamVR isn't running when launching the application.
+
+* Added positional tracking information on PSVR even when the device isn't optically tracked. This enables proper IPD to be taken into account for rendering.
+
+* Added support for the left menu button on the Oculus Touch controllers.
+
+* Adjusted Hidden/Visible area masks to eliminate black crescents at the top and the bottom edges of the PSVR HMD.
+
+* Updated changing world to meters scale while running in VR Preview mode in the Editor on Oculus Rift so that it now correctly applies the changes to the Oculus Touch controllers.
+
+* Updated VR play-in-editor to always use "Game has mouse focus" option, to ensure that motion controllers are routing input correctly to the game.
+
+* Added a PSVR Morpheus project setting to set the reprojection wrap mode to Mirror or Clamp To Border. This determines how left and right reprojection artifact pixels are filled.
+
+    * Mirror, the default, fills them with some nearby pixels.
+
+    * ClampToBorder fills them with the border color (usually black).
+
+    * Most content works best with mirror.
+
+* Updated PSVR to now ignore partial HMD tracking data before the HMD has been tracked in a game session because orientation based re-projection does not function until that happens. There is a project setting to disable this, should the PS4 behavior change in the future.
+
+### Other
+
+* New: Added the ability to change Device Profiles at runtime. 
+
+    * Use "dp.override <name>". If you do it again to another one, it will reset the settings to what they were originally, before applying the second new Device Profile. This is because the second Device Profile may not set all settings the first one did, but we want to undo the first settings that the second doesn't contain.
+
+    * Changed "Set Quality Levels" function, in Scalability.cpp, to no longer change the "SetBy" priority when setting console variables. This now keeps the "SetBy" the same as it was, which helps with conflicts between game settings and Device Profiles. See the "Set With Current Priority" function.
+
+* New: Added an option to toggle full screen with F11 key in addition to Alt+Enter in games.
+
+* New: Added Parse function to JsonObject.cs to be able to parse a string instead of only files.
+
+* New: Env Query Instance Blueprint Wrapper is now Blueprintable.
+
+    * This enables you to make Blueprints that enable more complex 'wrapper' behavior when using the Blueprint node "Run EQS Query" (on the EnvQueryManager).
+
+    * Made "EEnvQueryStatus" a Blueprint type, which makes it much easier to work with in Blueprints.
+
+* New: For Media Assets, added "Can Play Source" and "Can Play Url" methods to Media Player.
+
+* New: For Media Assets, added more verbose logging output, e.g. for dropped video frames. 
+
+* New: For Media Assets, added "Open File" method to Media Player.
+
+* New: For Media Assets, implemented Platform Media Source. 
+
+* New: For Media Assets, Implemented per-platform media player overrides in Media Asset derived classes.
+
+* New: For Media Player Editor, added "Original Size" viewport mode.
+
+* New: For Media Player Editor, added Stats tab to show decoder performance statistics during playback.
+
+* New: For Media Player Editor, added support for dragging and dropping media files into media player editor.
+
+* New: For Media Player Editor, implemented platform player overrides details customization for media assets. 
+
+* New: For Media Player Editor, trimming leading and trailing whitespace in URL textbox.
+
+* New: Project-specific versions of identically named engine plugins are now supported. 
+
+    * Multiple plugins with the same name are prohibited at other times.
+
+* New: For PS4 Media, improved memory allocator to avoid "out-of-memory" issues when playing multiple videos. 
+
+* New: For Wmf Media, added QuickTime to known video subtypes.
+
+* New: For Wmf Media, refactored playback topology handling to support multiple tracks and track switching.
+
+* New: For Avf Media, added missing supported media file extensions for macOS/iOS media player. 
+
+* Bugfix: Fixed a crash when loading Unreal Stats in the Session Frontend. 
+
+* Bugfix: Fixed a crash where sometimes the Crash Reporter Client would crash if it didn't receive a ping responses from the Epic crash server. 
+
+* Bugfix: Fixed a crash when after playing media that requires a different conversion shader. 
+
+* Bugfix: Fixed an issue with window centering when moving from fullscreen to windowed mode.
+
+* Bugfix: Fixed an issue with Media assets where existing native player not being reused when opening new media source.
+
+* Bugfix: For Media assets, improved the appearance of file media source paths. 
+
+* Bugfix: For the Media Player Editor, fixed the media player selection is ignored if media specifies player overrides. 
+
+* Bugfix: For the Media Player Editor, fixed where the URL box would sometimes incorrectly show media load error even though it loaded correctly. 
+
+* Bugfix: For PS4 Media, fixed a crash with the CPU/GPU due to a race condition in destructor.
+
+* Bugfix: For Wmf Media, fixed an issue where FourCC types printing in reverse order.
+
+* Bugfix: For Wmf Media, fixed an issue with GuidToString printing the incorrect byte sequence. 
+
+* Bugfix: For Wmf Media, fixed a memory leak while playing videos. 
+
+* Bugfix: For Wmf Media, fixed an issue where WMV2 encoded videos were too bright.
+
+* Bugfix: Removed still image tracks and sinks. 
+
+* Removed use of Windows Error Report API from engine crash handling. Fixes cases where WER settings in Windows could block UE4 crash report system.
+
+## Programming Release Notes
+
+#### AI
+
+##### Navigation
+
+* Added sanity checks to NavDataSet elements access in UNavigationSystem.
+
+* CrowdAgentMoveDirection will be set while traversing a link, unless the agent is falling.
+
+* Made NavigationSystem call TickAsyncBuild on all navigation data instances while there is an ongoing navigation build in progress, regardless of whether auto navmesh update is enabled.
+
+#### Animation
+
+* New: Made RawAnimationData (and associated structures) private and added accessors and new functionality to enable modification in a safer way than direct access.
+
+#### Blueprints
+
+* New: Added GetConvenientWindowedResolutions to KismetSystemLibrary.
+
+* New: Added ability to use "AdvancedDisplay" metadata in MakeStruct nodes.
+
+* Bugfix: Fixed an issue in Blueprint compilation that could cause crashes with certain memory allocators.
+
+* Bugfix: Fixed an issue where components of inherited grandchild blueprints would sometimes use invalid components as their outer in cooked builds.
+
+#### Core
+
+* New: Added UWorld::SetTimeUntilNextGarbageCollection to change the GC timer for use when doing automated performance capture measurements.
+
+    * Note: Things that force a GC will still force a GC after using this method (and they will also reset the timer).
+
+* New: Added a method to extract swing and twist quaternions from FQuat.
+
+* New: Added a new Invoke function which enables member function pointers and data member pointers to be used as simple projections in algorithms.
+
+* New: Added a Super typedef to USTRUCTs which specify another USTRUCT as their base.
+
+* New: Added a TNot metafunction which inverts the Value of another boolean metafunction.
+
+* New: Added Algo::Sort which sorts a range like TArray::Sort but doesn't auto-dereference pointers.
+
+* New: Added Algo::SortBy which sorts a range via a projection of the elements.
+
+* New: Added an OOM delegate so other systems can get a callback after OOM occurs (after deleting the backup memory pool if it exists).
+
+* New: Added an optional size to MallocLeak Stop and made the default filter size 128 KB for both MallocLeak Dump and MallocLeak Stop if no size was specified.
+
+* New: Added array-new support to TUniquePtr and MakeUnique.
+
+* New: Added AsyncThread function to enable stack size and priority to be specified
+
+* New: Added ENGINE_API to FStripDataFlags so that it can be used outside of the Engine module.
+
+* New: Added extra context to FTableRowBase::OnPostDataImport.
+
+    * It now takes the owning data table and the row name as parameters. This enables it to do more useful fix-up that depends on the context of where it is used.
+
+* New: Added Find and Contains functions to TBitArray.
+
+* New: Added FUNCTION_NO_NULL_RETURN_BEGIN and _END to statically annotate a function that never returns a null pointer.
+
+* New: Added GetData(), GetNum() generic functions, and TElementType, TIsContiguousContainer traits.
+
+    * The generic GetData() is only valid for contiguous containers, which include C arrays, initializer lists, TArray and TArrayView.
+
+* New: Added Invoke support for TFunction, which enables a member function pointer or data member pointer to be bound to a TFunction.
+
+* New: Added Invoke() support to CopyIf
+
+* New: Added MakeShared to enable the construction of an object and its reference controller in a single allocation.
+
+* New: Added new Algo::Count and CountIf, for counting elements in a container which match a given value / predicate.
+
+* New: Added some extra checks to catch uninitialized pointers passed to the garbage collector.
+
+* New: Added support for custom versions in FArchiveProxy.
+
+* New: Added TLess<> and TGreater<> specializations which can take any argument type.
+
+* New: Disabled logging to memory in shipping builds by default. From now on FOutputDeviceMemory will be an opt-in for games.
+
+* New: Improved bitwise swapping for Swap so that it produces more optimal code in general.
+
+* New: New Objects can now be created as duplicate transient or PIE duplicate transient consistent with the flags of properties that reference them.
+
+* New: When the Event Driven Loader is enabled in cooked games, the engine will use the async loading path when loading all assets (even when using LoadObject function).
+
+* New: Added display of the name of the corrupt pak file in various pak file error messages.
+
+* New: Added FArchive::operator<< for enum classes, which serializes the underlying type.
+
+    * Changed generated code to never use TEnumAsByte for generated code for enum classes.
+
+* New: Added TIsEnumClass metafunction, for determining if a type is an enum class.
+
+* New: Added TUseBitwiseSwap<T> type trait for user-defined types to specialize (as 'false') if it enables better code generation.
+
+* Bugfix: Fixed bad format string in FLinkerLoad::VerifyImport.
+
+* Bugfix: Fixed a compile error when declaring a TFunction returning a reference type.
+
+* Bugfix: Fixed all static analysis warnings in editor builds.
+
+* Bugfix: Fixed an Unreal Header Tool code generation error when defining a UFUNCTION with a TArray<TScriptInterface<T>>.
+
+* Bugfix: Fixed bulk data using source data pointer as an archive instead of raw data when saving.
+
+* Bugfix: Fixed constness of ILandscapeEditorModule::GetHeightmapFormatByExtension() and ILandscapeEditorModule::GetWeightmapFormatByExtension().
+
+* Bugfix: Fixed destructor calls not being called on the destination container during a TSparseArray, TMap and TSet move assignment.
+
+* Bugfix: Fixed FWeakObjectPtr/TWeakObjectPtr debugger visualization which depended upon C++ initialization order.
+
+* Bugfix: Fixed TCString::Strspn.
+
+* Bugfix: Fixed TSparseArray debugger visualizer.
+
+* Changed FMallocStomp::GetAllocationSize() to return the requested allocation size instead of the physical allocation size to make FMallocStomp work properly with FMallocPoisonProxy.
+
+* Changed checkSlow() to check() when verifying type compatibility in GetDefaultObject<> because this is potentially an unsafe static cast.
+
+* FindObjectWithOuter once again enables ClassToLookFor to be null as comment indicates is allowed.
+
+* Increased the initial memory allocation size for FLargeMemoryWriter to reduce the number of allocations when saving or cooking.
+
+* Lowered verbosity of several error messages that happen when deleting native properties or components.
+
+* Lowered verbosity of several warning messages that happen when deleting properties or components from parent classes. These issues used to require resaving the child blueprint or level.
+
+* Made FString::Mid assert when passed a negative count.
+
+* Made Unreal Header Tool able to parse final classes.
+
+* Moved RemoveNamesFromMasterList from UEnum destructor to BeginDestroy to avoid potential issues when its package has already been destroyed.
+
+* Removed blueprint stack warning messages from shipping builds, the stack was missing so the warnings were pointless.
+
+* Updated stack captures to now always go down the fast path. This avoids a game going unresponsive on Windows when printing or capturing large numbers of callstacks such as with the 'mallocleak' command.
+
+* Updated the JSON writer to now write floats with a suitable level of precision, where before it was always writing many redundant zeroes.
+
+#### Editor and Tools
+
+* New: Added allocation tagging to MemoryProfiler2
+
+    * See FMallocProfiler::AddTag and FMallocProfiler::RemoveTag. Tags should be in the format "Category:Tag", and the "Tag" part may include "/" to create levels in the tag hierarchy (as shown in the tool).
+
+    * Added object, package, and class tags by default (this uses FScopeCycleCounterUObject, so the stats system is now enabled when malloc profiling).
+
+* New: Optimization and usability improvements of the MemoryProfiler2 tool
+
+    * Optimized the processing of the Callgraph, Histogram, and Short lived allocations views.
+
+    * Added a user-configurable list of allocator functions to trim (which now includes FMemory and operator new by default, and produces much cleaner callstacks).
+
+* Bugfix: Fixed FDetailCategoryImpl::ShouldBeExpanded not honoring bShouldBeInitiallyCollapsed when bRestoreExpansionState was true.
+
+* PostEditComponentMove will now be called correctly when stopping a drag operation of a scene component of an actor in a level.
+
+* UWorldComposition::GetTilesList function can now be used in game code (previously only in Editor).
+
+* Fixed all static analysis warnings in Editor builds.
+
+##### Content Browser
+
+* Bugfix: Fixed an issue where re-scanning a package file may leave old assets in the asset registry.
+
+    * It used to not clear out anything associated with the old package before scanning the file, which could result in old assets being left if they'd since been removed from the package.
+
+    * This also exposes a PackageDeleted function to enable manually clearing anything associated with a package (if doing some custom asset work).
+
+##### World Browser
+
+* UWorldComposition::GetTilesList function can now be used in game code (previously only in Editor)
+
+##### Sequencer
+
+* New: Added interface to extend the matinee to level sequence converter.
+
+#### Gameplay Framework
+
+* New: Actors can now specify via bRunConstructionScriptOnDrag whether PostEditMove should fire the construction script on every move or only when the move is considered finished.
+
+* New: Added a bWasCancelled parameter to GameplayAbility::EndAbility. This enables us to determine if an ability was ended prematurely. Added an OnAbilitySucceeded delegate for determining when an ability was successfully ended.
+
+* New: Added an FName constructor to FCollisionProfileName.
+
+* New: FLocalPlayerContext functions no longer check by default.
+
+* New: UPrimitiveComponent now has an optional list of components it will ignore collision with when performing movement sweeps.
+
+* New: UWorld now has a member that accumulates the dilated/clamped game time even when paused called UnpausedTimeSeconds.
+
+* Bugfix: Fixed an issue where FGameplayTagContainer::HasTag(Tag, IncludeParent, IncludeParent) would miss some matches. Added tests to find other issues.
+
+* Bugfix: Fixed an issue where FindTeleportSpot used when spawning actors would sometimes fail to find a valid nearby spot depending on the previous collision results.
+
+* Bugfix: Fixed CharacterMovement hanging on to a bad/penetrating floor check result and not continuing to check for a valid floor. This only occurred if bAlwaysCheckFloor was false.
+
+    * This could in rare situations cause the character to continue to attempt to depenetrate an object far away from it until another floor check occurred.
+
+* Bugfix: Fixed tooltip text for p.NetEnableListenServerSmoothing command (it was a copy/paste from another command).
+
+* Updated bRemoteOwned to no longer be settable in FSpawnActorParameters as only spawning a replicated Actor from the replication code should ever have been setting this to true.
+
+* Updated Components to no longer need to specify whether they want begin play or not.
+
+* Updated GetNextViewablePlayer to now check and return correct PlayerState.
+
+#### Localization
+
+* New: Added basic support for localizing plugins.
+
+* Wrapped UMaterialExpression::MenuCategories in WITH_EDITORONLY_DATA to avoid gathering it for game-only localization.
+
+#### Networking
+
+* Updated UDPPing to enable a stack size to be set for the individual threads that the ping will be performed on to compensate for third party injected code that exceeds stack space available in default 32k limit.
+
+#### Online
+
+* New: Added AddGenericToInQueueOnlineThread(), useful running a delegate on Online thread instead of game thread.
+
+* New: Added a UChatroom class that does some basic chat room join/create/leave functionality using the online chat interface.
+
+* New: Added an ExpirationDate field to FOnlineStoreOffer. This value is initialized as FDateTime::MaxValue().
+
+* New: Added GetSessionIdStr to FOnlineSessionSearchResult and FOnlineSession as convenience functions.
+
+* New: Added OSS.VoiceLoopback <0/1> to debug chat by including local talkers into remote talkers audio queue. Additionally added "ONLINE VOICE DUMP" exec command to print out known state of voice interface for troubleshooting.
+
+* Removed spurious log message during voice initialization as it is covered by other logging.
+
+* Added proper GetResolvedConnectString functions for PS4/Live.
+
+* UPartyGameState::ResetLocalPlayerState() iterates over local players rather than game state player array for the very rare cases where AGameState is null.
+
+#### Physics
+
+* New: Kinematic actors now excluded from active transforms generation. This saves time in fetchResults.
+
+* New: PhysicsConstraintComponent now exported and can be subclassed outside Engine.
+
+* New: Sweeps can now take a rotation parameter that will enable the shapes used for the sweep to have their orientation changed.
+
+* Bugfix: PhysX no longer returns NaN from some penetration tests.
+
+* Bugfix: Fixed crash from user doing unsafe things during fetchResults by moving OnConstraintBreak delegate so that it fires outside of fetchResults.
+
+* Implemented copy and assignment operator for convexElem struct and aggregate geometry to avoid shallow copy of raw pointers.
+
+* Moved skeletal mesh component physics updating into EndPhysics so that other components can use it safely in PostPhysics.
+
+#### Platforms
+
+* New: Added support for delay load DLLs for Xbox One.
+
+* New: libcurl 7_48_0 - headers and lib updates
+
+* New: libwebsockets 1.7.4 - headers and lib updates
+
+* New: openssl 1_0_2h - headers and lib updates
+
+* New: webRTC rev.12643 - headers and lib updates
+
+* New: zlib 1.2.8 - headers and lib updates
+
+* Bugfix: Fixed a parameter order issue that could cause incorrect compression settings to be used.
+
+* Bugfix: Fixed a threading issue with the virtual keyboard implementation that could cause the use of protected Slate functions from a non-game thread.
+
+* Bugfix: Fixed an issue where stats were using unsupported 64 bit atomics on Windows XP.
+
+* Bugfix: Fixed an issue where the wrong xdk.h header could be included by ShaderCompileWorker if a Tools SDK with a different version than the target XDK was installed.
+
+* Bugfix: Fixed an issue where Windows XP would not build if the Windows 10 SDK was installed.
+
+* Bugfix: Fixed an issue where Xbox One was warning about PDB files missing that are not meant to be present.
+
+* Bugfix: Fixed incompatible VR code being built for Windows XP.
+
+* Converted the PS4MallocCrash class into a generic one.
+
+* Split PLATFORM_HAS_BSD_SOCKET_FEATURE_GETADDRINFO off from PLATFORM_HAS_BSD_SOCKET_FEATURE_GETHOSTNAME.
+
+##### iOS
+
+* New: Added PRAGMA_DISABLE_OPTIMIZATION_ACTUAL and PRAGMA_ENABLE_OPTIMIZATION_ACTUAL defines.
+
+##### Mac
+
+* New: Changed the code responsible for generating RPATH entries to enable third party dylibs on Mac to be loaded from plugin subfolders.
+
+* New: Implemented IsGamepadAttached() function.
+
+* Bugfix: Fixed indexing of iOS and tvOS-specific files in Xcode project.
+
+* Changed FString(const NSString* In) constructor to use CFStringGetBytes instead of [NSString cStringUsingEncoding]. It's faster and solves the problem with the previous implementation's incorrect assumption that UTF32 string returned from cStringUsingEncoding would be zero-terminated.
+
+##### Playstation 4
+
+* New: Added support for displaying one and two button message dialogs on Playstation 4.
+
+* New: Refactored symbol uploading in Unreal Automation Tool (UAT). The new system supports both Windows and Playstation, and handles the submission of build symbols to a symbol server as part of a UAT build script.
+
+* New: The "Play Together" feature is now implemented on PS4! You can override OnPlayTogetherEventReceived on UGameInstance to handle the event. Check out ShooterGame for an example implementation.
+
+* New: Upgraded to PlayStation 4 SDK version 4.008.061.
+
+* Bugfix: Fixed a bug introduced by trying to get the taskgraph broadcast to work on PC.
+
+* Changed the Playstation 4 exit behavior to avoid calling abort() when shutting down cleanly.
+
+* Updated GetBackMemoryPoolSize to now correctly return an integer.
+
+* Moved the Playstation 4 crash dump user callback registration to a much earlier point in engine initialization, to catch early crashes in engine preinit.
+
+##### Xbox One
+
+* New: Added support for staging and loading DLLs on Xbox One.
+
+* Bugfix: Fixed a crash bug when using "recompileshaders changed" on Xbox One, with a cook-on-the-fly server.
+
+* Bugfix: Fixed an issue that could cause a memory overwrite when using compute fences.
+
+* Bugfix: Fixed a compile warning when the current XDK is built using the RTM version of Visual Studio 2015.
+
+    * Note: Only Visual Studio Update 3 is officially supported in this release.
+
+* Bugfix: Fixed the incorrect DirectX header being included for the Xbox One target platform.
+
+##### HTML5
+
+* Bugfix: Fixed a crash in WebSocketNetDriver. A chunk of code was missing that calls PacketHandler's "packet modifiers."
+
+* Updated HTML5 plugin OnRawReceive overflow to check for minimum expected size upon data received from network.
+
+##### Windows
+
+* Bugfix: Fixed a crash bug during adapter selection, in the D3D11 RHI.
+
+#### Programming
+
+* New: Upgraded the Perforce C++ API libraries to version 2015.2.
+
+    * This version adds official support for Visual Studio 2015 from Perforce.
+
+    * This also fixes an issue where the debug version of the Perforce C++ API libraries was not built with the correct debug settings.
+
+* New: Added list of Dependant modules to warning about code that fails our EULA check.
+
+* Bugfix: Fixed a typo in a define that could cause Visual Studio 2015 to crash.
+
+* Bugfix: Fixed Clean process during a Hot Reload.
+
+* Bugfix: Fixed issue where project-specific modules could be considered out of date, when building multiple targets through Visual Studio at once. Unique build id strings in .modules files are now shared until an existing file is modified.
+
+* Bugfix: Implemented various fixes for building with -nosharedPCH.
+
+* Bugfix: Implemented various fixes to enable the Editor to be compiled in Monolithic mode.
+
+#### Rendering
+
+* New: Added delegates called after resolving scene color (e.g. to do custom post-processing).
+
+* New: Added direct support for the creation and binding of Pipeline State Objects at the RHI level. This is the first step on the road to moving PSO handling to the high level rendering code to have a unified solution for DX12, Vulkan, and Metal.
+
+* New: Updated DrawDebug functions to now compile out in Shipping and Test builds.
+
+    * If you wish to generate compile errors for trying to call them at all, you can define SHIPPING_DRAW_DEBUG_ERROR=1
+
+* New: Exposed GPU revision ID as GRHIDeviceRevision and added it to the FPS chart analytics (gathered on D3D11 and D3D12 only).
+
+* New: Added the first implementation of the new RHI functions for Pipeline State Objects on DX12.
+
+* New: Refactored the GL shader compiler to enable GLSL variant platforms to modify behavior without a lot of #if checks everywhere in the GL code.
+
+    * Added ability in the cross compiler to convert the global uniform arrays into named uniform buffer objects.
+
+    * Added ability for GL shader compiler to return the original resources names ("VertColor" instead of "u_v[3]" or whatever) to the shader compiler shader parameter metadata. This does NOT change the generated GLSL code.
+
+* Bugfix: Fixed async compute parameter problems on PS4 by adding SetShaderValue() specialization for bools on AsyncCompute commandlists.
+
+* Bugfix: Fixed a D3D error with mismatched vertex/pixel shader registers for SV_POSITION input.
+
+* Bugfix: Fixed an ensure fail when initializing point light shadow maps and simplified the logic for cubemap render target resolution parameters.
+
+* Bugfix: Fixed a race condition where debug drawing code was trying to read data off the SceneProxy from the Gamethread.
+
+* Exposed full USceneCaptureComponent classes to plugins instead of select methods.
+
+##### Materials
+
+* New: Refactored material properties during graph translation.
+
+    * Fixed material properties are being phased out in favor of better encapsulated and more extensible attribute lists.
+
+    * This is the first in a series of changes moving towards per-project attribute controls and improved high-level usability.
+
+#### UI
+
+##### Slate
+
+* Bugfix: Fixed SWrapBox InnerSlotPadding.Y not applying between the first and second lines.
+
+#### VR
+
+* New: GetNumOfTrackingSensors and GetTrackingSensorProperties are now supported by SteamVR as well as Oculus Rift.
+
+* Bugfix: Fixed saving of a few VR related config flags. They were being saved under the wrong name, preventing them from being read in again.
+
+* Bugfix: Fixed the tracking camera orientation returned from GetTrackingSensorProperties on the Oculus Rift to now point in the correct direction. Previously it was pointing towards the back of the camera instead of its facing direction.
+
+#### Other
+
+* New: Added ability for custom implementations of UGameplayModMagnitudeCalculation to specify they have dependencies external from the ability system that may invalidate their calculated magnitudes. Dependencies are specified by multicast delegate that is expected to fire when the magnitude needs to be updated.
+
+    * Also added advanced variant of the implementation enabling clients to also trigger the dependency invalidation and update themselves appropriately. Designed to only be used for extreme cases (such as games relying heavily on network dormancy) and only on attributes that are not gameplay-impactful (and therefore safe for the client to locally compute).
+
+* New: Added new API function for DrawDebugSolidBox, which takes an FBox instead of a Center and Extent. It also enables specifying a transform as well (but defaults to using the identity transform for convenience).
+
+* New: Added option to disable uploads to Crash Report Receivers while they continue to answer pings from clients.
+
+* New: Crash Report Process now has option to output crash files to Amazon S3.
+
+* New: Crash Report Process now has optional, dedicated PS4 crashes queue.
+
+* New: Crash Report Process performance improvements.
+
+    * Added -ForceUsePDBCache option to MinidumpDiagnostics to early abort runs that have no valid PDB cache entry. Early out for slowest reports.
+
+    * Changed MinidumpDiagnostics logging locations to avoid speed cost of logging to a folder with thousands of files.
+
+    * Improved CRP shutdown speed when updating.
+
+* New: Crash Reporter backend performance improvements.
+
+    * Crash Report Process can discard crashes when crash queues become too long.
+
+    * Crash Report Process can run multiple instances of MinidumpDiagnostics to speed up callstack generation.
+
+    * MinidumpDiagnostics can place a system-wide lock on resources (Perforce workspace and the PDB cache) to enable multiple instances without multiplying storage requirements.
+
+* New: CrashReportProcess logging and Slack reporting improvements to monitor disk space.
+
+* New: CrashReportProcess passes new BuildVersion string to MinidumpDiagnostics for symbol lookup.
+
+* New: Introduced the concept of "evaluation channels" to non-instant gameplay effects within the ability system! Evaluation channels enable game-specific control over the order in which modifiers are evaluated in attribute calculations.
+
+    * A game may specify how many channels it would like to use, as well as what they are named, and then modifiers can specify which channel they work on. The channels are evaluated in order, with the numerical output of the first channel serving as the input to the second channel, and so on.
+
+    * Example: Imagine a sample attribute with a base value of 100. Now apply a multiplicative mod with magnitude 1.1 in channel 0 and a multiplicative mod with magnitude 1.1 in channel 1. The result will evaluate as ((100 * 1.1) * 1.1) instead of (100 * 1.2).
+
+    * To use channels in a game, two settings within the INI section for AbilitySystemGlobals must be configured. First, bAllowGameplayModEvaluationChannels must be set to true, and second, the channels that should be used (and their names) must be specified in GameplayModEvaluationChannelAliases. Additionally, a game can optionally specify a default channel to use via DefaultGameplayModEvaluationChannel.
+
+    * In addition, added various new channel-related functionality to the ability classes' API, such as the ability to base the magnitude of a modifier off of an attribute only evaluated up to a certain channel.
+
+* New: Large Crash Report Process update.
+
+    * Removed legacy queues for deprecated Crash Report Receivers.
+
+    * Removed duplication tracking needed when using Crash Report Receivers as this was slow and didn't scale.
+
+* Bugfix: Fixed a header issue that was including 64k of test data into shipping builds, it is now stripped out appropriately.
+
+* Bugfix: Fixed issue wherein gameplay effects with non-snapshotted, attribute-based modifiers could potentially not replicate to clients correctly, resulting in incorrect values on clients.
+
+* Bugfix: Fixed problem where MinidumpDiagnostics used on CR backend would strip unnecessary lines from callstacks.
+
+## Upgrade Notes
+
+### C++ API Changes
+
+* New: Made the "RemoveAt" and "RemoveAtSwap" functions in TArray trigger a compile error when a boolean argument is passed as the count.
+
+* New: Renamed the Strfind function in TCString to Strifind, as it is case-insensitive.
+
+    * Added the Strfind function to TCString, which is actually case-sensitive.
+
+* New: The default allocator for cooked Windows games is now BinnedMalloc2 to match other platforms. The argument "-binnedmalloc" will switch back to the older allocator.
+
+* New: Component Activate and Deactivate events now supply the component being affected as an argument. Any C++ functions bound to Component Activate and Deactivate events will need to add a component argument to the signature.
+
+* New: libwebsockets 1.7.4 - headers and lib updates. The libwebsockets API has changed.
+
+* New: Skip cooking editor content now defaults to true in project settings.
+
+    * If your project uses content in "Engine\Editor*" you may need to disable the "Skip cooking editor content" project setting.
+
+* New: Refactored the View Matrices struct.
+
+    * The members of the View Matrices struct have become private, there are now appropriate Getters for the values. Heavy cost functions have been renamed to Compute* and functions that mutate the struct have been renamed to Hack*
+
+* New: Split "RHIClear" into "RHIClearDepthStencilTexture," and "RHIClearColorTexture."
+
+    * Anyone using the Clear function in "RHICmdList" will now need to use the appropriate version, passing in the current render target(s).
+
+* New: Moved the "Export/Import Text Item" functions from "Struct Property" to "Export/Import Text" functions in "Script Struct", and added a parameter to say whether to call native callbacks. This enables calling the standard "Import Text" from a native "Import Text Item" followed by post processing.
+
+* New: Refactored the "Add Item Data" function in "Env Query Instance" to get rid of the second template parameter and make it more type-safe.
+
+* New: Added "Ge ResourceSizeEx" and "GetResourceSizeBytes," and deprecated "GetResourceSize."
+
+    * "Get ResourceSizeEx" populates a struct which reports not only how much memory is allocated, but also which arenas the memory was allocated from.
+
+    * "Get ResourceSizeBytes" just wraps a call to "GetResourceSizeEx" and returns the total size from all arenas (this has the same behavior as the now deprecated "GetResourceSize").
+
+    * Classes that used to override "GetResourceSize" should instead override "GetResourceSizeEx" and report their allocations as appropriate.
+
+* New: Deprecated the "Get Bone Tree" function in the Skeleton class and removed references in engine.
+
+* The "Clear Pause" function in the "GameMode" class now returns a boolean. Any virtual overrides will need to have their signature changed and return the appropriate result.
+
+* Made it a compile error to copy the "AIMessageObserver" class, as it was not designed to be copied. You should change the logic of your code to avoid copies of this type.
+
+* Removed Publish Dir property from Installed Engine Build script as it is no longer used.
+
+* Removed obsolete array slack tracking code.
+
+* Removed unused "ToBuildInfoString" function declaration.
+
+* Added the "final" specifier to the UWorld class. UWorld is now non-inheritable.
+
+* Removed some obsolete code and state from the "ScriptStruct" class.
+
+    * The "InheritedCppStructOps" function has been removed. Any calls to it should be replaced with "false."
+
+* Since the HMD module names used when sorting by priority are now also used by the `hmd=<name>` command line option, the "GetModulePriorityKeyName" function in the "HeadMountedDisplayModule" class has now been renamed to "Get Module Key Name".
+
+* Removed the "UnregisterDebugDrawDelegate," "RegisterDebugDrawDelegate" and "DrawDebugLabels" functions from the  "DebugRenderSceneProxy" class and moved them into a new struct called "Debug Draw Delegate Helper."
+
+* Direct access to the Static Mesh property in the Static Mesh Component class has been deprecated. Please use "Get Static Mesh" and "Set Static Mesh" accessors.
+
+* The version of the Create Component function in the Audio Device class that took a long list of arguments has been deprecated. Use the version that takes a parameter struct instead.
+
+* In order to support asymmetric POVs, such as the Vive , the number of "Get Tracking Sensor Properties" has been expanded to return four instead of two POV values. The now deprecated "Get Positional Tracking Camera Properties" is still available, but only returns horizontal and vertical POVs and only supports a single tracking device.
+
+    * If you were previously using the data returned from "Get Tracking Sensor Properties" (or the now deprecated "Get Positional Tracking Camera Properties") on Oculus Rift, you may have to unflip any corrections you were doing on the orientation. This will also make your code compatible with other HMD plugins that support returning information about tracking devices.
+
+* Localization API changes for per-identity exporting:
+
+    * Moved "TextLocalizationResourceGenerator" from Core to Internationalization.
+
+    * Removed "IJsonInternationalizationManifestSerializer" and "IJsonInternationalizationArchiveSerializer."
+
+    * JsonInternationalizationManifestSerializer and FJsonInternationalizationArchiveSerializer now use static functions (due to no longer inheriting or requiring polymorphism).
+
+    * Removed the "Load From Manifest And Archives" function from "Text Localization Manager" (this is now handled by the "Generate And Update Live Entries From Config" function in "Text Localization Resource Generator").
+
+    * Removed the RepairLocalizationData commandlet (it was for fixing data pre-dating the 4.0 release).
+
+    * Removed the "FindEntryBySource" function from "InternationalizationArchive" as it is no-longer safe to use (use the "FindEntryByKey" function instead).
+
+* Refactored single node animation player APIs.
+
+    * Animation player APIs are not safe to be used during construction script because they're not being serialized, so marked not safe for construction script.
+
+    * Exposed override function that can be used in editor.
+
+* Replaced some hacky post-reflection-capture functions with "RHISubmitCommandsAndFlushGPU".
+
+* Made "FindShaderResourceById" and "Find Shader By Id" return a raw pointer instead of a TRefCountPtr to prevent creating temporary TRefCountPtrs. It is no longer necessary to use TRefCountPtrs as shader serialization has changed and shaders are no longer registered on the async loading thread.
+
+* Updated to Playstation 4 4.008.061
+
+### Core
+
+* Any projects that make use of -log or -abslog to specify log file names must make sure they are only using .txt or .log extensions.
+
+* FString's "Mid" no longer accepts negative counts as the resulting behavior was surprising. Please replace any usage with the desired behavior.
+
+* TArray "RemoveAt" and "RemoveAtSwap" no longer accept a boolean argument as a count, as it is likely this is intended to be the flag for "Allow Shrinking". 
+
+    * The argument should be cast to a non-boolean type if this was intentional or the code should be fixed if it is not.
+
+* TCString's "Strfind" is now case-sensitive. "Strifind" should be used instead to get case-insensitive behavior.
+
+* TEnumAsByte is no longer used for generated code involving enum classes. 
+
+    * If you get errors involving TEnumAsByte with your UENUM enum class, please provide a uint8 base to your type and remove the TEnumAsByte usage.
+
+### Gameplay Framework
+
+* Check the Game Mode Base notes and documentation to see if you should switch your existing GameMode to inherit from the new minimal Game Mode Base class.
+
+    * Game Mode Base's "Start New Player" has been deprecated, and split into "Initialize HUD For Player" and "Handle Starting New Player" which can be overridden in Blueprints.
+
+    * Several Login functions on Game Mode Base that take TSharedPtr<const FUniqueNetId> are now deprecated correctly, they stopped working in 4.13, but were not deprecated
+
+    * GameState's "Should Show Gore" and related matinee functions have been deprecated.
+
+* DumpFPSChartAnalytics has been removed.
+
+    * Games that used it should switch to creating their own instance of FPerformanceTrackingChart and call DumpChartToAnalyticsParams on it directly.
+
+* HTML output for stopfpschart is now generated to a single file rather than two duplicate files, which now use both map name and capture time as part of the file name.
+
+* In general, games should have no reason to programmatically call GEngine->StartFPSChart anymore, instead creating their own instance (this prevents conflicts when using the engine console commands).
+
+* Previously the delta time value of a Tick Interval tick function was the same as if the function had no tick interval. 
+
+    * This means tick functions that were relying on getting just the frame time need to get the delta time from the world rather than using the value passed into the Tick function.
+
+* Refactored FPS chart creation to make it modular so many performance data consumers can be active at once, enabling greater flexibility and decoupling game analytics from FPS chart exec commands
+
+    * DumpFPSChartAnalytics has been removed, games that used it should switch to creating their own instance of FPerformanceTrackingChart and call DumpChartToAnalyticsParams on it directly.
+
+    * In general games should have no reason to programmatically call GEngine->StartFPSChart anymore, instead creating their own instance (this prevents conflicts when using the engine console commands).
+
+    * HTML output for stopfpschart is now generated to a single file rather than two duplicate files (using both map name and capture time as part of the file name).
+
+    * Removed PauseFPSChart, IsFPSChartActive, and GetFPSChartBoundByFrameCounts to reflect that the GEngine instances aren't meant for external use (Start/Stop are left public for automated testing that wants to use them to do logging, but may also be moved private in the future).
+
+* Removed PauseFPSChart, IsFPSChartActive, and GetFPSChartBoundByFrameCounts to reflect that the GEngine instances aren't meant for external use (Start/Stop are left public for automated testing that wants to use them to do logging, but may also be moved private in the future).
+
+* Several Login functions on Game Mode Base that take TSharedPtr<const FUniqueNetId> are now deprecated correctly, they stopped working in 4.13 but were not deprecated
+
+* The config data for picking GameMode has moved to Game Maps Settings, with the old locations deprecated. Check BaseEngine.ini for an example.
+
+* Please replace all uses of EWorldType::Preview with either EWorldType's "EditorPreview" or "GamePreview."
+
+* If you desire the checking behavior of FLocalPlayerContext, you will need to change your code to opt in to it.
+
+* Actor Component's flag for "Wants Begin Play" is now deprecated and unused and all components will have "BeginPlay" called on them. Simply remove any code that sets "bWantsBeginPlay."
+
+### Platforms
+
+#### Playstation 4
+
+* PS4 Time's "System Time" now properly returns the local machine time instead of UTC. 
+
+    * Existing projects using "System Time" should investigate if this change affects their game logic. 
+
+    * To continue using UTC time call "Utc Time" instead.
+
+* For projects targeting PlayStation 4, the engine now requires PlayStation 4 SDK version 4.008.061. 
+
+    * Please ensure to update your installed SDK to the correct version, for the engine to compile correctly. This can be done with Sony's SDK manager utility.
+
+#### Windows
+
+* We are sunsetting support for Windows XP, 4.14 will be the last version to officially support running on Windows XP.
+
+### Rendering
+
+* The PS4 defrag pool size is now controlled by a new variable "r.PS4DefragPoolSize", and is not bound anymore to "r.Streaming.PoolSize".
+
+    * It will need to be configured to fit the project needs, in term of graphic resources.
+
+* Try and avoid mutating the FViewMatrices and FSceneView on your own. It is currently practiced within the engine, but we would like to get rid of that.
+
+* The component should own a FDebugDrawDelegateHelper (or derived object) that copies off all needed state for drawing during Draw Debug Labels (which is called from the Game Thread) during Create Scene Proxy. 
+
+    * The function to copy off the state is usually called Init Delegate Helper. Then the Helper is used to draw the debug information instead of the Scene Proxy (which is solely owned by the RenderThread at this point).
+
+#### Lighting
+
+* Due to a bug introduced 2 years ago, all clear coat and cloth was adding lighting twice for all direct lights. This has now been corrected so relevant materials will appear half as bright when in direct light. Indirect lighting such as lightmaps and reflections is unaffected. This makes it tricky to automate material fix up. We suggest reviewing all materials and deciding on a case by case basis whether modifications are required.
+
+* Fixed Reflection Captures having incorrect brightness when a Stationary Skylight is present.
+
+* Note that Lightmap data has been moved to a separate package, named *_BuiltData. Be sure to save these and check them into source control after doing a lighting build!
+
+#### Materials
+
+* There may be lighting differences in certain cases with materials using Subsurface Profile due to the fixes described above. For the most part, the changes should improve the rendering quality, but it may be necessary to make some minor adjustments to affected materials.
+
+### Other
+
+* Epic's live Crash Report Receivers no longer accept crash reports. Newer Crash Report Clients (from 4.11 onwards) upload to our Data Router endpoint.
+
