Copying //UE4/Dev-Documentation to Samples-Main (//UE4/Samples-Main)

diff --git a/Engine/Documentation/Source/UE4Home.INT.udn b/Engine/Documentation/Source/UE4Home.INT.udn
--- a/Engine/Documentation/Source/UE4Home.INT.udn
+++ b/Engine/Documentation/Source/UE4Home.INT.udn
@@ -4,6 +4,7 @@ Crumbs:
 Description: Complete resources for learning to use Unreal Engine 4
 Navigation:home
 social-image:home_social.png
+type:landing
 
 [PUBLISH:Rocket]
 ## Search
@@ -64,7 +65,7 @@ social-image:home_social.png
 
 [/PUBLISH:Rocket]
 
-[DIR(end:"1" output:"role" org:"hierarchy" parent:"%ROOT%")]
+[DIR(end:"1" output:"role" org:"hierarchy" parent:"%ROOT%" type:"landing")]
 
 <!--
 [REGION:frontpage]
