Copying //UE4/Dev-Editor to //UE4/Dev-Main (Source: //UE4/Dev-Editor @ 3279756)

==========================
MAJOR FEATURES + CHANGES
==========================

Change 3229490 on 2016/12/09 by [REDACTED_PERSON]

	Integrated fix to support named changelists in SVN

Change 3229574 on 2016/12/09 by [REDACTED_PERSON]

	Fix actor mobility getting changed on scene reimport

Change 3229692 on 2016/12/09 by [REDACTED_PERSON]

	Fixing an XML Parser assert when parsing a root tag that completes on the same line.

Change 3230582 on 2016/12/12 by [REDACTED_PERSON]

	PR #3024: Correct the outdated error message instructions for how to fix being unable to launch on an iOS device. (Contributed by [REDACTED_PERSON])

Change 3231470 on 2016/12/12 by [REDACTED_PERSON]

	Eliminate editor sounds that play when  you PIE, simulate or possess the player.  They get in the way of game sounds, are annoying to hear when you are constantly starting and stopping pie, and flush async loading that the game might be doing when they load.

Change 3231475 on 2016/12/12 by [REDACTED_PERSON]

Change 3231476 on 2016/12/12 by [REDACTED_PERSON]

Change 3231477 on 2016/12/12 by [REDACTED_PERSON]

Change 3231479 on 2016/12/12 by [REDACTED_PERSON]

Change 3231480 on 2016/12/12 by [REDACTED_PERSON]

Change 3231508 on 2016/12/12 by [REDACTED_PERSON]

	Removing todo comment

Change 3231603 on 2016/12/12 by [REDACTED_PERSON]

	Exposed a 0-1 UV set and the scaled pixel size for Box and Border brushes
	Also added a material function that exposes all of the current UV sets with nice names instead of indexed coordinates

Change 3231618 on 2016/12/12 by [REDACTED_PERSON]

Change 3231638 on 2016/12/12 by [REDACTED_PERSON]

	Add RF_Transactional to the list of default flags for creating or importing new assets.  All should be transactional by default

Change 3231642 on 2016/12/12 by [REDACTED_PERSON]

	Brighten up the output log by default

Change 3231648 on 2016/12/12 by [REDACTED_PERSON]

Change 3231666 on 2016/12/12 by [REDACTED_PERSON]

Change 3231674 on 2016/12/12 by [REDACTED_PERSON]

Change 3231745 on 2016/12/12 by [REDACTED_PERSON]

	Back out changelist 3231477 to fix build error C2259

Change 3232417 on 2016/12/13 by [REDACTED_PERSON]

	Add the following attributes to the Editor.Usage.FBX.Import EngineAnalytics event
	- FBX Version
	- Filename Hash
	- Import Type

Change 3232477 on 2016/12/13 by [REDACTED_PERSON]

Change 3232571 on 2016/12/13 by [REDACTED_PERSON]

	Back out changelist 3231745

Change 3232675 on 2016/12/13 by [REDACTED_PERSON]

	Fix a crash when reordering material with a fbx containing unused materials, add a fbx automation test to prevent similar issue.

Change 3232975 on 2016/12/13 by [REDACTED_PERSON]

	Fix to build error C2259 for the IPluginWizardDefinition API change.

Change 3233146 on 2016/12/13 by [REDACTED_PERSON]

	Fixed a rounding errors resulting in not flattening to the specified height
	Fixed a rounding error resulting in LandscapeDataAccess::GetTexHeight not always returning the appropriate value

Change 3233153 on 2016/12/13 by [REDACTED_PERSON]

	We cannot anymore change the instance override materials array topology, the topology is limited by the mesh materials array

Change 3234406 on 2016/12/14 by [REDACTED_PERSON]

	Fix window handle and device context being accessed by scene viewports after the underlying window has been destroyed by the OS.  This is an invalid state on linux and using some vr devices.

Change 3234485 on 2016/12/14 by [REDACTED_PERSON]

	tentative build fix for Mac

Change 3234495 on 2016/12/14 by [REDACTED_PERSON]

	Made a setting to control if PIE enter and exit sounds are played.  Off by default

Change 3236709 on 2016/12/15 by [REDACTED_PERSON]

	Fix camera export rotation offset

Change 3236782 on 2016/12/15 by [REDACTED_PERSON]

	Fixed EmitTermExpr failing to use the correct package ID

	FBPTerminal::Source used to be set to the pin, however when pins were moved away from being UObjects, FBPTerminal::SourcePin was added and FBPTerminal::Source is typically null.

Change 3236853 on 2016/12/15 by [REDACTED_PERSON]

	Fix the serialization of the staticmesh property FMeshSectionInfoMap

Change 3236890 on 2016/12/15 by [REDACTED_PERSON]

	Remove old define

Change 3239328 on 2016/12/18 by [REDACTED_PERSON]

	Fixed Focus Viewport action in Static Mesh Viewport. Problem was that the conversion to Orbit Camera for storing the camera position was trashing the desired position during cvamera transitions.  Orbit camera position is now only stored at the end of a transition.

Change 3239660 on 2016/12/19 by [REDACTED_PERSON]

Change 3239662 on 2016/12/19 by [REDACTED_PERSON]

Change 3239663 on 2016/12/19 by [REDACTED_PERSON]

Change 3240696 on 2016/12/20 by [REDACTED_PERSON]

	Added paddiing to columns view
	Added auto resize of column when double clicking on splitter handle in the header
	Remove right number alignment after discussion with Matt K.

Change 3240758 on 2016/12/20 by [REDACTED_PERSON]

	added missing non abstract implementation

Change 3240782 on 2016/12/20 by [REDACTED_PERSON]

	Added missing documentation for content browser column auto resizing

Change 3240817 on 2016/12/20 by [REDACTED_PERSON]

Change 3240834 on 2016/12/20 by [REDACTED_PERSON]

	tentative fix for build error

Change 3240984 on 2016/12/20 by [REDACTED_PERSON]

	Removed unnecessary functions

Change 3241174 on 2016/12/20 by [REDACTED_PERSON]

	Fix compile errors

Change 3241966 on 2016/12/21 by [REDACTED_PERSON]

	Fixed Typo and changed execution order in "ComboBoxString" Component
	[UE-38994] - GitHub 2971 : Fixed Typo and changed execution order in "ComboBoxString" Component

	PR #2971: Fixed Typo and changed execution order in "ComboBoxString" Component (Contributed by [REDACTED_PERSON])
	#github https://github.com/EpicGames/UnrealEngine/pull/2971

Change 3242126 on 2016/12/21 by [REDACTED_PERSON]

	Back out changelist 3236853
	We have to back out this change list because the change was implement in the 4.15 release branch and the EditorObjectVersion.h change is now implement in the ReleaseObjectVersion.h.

Change 3244492 on 2017/01/02 by [REDACTED_PERSON]

	Improved error message

Change 3244545 on 2017/01/02 by [REDACTED_PERSON]

	Navigation - Making it so we don't attempt to load HotReload during shutdown, we only access it if it's still loaded.

Change 3244549 on 2017/01/02 by [REDACTED_PERSON]

	Slate - Implementing custom hardware cursor loading across Windows, Mac and Linux and supports loading cursors from PAK files.  All platforms support loading PNGs through the FHardwareCursor interface.  Some platforms support additional formats, for multiresolution support, but there's a naming convention that can be used on PNGs for the same capability.  All of it is documented in the FHardwareCursor header.  The platform layer for ICursor, now has support for replacing cursor shapes as an override, and can be reset safely.

	The FHardwareCursor supports loading cursors from raw pixel buffers as well, the plan is to allow for the option to UTextures to also be used for hardware cursors.

	Now users through C++ can load and replace the hardware cursors with custom ones of their own,

	e.g. FSlateApplication::Get().RegisterCursor(EMouseCursor::Default, MakeShareable(new FHardwareCursor(FPaths::GameContentDir() / "Slate/FancyPointer", FIntPoint(0,0))));

	The next step is to expose a game friendly layer that supports caching cursors, and letting users change them out by name, without a bunch of destruction of OS resources.

Change 3244845 on 2017/01/03 by [REDACTED_PERSON]

	Fixing typo

Change 3244903 on 2017/01/03 by [REDACTED_PERSON]

	PR #3044: fix link error when FAssetData::PrintAssetData() is used in project (Contributed by [REDACTED_PERSON])

Change 3245125 on 2017/01/03 by [REDACTED_PERSON]

	Put back the dev-editor version because there was some data create before we back it out

Change 3246106 on 2017/01/04 by [REDACTED_PERSON]

	Removed broken CrashReportReciever pre-upload phase from CrashReportClient.
	[UE-40153] - CrashReportClient fails when used in legacy mode with a CrashReportReciever

Change 3246251 on 2017/01/04 by [REDACTED_PERSON]

Change 3246252 on 2017/01/04 by [REDACTED_PERSON]

	The "Map Directory To" field when setting directories to monitor for auto-reimport has also been changed to use the content browser path picker instead of relying on the user to manually enter a mount point.

Change 3247620 on 2017/01/05 by [REDACTED_PERSON]

	Automation - Removing an adjustment to the number of shots we take for high res shots.

Change 3247621 on 2017/01/05 by [REDACTED_PERSON]

	Automation - Adding a few more rendering tests to the cornell box.

Change 3247629 on 2017/01/05 by [REDACTED_PERSON]

	Automation - Improving the comparison row display for screenshots so it's obvious what each image represents.

Change 3248811 on 2017/01/05 by [REDACTED_PERSON]

	PR #3091: Removed unnecessary UPackage casts (Contributed by [REDACTED_PERSON])

Change 3248860 on 2017/01/06 by [REDACTED_PERSON]

	Made the plugin browser select the "built in" category by default instead of the 2D category.  There is no reason for a sub-category to be selected first as it makes searching for plugins globally an extra click because you have to click on the base category first

Change 3249264 on 2017/01/06 by [REDACTED_PERSON]

	Fixed automation test warnings

Change 3249481 on 2017/01/06 by [REDACTED_PERSON]

	Also fill new component added with the tool from neighbours predominance

Change 3249505 on 2017/01/06 by [REDACTED_PERSON]

	PR #3093: Include guard cleanup (Contributed by [REDACTED_PERSON])

Change 3249544 on 2017/01/06 by [REDACTED_PERSON]

Change 3250738 on 2017/01/09 by [REDACTED_PERSON]

	UMG - The WIC now checks if the Widget is enabled before it claims that it's over an interactable or keyboard focusable widget.

Change 3250865 on 2017/01/09 by [REDACTED_PERSON]

	Slate - Updating EAutoCenter and ESizingRule to use the newer enum class style enums.

Change 3250867 on 2017/01/09 by [REDACTED_PERSON]

	Slate - Adding more logging to the hardware cursor code so that it reports more information when it doesn't find an exact match when it comes to cursor size.

Change 3250936 on 2017/01/09 by [REDACTED_PERSON]

	Automation - Refactoring the screenshot comparison tool to no longer require one one generated report.  Doing screenshot comparions now generates individual reports for each failed comparison so that they can be evaluated in bits, and as changes occur as the user reviews aspects, we can remove the reports.  There is now async image loading for the comparison view so that it doesn't hitch.

Change 3250937 on 2017/01/09 by [REDACTED_PERSON]

	Automation - Adding another example to the CornellBox test.

Change 3250958 on 2017/01/09 by [REDACTED_PERSON]

	Slate - Fixing some other cases where people were referring to ESizingRule::Type.

Change 3251162 on 2017/01/09 by [REDACTED_PERSON]

	Slate - Fixing some other cases where people were referring to ESizingRule::Type.

Change 3251254 on 2017/01/09 by [REDACTED_PERSON]

	Attempt to fix static analysis warnings

Change 3251373 on 2017/01/09 by [REDACTED_PERSON]

	Core - Now writing a log warning instead of ensuring if calling LoadModule wouldn't have been safe to do here, depending on load order.

Change 3251525 on 2017/01/09 by [REDACTED_PERSON]

	Automation - Fixing a build issue in ImageComparer.

Change 3252321 on 2017/01/10 by [REDACTED_PERSON]

Change 3252354 on 2017/01/10 by [REDACTED_PERSON]

	Image Compare - Fixing a potential threading hazard in the image comparer.

Change 3252356 on 2017/01/10 by [REDACTED_PERSON]

	Automation - The screenshot metadata now captures the commit/CL that the screenshot was taken at and records it in the metadata.

Change 3252601 on 2017/01/10 by [REDACTED_PERSON]

	Fbx automation test, reload feature implementation

Change 3252761 on 2017/01/10 by [REDACTED_PERSON]

	Fixing some IWYU errors with PCH disabled

Change 3252765 on 2017/01/10 by [REDACTED_PERSON]

	Fixing some static analysis warnings

Change 3252793 on 2017/01/10 by [REDACTED_PERSON]

	Fixing FText natvis

	The text data visualizers have to be defined before the text visualizer

Change 3253987 on 2017/01/11 by [REDACTED_PERSON]

	PR #3108: Git Plugin: use asynchronous "MarkForAdd" and "CheckIn" operations for the initial commit (Contributed by [REDACTED_PERSON])

Change 3254378 on 2017/01/11 by [REDACTED_PERSON]

	Refactor scene importing to allow for plugins to make scene importers

Change 3254679 on 2017/01/11 by [REDACTED_PERSON]

	Fix calling LoadModule in perforce source control off the main thread

Change 3256472 on 2017/01/12 by [REDACTED_PERSON]

	Improved error reporting from IncludeTool

	- The error reporting was using zero-based line indices which was misleading.
	- The error reporting now includes the offending line to remove ambiguity.

Change 3256725 on 2017/01/13 by [REDACTED_PERSON]

	IncludeTool can now parse typedef in Fwd headers

Change 3256758 on 2017/01/13 by [REDACTED_PERSON]

	Added support for String Tables

	String Tables provide a way to centralize your localized text into one (or several) known locations, and then reference the entries within a string table from other assets or code in a robust way that allows for easy re-use of localized text.

	String Tables can be defined in C++ (using the LOCTABLE family of macros), loaded via CSV file, or created as an asset. They can be referenced in C++ using either the LOCTABLE macro, or the static FText::FromStringTable function. INI files can reference them using the LOCTABLE macro syntax, and FText properties in assets can reference them via the advanced settings combo.

Change 3257018 on 2017/01/13 by [REDACTED_PERSON]

	FbxAutomationTest fix the import reload operation, it was calling garbagecollect with no keep flag

Change 3257168 on 2017/01/13 by [REDACTED_PERSON]

	Removed code that was writing null into bytecode during save

Change 3257344 on 2017/01/13 by [REDACTED_PERSON]

	Backing out changelist 3256725, and excluding my header from the scan instead

Change 3257426 on 2017/01/13 by [REDACTED_PERSON]

	Slate - Adding the ability to invert alpha when drawing slate textures.  Going to be used in the future for rendering render targets for the scene which have inverted alpha.

Change 3257572 on 2017/01/13 by [REDACTED_PERSON]

	Slate - Fixing a build error.

Change 3257970 on 2017/01/14 by [REDACTED_PERSON]

	Fixing exclude path

Change 3258458 on 2017/01/16 by [REDACTED_PERSON]

	PR #3135: GameViewportClient: FOnCloseRequested is now a multicast delegate (Contributed by [REDACTED_PERSON])

Change 3258472 on 2017/01/16 by [REDACTED_PERSON]

	PR #3126: Fix to load editor style assets (Contributed by [REDACTED_PERSON])

Change 3258473 on 2017/01/16 by [REDACTED_PERSON]

	PR #3124: Fix wrong result with Image-DrawAsBox with PaperSprite. (Contributed by [REDACTED_PERSON])

Change 3258539 on 2017/01/16 by [REDACTED_PERSON]

	Slate - Pixel Snapping has been moved to the GPU for the RHI rendering policy.  Additionally, widgets with a render transform of Scale, Rotation or Sheer, and their children are no longer pixel snapped, this should reduce some of jittering seen by users when animations are applied to widgets.  NOTE: This only affects render transforms, any transform in layout space is still subject to pixel snapping.

Change 3258607 on 2017/01/16 by [REDACTED_PERSON]

	Fixing the mac build.

Change 3258661 on 2017/01/16 by [REDACTED_PERSON]

	Actors with experimental components no longer say
	"Uses experimental class: Actor" when selecting the actor root in the details panel

Change 3258678 on 2017/01/16 by [REDACTED_PERSON]

	Platform - Introducing a way to get the mimetype for a file on Windows.  Other platforms don't yet have an implementation outside of returning application/unknown.

Change 3258924 on 2017/01/16 by [REDACTED_PERSON]

	Platform - Implementing a fallback for the generic platform http, that can do some basic mimetype lookups.

Change 3258929 on 2017/01/16 by [REDACTED_PERSON]

	UMG - Fixing the animation to finish the evaluation before it notifies that the animation completed.

Change 3259109 on 2017/01/16 by [REDACTED_PERSON]

	Platform - The GetMimeType function now only takes in FilePath, since some platforms will require that actually resolve to a file on disk in order to determine the true mimetype.

Change 3259111 on 2017/01/16 by [REDACTED_PERSON]

	Avoid to move the camera when we re-import in the static mesh editor

Change 3259275 on 2017/01/16 by [REDACTED_PERSON]

	Fix crash when a slate window is resized and calls into a scene viewport during loading code when the scene viewport is not in a slate hierarchy and thus has no widget

Change 3259300 on 2017/01/16 by [REDACTED_PERSON]

	UMG - Introducing PreConstruct and NativePreConstruct to the base UUserWidget.  Users can now visualize non-binding based changes in the designer by evaluating a very limited amount of the blueprint code.  In the event your user widget crashes on load, due to calling something unsafe, you can disable evaluation in the editor preferences under Widget Designer.

Change 3259306 on 2017/01/16 by [REDACTED_PERSON]

	Games - Removing the Game Specific implementations of PreConstruct.

Change 3260182 on 2017/01/17 by [REDACTED_PERSON]

	Fix static analysis

Change 3261049 on 2017/01/17 by [REDACTED_PERSON]

	Slate - Putting in some fixes for the non-gpu pixel snapping mode, and disabling gpu snapping while we dig into why it looks weird.

Change 3261434 on 2017/01/17 by [REDACTED_PERSON]

	Fixing the mac build.

Change 3261435 on 2017/01/17 by [REDACTED_PERSON]

	Slate - Tweaking some aspects of the slate rounding code on the GPU.  There's still some precision loss somewhere causing subtle differences in where the snap occurs, that's different from previously.

Change 3261460 on 2017/01/17 by [REDACTED_PERSON]

	UMG - Tweaking the defintiions of NativePreConstruct, dropping passing in design time since that is readily available in native code.

Change 3261833 on 2017/01/18 by [REDACTED_PERSON]

	Fix all warning for fbx automation tests

Change 3261874 on 2017/01/18 by [REDACTED_PERSON]

	PR #3136: Fix Submit to Source Control Window for Git plugin : use CanCheckIn() to filter out unmodified assets files (Contributed by [REDACTED_PERSON])

Change 3262000 on 2017/01/18 by [REDACTED_PERSON]

	Updated Slate to allocate widgets using MakeShared

	This saves one allocation per-widget

Change 3262003 on 2017/01/18 by [REDACTED_PERSON]

	UMG - Widget Interaction Components now ignore Visible(false) Widget Components when tracing.

Change 3262052 on 2017/01/18 by [REDACTED_PERSON]

	Put back the staticmesh skinxx workflow

Change 3262775 on 2017/01/18 by [REDACTED_PERSON]

	Slate - Ditching moving vertex rounding to the GPU, some precision issues could not be overcome.  Ended up writing a clean way to implement it on the CPU.

Change 3262818 on 2017/01/18 by [REDACTED_PERSON]

Change 3263679 on 2017/01/19 by [REDACTED_PERSON]

	Slate - Adding some comments to the Slate Vertex Rounder.

Change 3265154 on 2017/01/19 by [REDACTED_PERSON]

	Slate/UMG - Putting in some more time into pixel snapping.  I've re-introduced the old constructors, and decided to go with the templated approach, as to not break old code that relied on the FSlateVertex working a certain way.

Change 3265478 on 2017/01/20 by [REDACTED_PERSON]

	Added config support for hang detection time and switching hang detection on/off in UnrealWatchdog
	[UE-40838] - Make hang time configurable and increase default in UnrealWatchdog

Change 3265600 on 2017/01/20 by [REDACTED_PERSON]

	Slate - Making some const local variables const.

Change 3265714 on 2017/01/20 by [REDACTED_PERSON]

Change 3265865 on 2017/01/20 by [REDACTED_PERSON]

Change 3267989 on 2017/01/23 by [REDACTED_PERSON]

	Exposed String Tables to Blueprints

Change 3268018 on 2017/01/23 by [REDACTED_PERSON]

	Small API clean-up for string tables

Change 3268455 on 2017/01/23 by [REDACTED_PERSON]

	Fix SaveAs (Which says SaveCurrentAs) not saving the current level and only saving the persistent level and then reloading everything thus causing work to be lost if editing a sub-level

Change 3269388 on 2017/01/24 by [REDACTED_PERSON]

	Refactored tick timing in UnrealWatchdog to stop bug where it doesn't close.
	[UE-40839] - UnrealWatchdog running and blocking use of Unreal Game Sync for internal users

	Standalone tool code only - doesn't touch engine

Change 3270205 on 2017/01/24 by [REDACTED_PERSON]

	Updated FUnrealEdMisc::OnMessageTokenActivated to properly traverse up the outer hierarchy of an object.

Change 3270231 on 2017/01/24 by [REDACTED_PERSON]

	Renamed and exposed GetFullScreenAlignment and GetViewportAnchors for consistency with the setters

Change 3271734 on 2017/01/25 by [REDACTED_PERSON]

	Add sorting for landscape target layer, user can now sort alphabetical, material based or custom
	Added a new vertical box SDragNDropVerticalBox to handle drag & drop of FSlot
	Fixed SDropTarget to only consider the drop action if it was started by it
	Added visibility toggle to only show used layers in the currently loaded data

Change 3271797 on 2017/01/25 by [REDACTED_PERSON]

	Renamed HasBeenAlreadyMadeSharable to DoesSharedInstanceExist as the old name was nonsense

Change 3271813 on 2017/01/25 by [REDACTED_PERSON]

	Fixed bad access of a shared this during widget destruction when a context menu was open

Change 3271988 on 2017/01/25 by [REDACTED_PERSON]

	Slate - Removing some old checkbox deprecated code from the 4.3 and 4.6 days.

Change 3271992 on 2017/01/25 by [REDACTED_PERSON]

	Blueprints - Making the checked call better to log out more information when dragging and dropping a missing property.

Change 3272134 on 2017/01/25 by [REDACTED_PERSON]

	Updated the GatherText commandlet to no longer hold a ConfigFile pointer while it runs

	This pointer is internal to GConfig, and may be updated (or invalidated) when other config files are loaded (as can happen via game code while gathering text).

Change 3272301 on 2017/01/25 by [REDACTED_PERSON]

	Slate - More cleanup from the removal of a old legacy enum that people were still using.

Change 3273070 on 2017/01/26 by [REDACTED_PERSON]

	Fix CIS errors in landscape code from CL 3271734

Change 3273123 on 2017/01/26 by [REDACTED_PERSON]

	Fix crash during init of CRC when running packaged without access to main engine config hierarchy.

Change 3273194 on 2017/01/26 by [REDACTED_PERSON]

	Fixing some build warnings.

Change 3273242 on 2017/01/26 by [REDACTED_PERSON]

Change 3273279 on 2017/01/26 by [REDACTED_PERSON]

	String Table INI redirects are now in the "Core.StringTable" section (rather than "/Script/Engine.Engine")

Change 3273483 on 2017/01/26 by [REDACTED_PERSON]

	Also makes a slight change to the ResavePackages commandlet to submit files marked for delete.

Change 3273568 on 2017/01/26 by [REDACTED_PERSON]

	Modifying changes made to SPluginWizard to have the plugin loading phase determined by the wizard's definition rather than from the first selected template.

Change 3273855 on 2017/01/26 by [REDACTED_PERSON]

Change 3274200 on 2017/01/26 by [REDACTED_PERSON]

	For IPluginWizardDefinition, temporarily adding function bodies to two methods instead of having them be pure virtual methods.

Change 3274317 on 2017/01/26 by [REDACTED_PERSON]

	Deleting a seemingly corrupted asset that was accidentially submitted

Change 3275072 on 2017/01/27 by [REDACTED_PERSON]

	Fix typo error
	Iterate all components, not only active one
	Force expand the Target Layers widget

Change 3275249 on 2017/01/27 by [REDACTED_PERSON]

	Color grading controls: Keep the vector ratio when changing the master slider

Change 3275282 on 2017/01/27 by [REDACTED_PERSON]

	Color grading controls: Cosmetic changes

Change 3275292 on 2017/01/27 by [REDACTED_PERSON]

	Make sure the build is called once when we import a staticmesh.

Change 3275430 on 2017/01/27 by [REDACTED_PERSON]

	Add some fbx automation tests
	- Import a mesh with no material
	- Import corrupted asset with no section in a LOD
	- Import morph targets
	- Materials name clash
	- Max Multimap material ordering

Change 3275683 on 2017/01/27 by [REDACTED_PERSON]

Change 3276237 on 2017/01/27 by [REDACTED_PERSON]

	Deleting a seemingly corrupted asset that was accidentially submitted

Change 3276266 on 2017/01/27 by [REDACTED_PERSON]

	Fix for accessing a potentially null pointer

Change 3277065 on 2017/01/30 by [REDACTED_PERSON]

	Move crash report temp files to saved config and cleanup on schedule.
	[UE-39506] - CrashReportClient ini folders are not cleaned when opening the editor

Change 3277236 on 2017/01/30 by [REDACTED_PERSON]

	Fix crash when cancelling SaveCurrentLevelAs

Change 3277409 on 2017/01/30 by [REDACTED_PERSON]

	Improved text rendering when the last resort font is missing

	The last resort font is no longer included in shipping builds, so this change makes some improvements to text rendering when it's missing.

	- The legacy font cache no longer tries to use the last resort font if it's not available (preventing warnings).
	- The Slate font renderer no longer tries to use the last resort font if it's not available.
	- Text shaping will use the last resort character if none of the available fonts can render a given character (likely because the last resort font is missing).
	- HarfBuzz shaped text now uses the fallback character correctly.

Change 3277749 on 2017/01/30 by [REDACTED_PERSON]

	Slate - Moving ESlateDrawEffect & ESlateBatchDrawFlag over to be enum class, found cases where users were improperly assuming the enum order, and so now it won't be possible to just treat an int32 or a bool as the draw effect value.

	Core - Adding EnumHasAllFlags and EnumHasAnyFlags, templated functions to make it easier to check for the existance of a flag on enum classes.

Change 3277805 on 2017/01/30 by [REDACTED_PERSON]

	Rendering - Changing some LoadModuleChecked calls to GetModuleChecked, as these calls are not happening on the main thread and are not safe to make.

Change 3277914 on 2017/01/30 by [REDACTED_PERSON]

	Fix Niagara slate style warning on startup

Change 3278058 on 2017/01/30 by [REDACTED_PERSON]

	Fixed compile error

Change 3278132 on 2017/01/30 by [REDACTED_PERSON]

	Fixed compile error

Change 3278133 on 2017/01/30 by [REDACTED_PERSON]

	Fixed compile errors

Change 3278186 on 2017/01/30 by [REDACTED_PERSON]

	Fixed compile error

Change 3278525 on 2017/01/30 by [REDACTED_PERSON]

	Fixed compile error

Change 3278534 on 2017/01/30 by [REDACTED_PERSON]

	Automation - Clearing up several warnings/errors with automation results, trying to get Automation Tests to at least yellow before integration.

Change 3278941 on 2017/01/31 by [REDACTED_PERSON]

	Fixing a build warning due to build team refactor.

Change 3278949 on 2017/01/31 by [REDACTED_PERSON]

	Fixing incrmenetal build issues.

Change 3278953 on 2017/01/31 by [REDACTED_PERSON]

	Fixing some incrmental linux build issues.

Change 3278964 on 2017/01/31 by [REDACTED_PERSON]

	FIxing more incremental build issues.

Change 3279256 on 2017/01/31 by [REDACTED_PERSON]

	Instead of getting the Landscape Actor, call GetLandscapeProxy so all case are handled, either proxy or landscape actor

Change 3279270 on 2017/01/31 by [REDACTED_PERSON]

	re-updating the automation test pool

diff --git a/Engine/Documentation/Source/Programming/Slate/Widgets/SlateWidgets.INT.udn b/Engine/Documentation/Source/Programming/Slate/Widgets/SlateWidgets.INT.udn
--- a/Engine/Documentation/Source/Programming/Slate/Widgets/SlateWidgets.INT.udn
+++ b/Engine/Documentation/Source/Programming/Slate/Widgets/SlateWidgets.INT.udn
@@ -175,18 +175,18 @@ Below is an example of both autosized horizontal boxes and fill sized horizontal
 
 	...
 
-	ESlateCheckBoxState::Type IsRadioChecked( ERadioChoice ButtonId ) const
+	ECheckBoxState::Type IsRadioChecked( ERadioChoice ButtonId ) const
 	{
 		return (CurrentChoice == ButtonId)
-			? ESlateCheckBoxState::Checked
-			: ESlateCheckBoxState::Unchecked;
+			? ECheckBoxState::Checked
+			: ECheckBoxState::Unchecked;
 	}
 
 	...
 
-	void OnRadioChanged( ERadioChoice RadioThatChanged, ESlateCheckBoxState::Type NewRadioState )
+	void OnRadioChanged( ERadioChoice RadioThatChanged, ECheckBoxState::Type NewRadioState )
 	{
-		if (NewRadioState == ESlateCheckBoxState::Checked)
+		if (NewRadioState == ECheckBoxState::Checked)
 		{
 			CurrentChoice = RadioThatChanged;
 		}
