Documentation for build tools

diff --git a/Engine/Documentation/Source/Programming/UnrealBuildSystem/TargetFiles/TargetFilesProperties/TargetFilesProperties.INT.udn b/Engine/Documentation/Source/Programming/UnrealBuildSystem/TargetFiles/TargetFilesProperties/TargetFilesProperties.INT.udn
--- a/Engine/Documentation/Source/Programming/UnrealBuildSystem/TargetFiles/TargetFilesProperties/TargetFilesProperties.INT.udn
+++ b/Engine/Documentation/Source/Programming/UnrealBuildSystem/TargetFiles/TargetFilesProperties/TargetFilesProperties.INT.udn
@@ -44,6 +44,8 @@ $ SolutionDirectory (String): Allows a Program Target to specify it's own soluti
 
 $ bOutputToEngineBinaries (Boolean): Output the executable to the engine binaries folder.
 
+$ bShouldCompileAsDLL (Boolean): Whether this target should be compiled as a DLL.  Requires LinkType to be set to TargetLinkType.Monolithic.
+
 $ ExeBinariesSubFolder (String): Subfolder to place executables in, relative to the default location.
 
 $ GeneratedCodeVersion (EGeneratedCodeVersion): Allow target module to override UHT code generation version.
@@ -232,4 +234,6 @@ $ ExtraModuleNames (List&lt;String&gt;): List of additional modules to be compil
 
 $ BuildEnvironment (TargetBuildEnvironment): Specifies the build environment for this target. See TargetBuildEnvironment for more infomation on the available options.
 
+$ OverrideExecutableFileExtension (String): Can be set to override the file extension of the executable file (normally .exe or .dll on Windows, for example)
+
 
