Copying //UE4/Dev-Documentation to //UE4/Samples-Main (Source: //UE4/Dev-Documentation @ 3372893)

==========================
MAJOR FEATURES + CHANGES
==========================

Change 3239909 on 2016/12/19 by [REDACTED_PERSON]

	Fixed caption spacing for hero image.

	#UE4doc

Change 3240118 on 2016/12/19 by [REDACTED_PERSON]

	#UE Docs: UEDOC-3961 - SME review edits and cleaned up project settings image without tooltip included.

Change 3240145 on 2016/12/19 by [REDACTED_PERSON]

	#Docs - adding new work and additional cleanup to Photo Real Character bust, formerly Twinblast Bust documentation

Change 3240452 on 2016/12/19 by [REDACTED_PERSON]

	Updated File against INTSourceChangeList:3229370

Change 3240453 on 2016/12/19 by [REDACTED_PERSON]

	Updated File against INTSourceChangeList:3229370

Change 3240497 on 2016/12/19 by [REDACTED_PERSON]

	Updated File against INTSourceChangeList:3239909

Change 3240602 on 2016/12/20 by [REDACTED_PERSON]

	Updated File against INTSourceChangeList:3235241

Change 3240633 on 2016/12/20 by [REDACTED_PERSON]

	Updated File against INTSourceChangeList:3216146

Change 3240638 on 2016/12/20 by [REDACTED_PERSON]

	Updated File against INTSourceChangelist:3216146

Change 3241214 on 2016/12/20 by [REDACTED_PERSON]

	Updated C+ Interface document.

Change 3241323 on 2016/12/20 by [REDACTED_PERSON]

	Checking in MorphTargetDebugViewMode video for upload to Youtube when changes shelved in CL-3241241 are accepted.

Change 3241425 on 2016/12/20 by [REDACTED_PERSON]

	Adding images per SME suggestions.

	#UE4doc

Change 3241426 on 2016/12/20 by [REDACTED_PERSON]

	Updating images per SME suggestions.

	#UE4doc

Change 3241428 on 2016/12/20 by [REDACTED_PERSON]

	Updated steps to include UE4Editor-Linux-Debug build configuration for CPP Projects in Linux (per SME suggestion).

	#UE4doc

Change 3241590 on 2016/12/21 by [REDACTED_PERSON]

	Updated File against INTSourceChangeList:3216146

Change 3241602 on 2016/12/21 by [REDACTED_PERSON]

	Updated File against INTSourceChangeList:

Change 3241627 on 2016/12/21 by [REDACTED_PERSON]

	Updated File against INTSourceChangeList:3216146

Change 3241640 on 2016/12/21 by [REDACTED_PERSON]

	Updated File against INTSourceChangeList:3216146

Change 3241662 on 2016/12/21 by [REDACTED_PERSON]

	Updated File against INTSourceChangeList:3216146

Change 3242329 on 2016/12/21 by [REDACTED_PERSON]

	#UEDoc - Added meta data to materials pages. Added materials tag

Change 3242400 on 2016/12/21 by [REDACTED_PERSON]

	Finished incorporating all of the SME's suggested additions.

	#UE4doc

Change 3242448 on 2016/12/21 by [REDACTED_PERSON]

	#UE4 Docs: Checking in images for shelved files for FBX Import Options update. UEDOC-3979

Change 3242511 on 2016/12/21 by [REDACTED_PERSON]

	Updated File against INTSourceChangeList:3216146

Change 3242523 on 2016/12/21 by [REDACTED_PERSON]

	Updated File against INTSourceChangeList:3197111

Change 3242534 on 2016/12/21 by [REDACTED_PERSON]

	#loc UE4DocJpn File updated against INT#3238552

Change 3242536 on 2016/12/21 by [REDACTED_PERSON]

	#loc UE4DocJpn File updated against INT#3240118

Change 3242537 on 2016/12/21 by [REDACTED_PERSON]

	Updated File against INTSourceChangeList:3213377

Change 3242541 on 2016/12/21 by [REDACTED_PERSON]

	Updated File against INTSourceChangeList:

Change 3242544 on 2016/12/21 by [REDACTED_PERSON]

	#loc UE4DocJpn File updated against INT#3218993

Change 3242547 on 2016/12/21 by [REDACTED_PERSON]

	#loc UE4DocJpn File updated against INT#3218993

Change 3242550 on 2016/12/21 by [REDACTED_PERSON]

	#loc UE4DocJpn File updated against INT#3218993

Change 3242552 on 2016/12/21 by [REDACTED_PERSON]

	#loc UE4DocJpn File updated against INT#3218993

Change 3242561 on 2016/12/21 by [REDACTED_PERSON]

	#loc UE4DocJpn File updated against INT#3218993

Change 3242563 on 2016/12/21 by [REDACTED_PERSON]

	#loc UE4DocJpn File updated against INT#3218993

Change 3242566 on 2016/12/21 by [REDACTED_PERSON]

	#loc UE4DocJpn File updated against INT#3219004

Change 3242567 on 2016/12/21 by [REDACTED_PERSON]

	#loc UE4DocJpn File updated against INT#3227082

Change 3242569 on 2016/12/21 by [REDACTED_PERSON]

	#loc UE4DocJpn File updated against INT#3224908

Change 3242571 on 2016/12/21 by [REDACTED_PERSON]

	Updated File against INTSourceChangeList:3213377

Change 3242572 on 2016/12/21 by [REDACTED_PERSON]

	#loc UE4DocJpn File updated against INT#3206329

Change 3242575 on 2016/12/21 by [REDACTED_PERSON]

	Updated File against INTSourceChangeList:3213377

Change 3242596 on 2016/12/22 by [REDACTED_PERSON]

	Updated File against INTSourcechangeList:3213377

Change 3242603 on 2016/12/22 by [REDACTED_PERSON]

	Updated File against INTSourceChangeList:3213377

Change 3242620 on 2016/12/22 by [REDACTED_PERSON]

	Updated File against INTSourceChangeList:3197111

Change 3242624 on 2016/12/22 by [REDACTED_PERSON]

	#loc UE4DocJpn File updated against INT#3234766

Change 3242641 on 2016/12/22 by [REDACTED_PERSON]

	Updated File against INTSourceChangeList:3197111

Change 3242645 on 2016/12/22 by [REDACTED_PERSON]

	Updated File against INTSourceChangelist:3232461

Change 3242659 on 2016/12/22 by [REDACTED_PERSON]

	#loc UE4DocJpn File updated against INT#3225484

Change 3242672 on 2016/12/22 by [REDACTED_PERSON]

	#loc UE4DocJpn File updated against INT#3225484

Change 3242691 on 2016/12/22 by [REDACTED_PERSON]

	Updated File against INTSourceChangeList:3181495

Change 3242764 on 2016/12/22 by [REDACTED_PERSON]

	#loc UE4DocKor

Change 3243595 on 2016/12/25 by [REDACTED_PERSON]

	#loc UE4DocJpn File updated against INT#3198551

Change 3243596 on 2016/12/25 by [REDACTED_PERSON]

	#loc UE4DocJpn File updated against INT#3227911

Change 3243598 on 2016/12/25 by [REDACTED_PERSON]

	#loc UE4DocJpn File updated against INT#3236590

Change 3243599 on 2016/12/25 by [REDACTED_PERSON]

	#loc UE4DocJpn File updated against INT#3108692

Change 3243600 on 2016/12/25 by [REDACTED_PERSON]

	#loc UE4DocJpn File updated against INT#3108692

Change 3243603 on 2016/12/25 by [REDACTED_PERSON]

	#loc UE4DocJpn File updated against INT#3108692

Change 3243605 on 2016/12/25 by [REDACTED_PERSON]

	Updated File against INTSourceChangeList:3213377

Change 3243608 on 2016/12/25 by [REDACTED_PERSON]

	Updated File against INTSourceChangeList:3108692

Change 3243610 on 2016/12/25 by [REDACTED_PERSON]

	Updated File against INTSourceChangeList:3175190

Change 3243611 on 2016/12/25 by [REDACTED_PERSON]

	Updated File against INTSourceChangelist:3235148

Change 3243613 on 2016/12/26 by [REDACTED_PERSON]

	Updated File against INTSourceChangeList:3155346

Change 3243620 on 2016/12/26 by [REDACTED_PERSON]

	Updated File against INTSourceChangeList:3155346

Change 3243623 on 2016/12/26 by [REDACTED_PERSON]

	Updated File against INTSourceChangeList:3150518

Change 3243624 on 2016/12/26 by [REDACTED_PERSON]

	Fixed translation (Title, Description)

Change 3243625 on 2016/12/26 by [REDACTED_PERSON]

	Updated File against  INTSourceChangelist:3150518

Change 3243626 on 2016/12/26 by [REDACTED_PERSON]

	Fixed translation (Description)

Change 3243627 on 2016/12/26 by [REDACTED_PERSON]

	Fixed translation of description

Change 3243630 on 2016/12/26 by [REDACTED_PERSON]

	#loc UE4DocJpn File updated against INT#3185922

Change 3243631 on 2016/12/26 by [REDACTED_PERSON]

	Updated File against INTSourceChangeList:3150518

Change 3243721 on 2016/12/26 by [REDACTED_PERSON]

	Updated File against INTSourceChangeList:3150518

Change 3243723 on 2016/12/26 by [REDACTED_PERSON]

	Updated File against INTSourceChangeList:3150518

Change 3243726 on 2016/12/26 by [REDACTED_PERSON]

	Updated File against INTSoruceChangeList:3150518

Change 3243733 on 2016/12/26 by [REDACTED_PERSON]

	Updated File against INTSourceChangeList:3150518

Change 3243860 on 2016/12/27 by [REDACTED_PERSON]

	#UE4 Docs - Updating metadata, added particles tag to Rendering/ParticleSystems pages

Change 3244012 on 2016/12/28 by [REDACTED_PERSON]

	#UE4 Docs - Added metadata to sequencer pages.

Change 3244080 on 2016/12/28 by [REDACTED_PERSON]

	Shortened the page title's character count to be more compatible with mobile screens.

	#UE4doc

Change 3244081 on 2016/12/28 by [REDACTED_PERSON]

	Added reminder to uncomment the 'Setting up an IDE' topic object when the page is marked as 'Publish Ready'.

	#UE4doc

Change 3244089 on 2016/12/28 by [REDACTED_PERSON]

	#UE4 Docs - Added metadata to UMG pages

Change 3244339 on 2016/12/30 by [REDACTED_PERSON]

	#loc UE4DocCHN. do translation

Change 3244370 on 2016/12/30 by [REDACTED_PERSON]

	#UE4 Docs - Added metadata to blueprints pages

Change 3244481 on 2017/01/02 by [REDACTED_PERSON]

	#loc UE4DocKOR

Change 3244498 on 2017/01/02 by [REDACTED_PERSON]

	Repaired 'Getting Started' topic image.

	#UE4doc

Change 3244553 on 2017/01/03 by [REDACTED_PERSON]

	#loc UE4DocKOR

Change 3244554 on 2017/01/03 by [REDACTED_PERSON]

	#loc UE4DocKor

Change 3244596 on 2017/01/03 by [REDACTED_PERSON]

	typo correction

Change 3244600 on 2017/01/03 by [REDACTED_PERSON]

	#loc UE4DocKOR

Change 3244648 on 2017/01/03 by [REDACTED_PERSON]

	#loc UE4DocKOR

Change 3244824 on 2017/01/03 by [REDACTED_PERSON]

	Implemented recommeded additional fixes.

Change 3244832 on 2017/01/03 by [REDACTED_PERSON]

	#UE4 Docs: Applying peer feedabck to the doc.

Change 3244908 on 2017/01/03 by [REDACTED_PERSON]

	#UE4 Docs: Applying peer feedback and adding new images to go along with it.

Change 3244910 on 2017/01/03 by [REDACTED_PERSON]

	Adding this missing file.

Change 3245914 on 2017/01/04 by [REDACTED_PERSON]

	#loc UE4DocJpn File updated against INT#3240145

Change 3245917 on 2017/01/04 by [REDACTED_PERSON]

	#loc UE4DocJpn File updated against INT#3244089

Change 3245921 on 2017/01/04 by [REDACTED_PERSON]

	#loc UE4DocJpn File updated against INT#3244089

Change 3245923 on 2017/01/04 by [REDACTED_PERSON]

	#loc UE4DocJpn File updated against INT#3244089

Change 3245927 on 2017/01/04 by [REDACTED_PERSON]

	#loc UE4DocJpn File updated against INT#3244089

Change 3245930 on 2017/01/04 by [REDACTED_PERSON]

	#loc UE4DocJpn File updated against INT#3244089

Change 3245931 on 2017/01/04 by [REDACTED_PERSON]

	#loc UE4DocJpn File updated against INT#3244089

Change 3245932 on 2017/01/04 by [REDACTED_PERSON]

	#loc UE4DocJpn File updated against INT#3244089

Change 3245934 on 2017/01/04 by [REDACTED_PERSON]

	#loc UE4DocJpn File updated against INT#3244089

Change 3245951 on 2017/01/04 by [REDACTED_PERSON]

	#loc UE4DocJpn File updated against INT#3244089

Change 3245967 on 2017/01/04 by [REDACTED_PERSON]

	Updated File against INTSourceChangeList:3108692

Change 3246037 on 2017/01/04 by [REDACTED_PERSON]

	#loc UE4DocKOR

Change 3246299 on 2017/01/04 by [REDACTED_PERSON]

	#UE4 Docs: Applying SME feedback.

Change 3247009 on 2017/01/04 by [REDACTED_PERSON]

	Updating VR profiling documents based on feedback

Change 3247223 on 2017/01/04 by [REDACTED_PERSON]

	#loc UE4DocJpn File updated against INT#3244089

Change 3247334 on 2017/01/04 by [REDACTED_PERSON]

	#loc UE4DocJpn File updated against INT#3244089

Change 3247339 on 2017/01/04 by [REDACTED_PERSON]

	Updated File against INTSourceChangeList:3108692

Change 3247340 on 2017/01/04 by [REDACTED_PERSON]

	#loc UE4DocJpn File updated against INT#3244089

Change 3247343 on 2017/01/04 by [REDACTED_PERSON]

	#loc UE4DocJpn File updated against INT#3244089

Change 3247347 on 2017/01/04 by [REDACTED_PERSON]

	#loc UE4DocJpn File updated against INT#3244089

Change 3247356 on 2017/01/04 by [REDACTED_PERSON]

	#loc UE4DocJpn File updated against INT#3244089

Change 3247362 on 2017/01/04 by [REDACTED_PERSON]

	#loc UE4DocJpn File updated against INT#3244089

Change 3247363 on 2017/01/04 by [REDACTED_PERSON]

	#loc UE4DocJpn File updated against INT#3244089

Change 3247411 on 2017/01/05 by [REDACTED_PERSON]

	Updated File against INTSourceChangeList:3108692

Change 3247417 on 2017/01/05 by [REDACTED_PERSON]

	Updated File against INTSourceChangeList:3108692

Change 3247428 on 2017/01/05 by [REDACTED_PERSON]

	Updated File against INTSourceChangeList:3108692

Change 3247440 on 2017/01/05 by [REDACTED_PERSON]

	#loc UE4DocJpn File updated against INT#3244089

Change 3247446 on 2017/01/05 by [REDACTED_PERSON]

	#loc UE4DocJpn File updated against INT#3244089

Change 3247448 on 2017/01/05 by [REDACTED_PERSON]

	#loc UE4DocJpn File updated against INT#3244089

Change 3247450 on 2017/01/05 by [REDACTED_PERSON]

	Updated File against INTSourceChangeList:3108692

Change 3247452 on 2017/01/05 by [REDACTED_PERSON]

	#loc UE4DocJpn File updated against INT#3244089

Change 3247457 on 2017/01/05 by [REDACTED_PERSON]

	#loc UE4DocJpn File updated against INT#3244089

Change 3247463 on 2017/01/05 by [REDACTED_PERSON]

	#loc UE4DocJpn File updated against INT#3244089

Change 3247466 on 2017/01/05 by [REDACTED_PERSON]

	#loc UE4DocJpn File updated against INT#3244089

Change 3247467 on 2017/01/05 by [REDACTED_PERSON]

	Updated File against INTSourceChangeList:3108692

Change 3247483 on 2017/01/05 by [REDACTED_PERSON]

	Updated File against INTSourceChangeList:3108692

Change 3247497 on 2017/01/05 by [REDACTED_PERSON]

	#loc UE4DocKOR

Change 3248719 on 2017/01/05 by [REDACTED_PERSON]

	#loc UE4DocJpn File updated against INT#3244089

Change 3248726 on 2017/01/05 by [REDACTED_PERSON]

	#loc UE4DocJpn File updated against INT#3244089

Change 3248765 on 2017/01/05 by [REDACTED_PERSON]

	Updated File against INTSourceChangeList:3108692

Change 3248787 on 2017/01/05 by [REDACTED_PERSON]

	Updated File against INTSourceChangeList:3213360

Change 3248818 on 2017/01/05 by [REDACTED_PERSON]

	Updated File against INTSourceChangeList:3213360

Change 3248819 on 2017/01/05 by [REDACTED_PERSON]

	Updated File against INTSourceChangeList:3213360

Change 3248870 on 2017/01/06 by [REDACTED_PERSON]

	Updated File against INTSourceChangeList:3213360

Change 3248912 on 2017/01/06 by [REDACTED_PERSON]

	#loc UE4DocJpn Initial translation against INT#3225484

Change 3248928 on 2017/01/06 by [REDACTED_PERSON]

	#loc UE4DocJpn File updated against INT#3244089

Change 3248941 on 2017/01/06 by [REDACTED_PERSON]

	#loc UE4DocJpn File updated against INT#3244089

Change 3248951 on 2017/01/06 by [REDACTED_PERSON]

	#loc UE4DocJpn File updated against INT#3244089

Change 3248961 on 2017/01/06 by [REDACTED_PERSON]

	#loc UE4DocJpn File updated against INT#3244089

Change 3248969 on 2017/01/06 by [REDACTED_PERSON]

	#loc UE4DocJpn File updated against INT#3244089

Change 3249082 on 2017/01/06 by [REDACTED_PERSON]

	#loc UE4DocKOR

Change 3249658 on 2017/01/06 by [REDACTED_PERSON]

	#ue4 docs

	Updated Sequencer EDL How-to / Overivew Page to include a note about adding frame handles as part of an EDL export.

Change 3249741 on 2017/01/06 by [REDACTED_PERSON]

	#UE4 Docs: Adding a new image showing off landscape in VR.

Change 3251914 on 2017/01/10 by [REDACTED_PERSON]

	Updated File against INTSourceChangeList:3178213

Change 3252541 on 2017/01/10 by [REDACTED_PERSON]

	#ue4 docs
	Applying Peer Feeback

Change 3252710 on 2017/01/10 by [REDACTED_PERSON]

	#ue4 docs #jira UEDOC-4351 Small formatting and gramatical fixes to SequencerOverview and ImportExportEDLs pages

Change 3253176 on 2017/01/10 by [REDACTED_PERSON]

	After discussing CodeWorks version updates with SamD, we agreed that the note region needed to be updated to account for future releases of CodeWorks for Android.
	Subsequently, the landing page for the Android QS series has been updated to account for future releases of CodeWorks (thereby complying with the *.tps callout that's packaged with the engine).
	#UE4Doc

Change 3253196 on 2017/01/10 by [REDACTED_PERSON]

	Removed bolding from hyperlink to comply with Style Guide.
	#UE4Doc

Change 3253207 on 2017/01/10 by [REDACTED_PERSON]

	Updated document title to more accurately label the page's contents.
	Added tags and type metadata per Visual Style Guide.
	#UE4Doc

Change 3253275 on 2017/01/10 by [REDACTED_PERSON]

	#loc UE4DocJpn File updated against INT#3244089

Change 3253282 on 2017/01/10 by [REDACTED_PERSON]

	#loc UE4DocJpn File updated against INT#3244089

Change 3253287 on 2017/01/10 by [REDACTED_PERSON]

	#loc UE4DocJpn File updated against INT#3244089

Change 3253312 on 2017/01/10 by [REDACTED_PERSON]

	#loc UE4DocJpn File updated against INT#3244089

Change 3253328 on 2017/01/10 by [REDACTED_PERSON]

	#loc UE4DocJpn File updated against INT#3244089

Change 3253414 on 2017/01/11 by [REDACTED_PERSON]

	#loc UE4DocKor

Change 3253430 on 2017/01/11 by [REDACTED_PERSON]

	#loc UE4DocJpn File updated against INT#3244089

Change 3253442 on 2017/01/11 by [REDACTED_PERSON]

	#loc UE4DocJpn File updated against INT#3244089

Change 3253458 on 2017/01/11 by [REDACTED_PERSON]

	#loc UE4DocJpn File updated against INT#3244089

Change 3253475 on 2017/01/11 by [REDACTED_PERSON]

	#loc UE4DocJpn File updated against INT#3244089

Change 3253478 on 2017/01/11 by [REDACTED_PERSON]

	#loc UE4DocJpn File updated against INT#3244089

Change 3253483 on 2017/01/11 by [REDACTED_PERSON]

	#loc UE4DocJpn File updated against INT#3244089

Change 3253522 on 2017/01/11 by [REDACTED_PERSON]

	#loc UE4DocJpn File updated against INT#3244089

Change 3253529 on 2017/01/11 by [REDACTED_PERSON]

	#loc UE4DocKOR

Change 3253663 on 2017/01/11 by [REDACTED_PERSON]

	#UE4 Docs

	Updated Camera Rig Crane How-to page to include new Lock Mount options
	- Refactored existing How-to page
	- Updated content/images

Change 3254376 on 2017/01/11 by [REDACTED_PERSON]

	#UE4 Docs: Applying Peer Feedback.

Change 3254776 on 2017/01/11 by [REDACTED_PERSON]

	#loc UE4DocJpn File updated against INT#3244089

Change 3254795 on 2017/01/11 by [REDACTED_PERSON]

	#loc UE4DocJpn File updated against INT#3244089

Change 3254797 on 2017/01/11 by [REDACTED_PERSON]

	#UE4Doc
	Performed paradigmatic analysis and substituded misused word (allow) with proper word (enable).

Change 3254802 on 2017/01/11 by [REDACTED_PERSON]

	Updated File against INTSourceChangeList:3108692

Change 3254861 on 2017/01/11 by [REDACTED_PERSON]

	#loc UE4DocJpn File updated against INT#3244089

Change 3254877 on 2017/01/11 by [REDACTED_PERSON]

	#loc UE4DocJpn File updated against INT#3244089

Change 3254882 on 2017/01/11 by [REDACTED_PERSON]

	#loc UE4DocJpn File updated against INT#3244012

Change 3254905 on 2017/01/11 by [REDACTED_PERSON]

	Updated File against INTSourceChangeList:3108692

Change 3254927 on 2017/01/11 by [REDACTED_PERSON]

	#loc UE4DocKOR

Change 3254957 on 2017/01/12 by [REDACTED_PERSON]

	Updated File against INTSourceChangeList:3244080

Change 3254958 on 2017/01/12 by [REDACTED_PERSON]

	Updated File against INTSourceChangeList:3244081

Change 3254976 on 2017/01/12 by [REDACTED_PERSON]

	Updated File against INTSourceChangeList:3243860

Change 3254995 on 2017/01/12 by [REDACTED_PERSON]

	#loc UE4DocJpn File updated against INT#3253663

Change 3255016 on 2017/01/12 by [REDACTED_PERSON]

	#loc UE4DocJpn File updated against INT#3244012

Change 3255028 on 2017/01/12 by [REDACTED_PERSON]

	#loc UE4DocJpn File updated against INT#3244012

Change 3255040 on 2017/01/12 by [REDACTED_PERSON]

	#loc UE4DocJpn File updated against INT#3244012

Change 3255053 on 2017/01/12 by [REDACTED_PERSON]

	#loc UE4DocJpn File updated against INT#3244012

Change 3255054 on 2017/01/12 by [REDACTED_PERSON]

	#loc UE4DocJpn File updated against INT#3244012

Change 3255346 on 2017/01/12 by [REDACTED_PERSON]

	#ue4 docs

	Updated Sequencer Audio Recording How-to to include an optional step of recording audio.
	Updated Sequencer Overview page to include a line about audio recording.

Change 3255615 on 2017/01/12 by [REDACTED_PERSON]

	#UE4 Docs: Adding a new part for the SteamVR perf tools

Change 3255730 on 2017/01/12 by [REDACTED_PERSON]

	Removed topic variables and added topic-image metadata

Change 3255731 on 2017/01/12 by [REDACTED_PERSON]

	Fixed link to non-existent bookmark

Change 3255735 on 2017/01/12 by [REDACTED_PERSON]

	Updated metadata

Change 3255737 on 2017/01/12 by [REDACTED_PERSON]

	Updated DB file

Change 3255740 on 2017/01/12 by [REDACTED_PERSON]

	Updated learning track template to use directory list

Change 3255748 on 2017/01/12 by [REDACTED_PERSON]

	Physics Property Reference Formatting Updates

Change 3255757 on 2017/01/12 by [REDACTED_PERSON]

	Prototype for new sidebar

Change 3255761 on 2017/01/12 by [REDACTED_PERSON]

	#ue4 docs

	Updated page to include some peer review feedback.

Change 3255882 on 2017/01/12 by [REDACTED_PERSON]

	#ue4 docs
	Minor upates to Pose Drive node section on Pose Nodes page.
	- Added Translation option to Pose Node reference table.
	- Updated Overview Text.

Change 3255885 on 2017/01/12 by [REDACTED_PERSON]

	Updated metadata

Change 3256541 on 2017/01/13 by [REDACTED_PERSON]

	Updated File against INTSourceChangeList:3243860

Change 3256576 on 2017/01/13 by [REDACTED_PERSON]

	Updated File against INTSourceChangeList:3243860

Change 3256594 on 2017/01/13 by [REDACTED_PERSON]

	Updated File against INTSourceChangeList:3243860

Change 3256621 on 2017/01/13 by [REDACTED_PERSON]

	#loc UE4DocJpn File updated against INT#3252710

Change 3256625 on 2017/01/13 by [REDACTED_PERSON]

	#loc UE4DocJpn File updated against INT#3244012

Change 3256626 on 2017/01/13 by [REDACTED_PERSON]

	Updated File against INTSourceChangeList:3243860

Change 3256628 on 2017/01/13 by [REDACTED_PERSON]

	#loc UE4DocJpn File updated against INT#3244012

Change 3256865 on 2017/01/13 by [REDACTED_PERSON]

	#ue4 docs

	Applying peer review feedback.

Change 3256942 on 2017/01/13 by [REDACTED_PERSON]

	Added lloc warning text for outdate translations

Change 3256943 on 2017/01/13 by [REDACTED_PERSON]

	Updated metadata

Change 3256979 on 2017/01/13 by [REDACTED_PERSON]

	Unpublish empty pages

Change 3257013 on 2017/01/13 by [REDACTED_PERSON]

	Fixed typo

	PR #2816

Change 3258199 on 2017/01/15 by [REDACTED_PERSON]

	#loc UE4DocJpn File updated against INT#3255885

Change 3258213 on 2017/01/15 by [REDACTED_PERSON]

	#loc UE4DocJpn File updated against INT#3255731

Change 3258244 on 2017/01/16 by [REDACTED_PERSON]

	Updated File against INTSourceChangeList:3243860

Change 3258260 on 2017/01/16 by [REDACTED_PERSON]

	#loc UE4DocJpn File updated against INT#3255885

Change 3258263 on 2017/01/16 by [REDACTED_PERSON]

	#loc UE4DocJpn File updated against INT#3255885

Change 3258264 on 2017/01/16 by [REDACTED_PERSON]

	#loc UE4DocJpn File updated against INT#3255885

Change 3258268 on 2017/01/16 by [REDACTED_PERSON]

	#loc UE4DocJpn File updated against INT#3255885

Change 3258272 on 2017/01/16 by [REDACTED_PERSON]

	#loc UE4DocJpn File updated against INT#3255885

Change 3258274 on 2017/01/16 by [REDACTED_PERSON]

	#loc UE4DocJpn File updated against INT#3255885

Change 3258276 on 2017/01/16 by [REDACTED_PERSON]

	#loc UE4DocJpn File updated against INT#3255885

Change 3258282 on 2017/01/16 by [REDACTED_PERSON]

	#loc UE4DocJpn File updated against INT#3255885

Change 3258283 on 2017/01/16 by [REDACTED_PERSON]

	#loc UE4DocJpn File updated against INT#3255885

Change 3258287 on 2017/01/16 by [REDACTED_PERSON]

	#loc UE4DocJpn File updated against INT#3255885

Change 3258290 on 2017/01/16 by [REDACTED_PERSON]

	#loc UE4DocJpn File updated against INT#3255885

Change 3258292 on 2017/01/16 by [REDACTED_PERSON]

	#loc UE4DocJpn File updated against INT#3255885

Change 3258294 on 2017/01/16 by [REDACTED_PERSON]

	#loc UE4DocJpn File updated against INT#3255885

Change 3258311 on 2017/01/16 by [REDACTED_PERSON]

	#loc UE4DocKor

Change 3258319 on 2017/01/16 by [REDACTED_PERSON]

	#loc UE4DocCHN. do translation

Change 3258323 on 2017/01/16 by [REDACTED_PERSON]

	#loc UE4DocJpn File updated against INT#3255882

Change 3258327 on 2017/01/16 by [REDACTED_PERSON]

	#loc UE4DocJpn File updated against INT#:3244012

Change 3258329 on 2017/01/16 by [REDACTED_PERSON]

	#loc UE4DocJpn File updated against INT#3244012

Change 3258331 on 2017/01/16 by [REDACTED_PERSON]

	#loc UE4DocJpn File updated against INT#3244012

Change 3258333 on 2017/01/16 by [REDACTED_PERSON]

	#loc UE4DocKor

Change 3258335 on 2017/01/16 by [REDACTED_PERSON]

	Fixed a layout error.

Change 3258389 on 2017/01/16 by [REDACTED_PERSON]

	#loc UE4DocKor ?í???? -> Enum

Change 3258677 on 2017/01/16 by [REDACTED_PERSON]

	#UE4 Docs: Adding an image of the foliage tools.

Change 3259022 on 2017/01/16 by [REDACTED_PERSON]

	Set page to publish

Change 3259023 on 2017/01/16 by [REDACTED_PERSON]

	Set page to publish

Change 3259024 on 2017/01/16 by [REDACTED_PERSON]

	Updated formatting to use new localized text feature

Change 3259094 on 2017/01/16 by [REDACTED_PERSON]

	submitting peer review edits UEDOC-4283

Change 3259182 on 2017/01/16 by [REDACTED_PERSON]

	#UE4 Docs: Adding this newly created doc.

Change 3259604 on 2017/01/16 by [REDACTED_PERSON]

	#loc UE4DocJPN moved or deleted

Change 3259614 on 2017/01/16 by [REDACTED_PERSON]

	#loc UE4DocJpn File updated against INT#3244012

Change 3259618 on 2017/01/16 by [REDACTED_PERSON]

	#loc UE4DocJpn File updated against INT#3244012

Change 3259626 on 2017/01/16 by [REDACTED_PERSON]

	#loc UE4DocJpn File updated against INT#3244012

Change 3259628 on 2017/01/16 by [REDACTED_PERSON]

	#loc UE4DocJpn File updated against INT#3244012

Change 3259647 on 2017/01/16 by [REDACTED_PERSON]

	#loc UE4DocJpn File updated against INT#3244012

Change 3259653 on 2017/01/16 by [REDACTED_PERSON]

	#loc UE4DocJpn File updated against INT#3244012

Change 3259662 on 2017/01/16 by [REDACTED_PERSON]

	Updated File against INTSourceChangeList:3243860

Change 3259676 on 2017/01/16 by [REDACTED_PERSON]

	#loc UE4DocJpn File updated against INT#3244012

Change 3259679 on 2017/01/16 by [REDACTED_PERSON]

	Updated File against INTSourceChangeList:3243860

Change 3259694 on 2017/01/16 by [REDACTED_PERSON]

	Updated File against INTSourceChangeList:3211040

Change 3259720 on 2017/01/16 by [REDACTED_PERSON]

	Updated File against INTSourceChangeList:3211040

Change 3259732 on 2017/01/16 by [REDACTED_PERSON]

	#loc UE4DocJpn File updated against INT#3256865

Change 3259745 on 2017/01/16 by [REDACTED_PERSON]

	Updated File against INTSourceChangeList:3211040

Change 3259756 on 2017/01/16 by [REDACTED_PERSON]

	Updated File against INTSourceChangeList:3211040

Change 3259772 on 2017/01/16 by [REDACTED_PERSON]

	Updated File against INTSourceChangeList:3211040

Change 3259800 on 2017/01/16 by [REDACTED_PERSON]

	#loc UE4DocKOR moved or deleted

Change 3259855 on 2017/01/17 by [REDACTED_PERSON]

	#loc UE4DocJpn File updated against INT#3256865

Change 3259868 on 2017/01/17 by [REDACTED_PERSON]

	#loc UE4DocJpn File updated against INT#3244012

Change 3259874 on 2017/01/17 by [REDACTED_PERSON]

	#loc UE4DocJpn File updated against INT#3244012

Change 3259878 on 2017/01/17 by [REDACTED_PERSON]

	#loc UE4DocJpn File updated against INT#3244012

Change 3259882 on 2017/01/17 by [REDACTED_PERSON]

	#loc UE4DocJpn File updated against INT#3244012

Change 3259895 on 2017/01/17 by [REDACTED_PERSON]

	#loc UE4DocJpn File updated against INT#3244012

Change 3259913 on 2017/01/17 by [REDACTED_PERSON]

	#loc UE4DocJpn File updated against INT#3244012

Change 3259923 on 2017/01/17 by [REDACTED_PERSON]

	#loc UE4DocJpn File updated against INT#3244012

Change 3259925 on 2017/01/17 by [REDACTED_PERSON]

	#loc UE4DocJpn File updated against INT#3244012

Change 3259955 on 2017/01/17 by [REDACTED_PERSON]

	#loc UE4DocJpn File updated against INT#3255761

Change 3259987 on 2017/01/17 by [REDACTED_PERSON]

	#loc UE4DocJpn File updated against INT#3259023

Change 3259990 on 2017/01/17 by [REDACTED_PERSON]

	#loc UE4DocJpn File updated against INT#3259023

Change 3259992 on 2017/01/17 by [REDACTED_PERSON]

	#loc UE4DocJpn File updated against INT#3259023

Change 3259998 on 2017/01/17 by [REDACTED_PERSON]

	#loc UE4DocJpn File updated against INT#3255730

Change 3260021 on 2017/01/17 by [REDACTED_PERSON]

	#loc UE4DocJpn File updated against INT#
	3259182

Change 3260551 on 2017/01/17 by [REDACTED_PERSON]

	Added note about painting Foliage in VR with link to VR Editor Controls page

Change 3260598 on 2017/01/17 by [REDACTED_PERSON]

	Set page to publish

Change 3260620 on 2017/01/17 by [REDACTED_PERSON]

	Updated metadata

Change 3260622 on 2017/01/17 by [REDACTED_PERSON]

	Accepting editorial changes.

Change 3260646 on 2017/01/17 by [REDACTED_PERSON]

	#ue4 docs

	initial page updates for Post Process Animation Blueprints
	- Added to Animation Blueprint overview page.
	- Added to Skeletal Mesh / Mesh Details page where you set the Post Process Anim BP

Change 3260662 on 2017/01/17 by [REDACTED_PERSON]

	Set page to publish

Change 3260663 on 2017/01/17 by [REDACTED_PERSON]

	Set page to publish

Change 3260683 on 2017/01/17 by [REDACTED_PERSON]

	Set page to publish

Change 3260884 on 2017/01/17 by [REDACTED_PERSON]

	#UE4 Docs: Adding that Landscape works in VR to the pages.

Change 3260949 on 2017/01/17 by [REDACTED_PERSON]

	#UE4 Docs: Adding header images for each of the topics.

Change 3261032 on 2017/01/17 by [REDACTED_PERSON]

	#UE4 Docs: UEDOC-4283 Added MorphTarget Debug View Mode to existing MorphTarget Previewer page.

Change 3261039 on 2017/01/17 by [REDACTED_PERSON]

	#UE4 Docs: UEDOC-3979 SME Review edits

Change 3261132 on 2017/01/17 by [REDACTED_PERSON]

	Addressed Editorial Review comments, specifically:
	Set topic image using the topic-image metadata in the Contact Shadows page.
	Added a link to the Contact Shadows page in the Lighting and Shadows Landing page.
	Commented out the Contact Shadows page link because it was generating a rogue Doc Page Error.

	#UE4doc

Change 3261385 on 2017/01/17 by [REDACTED_PERSON]

	Updated File against INTSourceChangeList:3108692

Change 3261395 on 2017/01/17 by [REDACTED_PERSON]

	Updated File against INTSourceChangeList:3178213

Change 3261402 on 2017/01/17 by [REDACTED_PERSON]

	Updated File against INTSourceChangeList:3108692

Change 3261425 on 2017/01/17 by [REDACTED_PERSON]

	#loc UE4DocJpn File updated against INT#3259182

Change 3261437 on 2017/01/17 by [REDACTED_PERSON]

	0118

Change 3261456 on 2017/01/17 by [REDACTED_PERSON]

	#loc UE4DocKor ?┐???? -> ?┐

Change 3261520 on 2017/01/18 by [REDACTED_PERSON]

	#loc UE4DocJpn File updated against INT#3261039

Change 3261524 on 2017/01/18 by [REDACTED_PERSON]

	#loc UE4DocJpn File updated against INT#3260622

Change 3261526 on 2017/01/18 by [REDACTED_PERSON]

	#loc UE4DocJpn File updated against INT#3260683

Change 3261537 on 2017/01/18 by [REDACTED_PERSON]

	Updated File against INTSourceChangeList:3108692

Change 3261571 on 2017/01/18 by [REDACTED_PERSON]

	#loc UE4DocJpn File updated against INT#3260683

Change 3261592 on 2017/01/18 by [REDACTED_PERSON]

	Updated File against INTSourceChangeList:3108692

Change 3261886 on 2017/01/18 by [REDACTED_PERSON]

	Added ChildActorTemplate information to Utility classes.

Change 3261940 on 2017/01/18 by [REDACTED_PERSON]

	Added minor detail.

Change 3261961 on 2017/01/18 by [REDACTED_PERSON]

	Minor formatting and text udpates

Change 3262190 on 2017/01/18 by [REDACTED_PERSON]

	Updated metadata

Change 3262197 on 2017/01/18 by [REDACTED_PERSON]

	Updated BD file

Change 3262201 on 2017/01/18 by [REDACTED_PERSON]

	Set page to publish

Change 3262203 on 2017/01/18 by [REDACTED_PERSON]

	Updated to use localized descriptions

Change 3262227 on 2017/01/18 by [REDACTED_PERSON]

	Set page to publish

Change 3262538 on 2017/01/18 by [REDACTED_PERSON]

	#UE4 Docs: Adding images for UEDOC-4475 for addition of the Mirror Tool to the landscape UI and Sculpt Mode pages.

Change 3262605 on 2017/01/18 by [REDACTED_PERSON]

	#UE4 Docs: Rename/move files for Pixel Normal Offset doc.

Change 3263133 on 2017/01/18 by [REDACTED_PERSON]

	Updated File against INTSourceChangeList:3108692

Change 3263158 on 2017/01/18 by [REDACTED_PERSON]

	Updated File against INTSourceChangeList:3262227

Change 3263171 on 2017/01/18 by [REDACTED_PERSON]

	#loc UE4DocJpn File updated against INT#3260646

Change 3263212 on 2017/01/18 by [REDACTED_PERSON]

	Updated File against INTSourceChangeList:3262227

Change 3263244 on 2017/01/18 by [REDACTED_PERSON]

	Updated File against INTSourceChangeList:3262227

Change 3263261 on 2017/01/18 by [REDACTED_PERSON]

	Updated File against INTSourceChangeList:3262227

Change 3263264 on 2017/01/18 by [REDACTED_PERSON]

	#loc UE4DocJpn File updated against INT#3260646

Change 3263269 on 2017/01/18 by [REDACTED_PERSON]

	#loc UE4DocJpn File updated against INT#3255885

Change 3263276 on 2017/01/18 by [REDACTED_PERSON]

	Updated File against INTSourceChangeList:3262227

Change 3263279 on 2017/01/18 by [REDACTED_PERSON]

	#loc UE4DocJpn File updated against INT#3262201

Change 3263283 on 2017/01/18 by [REDACTED_PERSON]

	Updated File against INTSourceChangeList:3262227

Change 3263289 on 2017/01/18 by [REDACTED_PERSON]

	Updated File against INTSourceChangeList:3262227

Change 3263296 on 2017/01/18 by [REDACTED_PERSON]

	Updated File against INTSourceChangeList:3262227

Change 3263299 on 2017/01/18 by [REDACTED_PERSON]

	#loc UE4DocJpn File updated against INT#3261940

Change 3263311 on 2017/01/18 by [REDACTED_PERSON]

	Updated File against INTSourceChangeList:3262190

Change 3263313 on 2017/01/18 by [REDACTED_PERSON]

	Updated File against INTSourceChangelist:3262190

Change 3263321 on 2017/01/18 by [REDACTED_PERSON]

	Updated File against INTSourceChangelist:3262190

Change 3263324 on 2017/01/18 by [REDACTED_PERSON]

	Updated File against INTSourceChangelist:3261961

Change 3263325 on 2017/01/18 by [REDACTED_PERSON]

	Updated File against INTSourceChangelist:3261961

Change 3263370 on 2017/01/19 by [REDACTED_PERSON]

	Updated File against INTSourceChangelist:3261961

Change 3263373 on 2017/01/19 by [REDACTED_PERSON]

	Updated File against INTSourceChangelist:3261961

Change 3263384 on 2017/01/19 by [REDACTED_PERSON]

	#loc UE4DocJpn File updated against INT#3261032

Change 3263399 on 2017/01/19 by [REDACTED_PERSON]

	#loc UE4DocJpn File updated against INT#3254797

Change 3263404 on 2017/01/19 by [REDACTED_PERSON]

	Updated File against INTSourceChangelist:3261961

Change 3263414 on 2017/01/19 by [REDACTED_PERSON]

	Updated File against INTSourceChangelist:3261961

Change 3263421 on 2017/01/19 by [REDACTED_PERSON]

	Updated File against INTSourceChangelist:3261961

Change 3263431 on 2017/01/19 by [REDACTED_PERSON]

	Updated File against INTSourceChangeList:3261961

Change 3263474 on 2017/01/19 by [REDACTED_PERSON]

	Updated File against INTSourceChangelist:3261961

Change 3263969 on 2017/01/19 by [REDACTED_PERSON]

	Updated metadata

Change 3264164 on 2017/01/19 by [REDACTED_PERSON]

	Updated metadata

Change 3264165 on 2017/01/19 by [REDACTED_PERSON]

	Updated metadata

Change 3264406 on 2017/01/19 by [REDACTED_PERSON]

	#UE4 Docs: Applyng feedback from jeff.

Change 3264524 on 2017/01/19 by [REDACTED_PERSON]

	#UE4 Docs: Fixing the formating of the topic image.

Change 3264527 on 2017/01/19 by [REDACTED_PERSON]

	#UE4 Docs:Fixing the formating of the topic image meta tag.

Change 3265224 on 2017/01/19 by [REDACTED_PERSON]

	Updated File against INTSourceChangeList:3260884

Change 3265256 on 2017/01/19 by [REDACTED_PERSON]

	Updatd File against  INTSourceChangelist:3238275

Change 3265260 on 2017/01/19 by [REDACTED_PERSON]

	Updated File against INTSourceChangelist:3235148

Change 3265264 on 2017/01/19 by [REDACTED_PERSON]

	Updated File against INTSourceChangelist:3235148

Change 3265283 on 2017/01/19 by [REDACTED_PERSON]

	Updated File against INTSourceChangelist:3235148

Change 3265295 on 2017/01/19 by [REDACTED_PERSON]

	Updated File against INTSourceChangelist:3235148

Change 3265341 on 2017/01/20 by [REDACTED_PERSON]

	Updated File against INTSourceChangeList:3227073

Change 3265350 on 2017/01/20 by [REDACTED_PERSON]

	Updated File against INTSourceChangelist:3223841

Change 3265360 on 2017/01/20 by [REDACTED_PERSON]

	Updated File against INTSourceChangelist:3223808

Change 3265396 on 2017/01/20 by [REDACTED_PERSON]

	Updated File against INTSourceChangelist:3188564

Change 3266245 on 2017/01/20 by [REDACTED_PERSON]

	#ue4 docs

	Update to Audio Track How-to
	- Content refactored to include volume/pitch adjustments.

Change 3267497 on 2017/01/22 by [REDACTED_PERSON]

	Updated File against INTSourceChangelist:3184883

Change 3267514 on 2017/01/22 by [REDACTED_PERSON]

	Updated File against INTSourceChangelist:3236537

Change 3267536 on 2017/01/22 by [REDACTED_PERSON]

	Updated File against INTSourceChangelist:3232956

Change 3267538 on 2017/01/22 by [REDACTED_PERSON]

	Updated File against INTSourceChangelist:3200686

Change 3267540 on 2017/01/22 by [REDACTED_PERSON]

	#loc UE4DocCHN moved or deleted

Change 3267550 on 2017/01/22 by [REDACTED_PERSON]

	#loc UE4DocCHN. update root page against latest INT version. update template for CHN page out of date warning message.

Change 3267569 on 2017/01/23 by [REDACTED_PERSON]

	#loc UE4DocJpn Initial translation against INT#3244824

Change 3267573 on 2017/01/23 by [REDACTED_PERSON]

	#loc UE4DocJpn Initial translation against INT#3244824

Change 3267575 on 2017/01/23 by [REDACTED_PERSON]

	Updated File against  INTSourceChangeList:3191723

Change 3267594 on 2017/01/23 by [REDACTED_PERSON]

	Updated File against INTSourceChangelist:3188564

Change 3267600 on 2017/01/23 by [REDACTED_PERSON]

	Updated File against INTSourceChangelist:3188564

Change 3267605 on 2017/01/23 by [REDACTED_PERSON]

	Updated File against INTSourceChangelist:3149430

Change 3267607 on 2017/01/23 by [REDACTED_PERSON]

	Updated File against INTSourceChangelist:3149199

Change 3267614 on 2017/01/23 by [REDACTED_PERSON]

	Updated File against INTSourceChangelist:3147055

Change 3267618 on 2017/01/23 by [REDACTED_PERSON]

	Updated File against INTSourceChangelist:3213117

Change 3267838 on 2017/01/23 by [REDACTED_PERSON]

	#ue4 docs

	Updated Character Animation How-to
	- refactored to include blending animations

Change 3268172 on 2017/01/23 by [REDACTED_PERSON]

	#UE4 Docs: Fixed an issue with the link to the Skel mesh pose page.

Change 3268173 on 2017/01/23 by [REDACTED_PERSON]

	#UE4 Docs:Applying peer feedback.

Change 3268222 on 2017/01/23 by [REDACTED_PERSON]

	#UE4 Docs: Adding the doc updates.

Change 3268421 on 2017/01/23 by [REDACTED_PERSON]

	#UE4 Docs: Small sentence formatting change to SkeletalMeshConversion page

Change 3269086 on 2017/01/23 by [REDACTED_PERSON]

	Updated File against INTSourceChangelist:3215424

Change 3269098 on 2017/01/23 by [REDACTED_PERSON]

	Updated File against INTSourceChangelist:3259022

Change 3269104 on 2017/01/23 by [REDACTED_PERSON]

	Updated File against INTSourceChangelist:3259022

Change 3269105 on 2017/01/23 by [REDACTED_PERSON]

	Updated File against INTSourceChangelist:3259022

Change 3269139 on 2017/01/23 by [REDACTED_PERSON]

	Updated File against INTSourceChangelist:3177941

Change 3269147 on 2017/01/23 by [REDACTED_PERSON]

	Updated File against INTSourceChangelist:3175238

Change 3269148 on 2017/01/23 by [REDACTED_PERSON]

	Updated File against INTSourceChangelist:3155346

Change 3269151 on 2017/01/23 by [REDACTED_PERSON]

	Updated File against INTSourceChangelist:3155346

Change 3269206 on 2017/01/24 by [REDACTED_PERSON]

	Updated File against INTSourceChangelist:3155346

Change 3269218 on 2017/01/24 by [REDACTED_PERSON]

	Updated File against INTSourceChangelist:3155346

Change 3269222 on 2017/01/24 by [REDACTED_PERSON]

	Updated File against INTSourceChangelist:3151855

Change 3269249 on 2017/01/24 by [REDACTED_PERSON]

	#loc UE4DocJpn File updated against INT#3225556

Change 3269293 on 2017/01/24 by [REDACTED_PERSON]

	#loc UE4DocJpn File updated against INT#3267838

Change 3269302 on 2017/01/24 by [REDACTED_PERSON]

	#loc UE4DocKOR

Change 3269455 on 2017/01/24 by [REDACTED_PERSON]

	#loc UE4DocKOR

Change 3269482 on 2017/01/24 by [REDACTED_PERSON]

	#loc UE4DocKOR

Change 3270938 on 2017/01/24 by [REDACTED_PERSON]

	Updated File against INTSourceChangelist:3151855

Change 3270971 on 2017/01/24 by [REDACTED_PERSON]

	Updated File against INTSourceChangelist:3268222

Change 3271008 on 2017/01/24 by [REDACTED_PERSON]

	Updated File against INTSourceChangelist:3268222

Change 3271010 on 2017/01/24 by [REDACTED_PERSON]

	Updated File against INTSourceChangelist:3268222

Change 3271020 on 2017/01/24 by [REDACTED_PERSON]

	Updated File against INTSourceChangelist:3262227

Change 3271034 on 2017/01/24 by [REDACTED_PERSON]

	Updated File against INTSourceChangelist:3262227

Change 3271044 on 2017/01/24 by [REDACTED_PERSON]

	Updated File against INTSourceChangelist:3261961

Change 3271052 on 2017/01/24 by [REDACTED_PERSON]

	Updated File against INTSourceChangelist:3261961

Change 3271081 on 2017/01/24 by [REDACTED_PERSON]

	#loc UE4DocKOR

Change 3271108 on 2017/01/25 by [REDACTED_PERSON]

	Updated File against INTSourceChangelist:3260598

Change 3271110 on 2017/01/25 by [REDACTED_PERSON]

	Updated File against INTSourceChangelist:3260598

Change 3271112 on 2017/01/25 by [REDACTED_PERSON]

	Updated File against INTSourceChangelist:3259022

Change 3271113 on 2017/01/25 by [REDACTED_PERSON]

	Updated File against INTSourceChangeList:3259022

Change 3271114 on 2017/01/25 by [REDACTED_PERSON]

	Updated File against INTSourceChangelist:3259022

Change 3271116 on 2017/01/25 by [REDACTED_PERSON]

	Updated File against INTSourceChangelist:3259022

Change 3271120 on 2017/01/25 by [REDACTED_PERSON]

	Updated File against INTSourceChangelist:3259022

Change 3271123 on 2017/01/25 by [REDACTED_PERSON]

	Updated File against INTSourceChangelist:3259022

Change 3271130 on 2017/01/25 by [REDACTED_PERSON]

	Updated File against INTSourceChangelist:3259022

Change 3271135 on 2017/01/25 by [REDACTED_PERSON]

	Updated File against INTSourceChangelist:3259022

Change 3271138 on 2017/01/25 by [REDACTED_PERSON]

	Updated File against INTSourceChangelist:3259022

Change 3271145 on 2017/01/25 by [REDACTED_PERSON]

	Upated File against INTSourceChangelist:3244370

Change 3271163 on 2017/01/25 by [REDACTED_PERSON]

	Updated File against INTSourceChangelist:3244370

Change 3271174 on 2017/01/25 by [REDACTED_PERSON]

	Updated File against INTSourceChangelist:3218894

Change 3271178 on 2017/01/25 by [REDACTED_PERSON]

	Updated File against INTSourceChangelist:3242329

Change 3271234 on 2017/01/25 by [REDACTED_PERSON]

	#loc UE4DocKOR

Change 3271292 on 2017/01/25 by [REDACTED_PERSON]

	#loc UE4DocCHN. do translation

Change 3271840 on 2017/01/25 by [REDACTED_PERSON]

	updated DB file

Change 3272803 on 2017/01/25 by [REDACTED_PERSON]

	Updated File against INTSourceChangelist:3257013

Change 3272824 on 2017/01/25 by [REDACTED_PERSON]

	Updated File against INTSourceChangelist:3242329

Change 3272828 on 2017/01/25 by [REDACTED_PERSON]

	Updated File against INTSourceChangelist:3264527

Change 3272856 on 2017/01/25 by [REDACTED_PERSON]

	Updated File against INTSourceChangelist:3264527

Change 3272857 on 2017/01/25 by [REDACTED_PERSON]

	#loc UE4DocJpn File updated against INT#3255748

Change 3272869 on 2017/01/25 by [REDACTED_PERSON]

	Updated File against INTSourceChangelist:3264524

Change 3272876 on 2017/01/25 by [REDACTED_PERSON]

	Updated File against INTSourceChangelist:3264524

Change 3272878 on 2017/01/25 by [REDACTED_PERSON]

	#loc UE4DocJpn File updated against INT#3268421

Change 3272887 on 2017/01/25 by [REDACTED_PERSON]

	Updated File against INTSourceChangelist:3264524

Change 3272895 on 2017/01/25 by [REDACTED_PERSON]

	Updated File against INTSourceChangelist:3264524

Change 3272898 on 2017/01/25 by [REDACTED_PERSON]

	#loc UE4DocJpn File updated against INT#3268172

Change 3272902 on 2017/01/25 by [REDACTED_PERSON]

	Updated File against INTSourceChangelist:3264524

Change 3272913 on 2017/01/25 by [REDACTED_PERSON]

	Updated File against INTSourceChangelist:3264524

Change 3272915 on 2017/01/25 by [REDACTED_PERSON]

	Updated File against INTSourceChangelist:3261961

Change 3272925 on 2017/01/25 by [REDACTED_PERSON]

	Updated File against INTSourceChangelist:3260598

Change 3272929 on 2017/01/25 by [REDACTED_PERSON]

	#loc UE4DocJpn File updated against INT#3224907

Change 3272930 on 2017/01/25 by [REDACTED_PERSON]

	Updated File against INTSourceChangelist:3260598

Change 3272936 on 2017/01/25 by [REDACTED_PERSON]

	#loc UE4DocJpn File updated against INT#3151855

Change 3272946 on 2017/01/25 by [REDACTED_PERSON]

	Updated File against INTSourceChangelist:3256943

Change 3272947 on 2017/01/25 by [REDACTED_PERSON]

	#loc UE4DocCHN. do translation

Change 3272982 on 2017/01/26 by [REDACTED_PERSON]

	#loc UE4DocJpn File updated against INT#3255735

Change 3273285 on 2017/01/26 by [REDACTED_PERSON]

	#UE4 Docs

	Apply SME feedback to Forward Rendering document.

Change 3273300 on 2017/01/26 by [REDACTED_PERSON]

	#UE4 Docs

	Minor updates based on Peer Review.

Change 3273317 on 2017/01/26 by [REDACTED_PERSON]

	#UE4 Docs

	Applying Peer Review Feedback.

Change 3273564 on 2017/01/26 by [REDACTED_PERSON]

	Updating engine version on Landing page.
	Revising description on Getting Started page.

	#UE4doc

Change 3273672 on 2017/01/26 by [REDACTED_PERSON]

	#loc UE4DocKOR

Change 3274763 on 2017/01/27 by [REDACTED_PERSON]

	#loc UE4DocJpn Initial translation against INT#3108692

Change 3274767 on 2017/01/27 by [REDACTED_PERSON]

	Updated File against INTSourceChangelist:3255748

Change 3274768 on 2017/01/27 by [REDACTED_PERSON]

	Updated File against INTSourceChangelist:3108692

Change 3274776 on 2017/01/27 by [REDACTED_PERSON]

	#loc UE4DocJpn File updated against INT#3273317

Change 3274781 on 2017/01/27 by [REDACTED_PERSON]

	#loc UE4DocJpn File updated against INT#3273300

Change 3274787 on 2017/01/27 by [REDACTED_PERSON]

	Updated File against INTSourceChangelist:3242329

Change 3274832 on 2017/01/27 by [REDACTED_PERSON]

	Updated File against INTSourceChangelist:3242329

Change 3274837 on 2017/01/27 by [REDACTED_PERSON]

	#loc UE4DocJpn File updated against INT#3273317

Change 3274838 on 2017/01/27 by [REDACTED_PERSON]

	Updated File against INTSourceChangelist:3242329

Change 3274856 on 2017/01/27 by [REDACTED_PERSON]

	Updated File against INTSourceChangelist:3242329

Change 3275198 on 2017/01/27 by [REDACTED_PERSON]

	#UE4 Docs: Applying peer feedback to this doc.

Change 3275202 on 2017/01/27 by [REDACTED_PERSON]

	Fixed broken link

Change 3275203 on 2017/01/27 by [REDACTED_PERSON]

	Updated search box and button to use CSS styles

Change 3275206 on 2017/01/27 by [REDACTED_PERSON]

	Added link to Contact Shadows page

Change 3275207 on 2017/01/27 by [REDACTED_PERSON]

	Added link to Contact Shadows page

Change 3275208 on 2017/01/27 by [REDACTED_PERSON]

	Fixed formatting

Change 3275209 on 2017/01/27 by [REDACTED_PERSON]

	Fixed links to videos

Change 3275254 on 2017/01/27 by [REDACTED_PERSON]

	Updated metadata

Change 3275257 on 2017/01/27 by [REDACTED_PERSON]

	Fixed broken link

Change 3275285 on 2017/01/27 by [REDACTED_PERSON]

	Added redirect page

Change 3275293 on 2017/01/27 by [REDACTED_PERSON]

	Updated metadata

Change 3275367 on 2017/01/27 by [REDACTED_PERSON]

	Fixed broken link

Change 3275566 on 2017/01/27 by [REDACTED_PERSON]

	#UE Docs #jira UEDOC-4459 - Spelling and formatting fixes for GoogleVR pages

Change 3275878 on 2017/01/27 by [REDACTED_PERSON]

	Updated metadata

Change 3275881 on 2017/01/27 by [REDACTED_PERSON]

	Removing old learning track prototype pages

Change 3275884 on 2017/01/27 by [REDACTED_PERSON]

	Adding learning track template

Change 3276951 on 2017/01/30 by [REDACTED_PERSON]

	#loc UE4DocJpn File updated against INT#3275254

Change 3276952 on 2017/01/30 by [REDACTED_PERSON]

	#loc UE4DocJpn File updated against INT#3275254

Change 3276954 on 2017/01/30 by [REDACTED_PERSON]

	Updated File against INTSourceChangelist:3242329

Change 3276964 on 2017/01/30 by [REDACTED_PERSON]

	#loc UE4DocJpn File updated against INT#3275878

Change 3276981 on 2017/01/30 by [REDACTED_PERSON]

	Updated File against INTSourceChangelist:3242329

Change 3276984 on 2017/01/30 by [REDACTED_PERSON]

	#loc UE4DocJpn File updated against INT#3275878

Change 3277001 on 2017/01/30 by [REDACTED_PERSON]

	#loc UE4DocJpn File updated against INT#3275202

Change 3277051 on 2017/01/30 by [REDACTED_PERSON]

	#loc UE4DocJpn File updated against INT#3275202

Change 3277055 on 2017/01/30 by [REDACTED_PERSON]

	Updated File against INTSourceChangelist:3242329

Change 3277377 on 2017/01/30 by [REDACTED_PERSON]

	Set page to publish

Change 3277378 on 2017/01/30 by [REDACTED_PERSON]

	Force page height update after async page load returns

Change 3277768 on 2017/01/30 by [REDACTED_PERSON]

	#UE4 Docs: Fixing minor grammatical error and rewording the intro sentences.

Change 3277813 on 2017/01/30 by [REDACTED_PERSON]

	#UE4 Docs: Capsule Shadows docs: Removing files that are no longer relevant for cleanup of UEDOC-3973.

Change 3277874 on 2017/01/30 by [REDACTED_PERSON]

	#UE4 Docs: UEDOC-3973 - SME Review edits and rewrite of Capsule Shadows Overview Doc. Moved How-to's to QuickStart. Still minor work to do with navigation links for pages in QS

Change 3278417 on 2017/01/30 by [REDACTED_PERSON]

	#loc UE4DocJPN moved or deleted

Change 3278584 on 2017/01/30 by [REDACTED_PERSON]

	#loc UE4DocKOR moved or deleted

Change 3278652 on 2017/01/31 by [REDACTED_PERSON]

	#loc UE4DocJpn File updated against INT#3277768

Change 3278666 on 2017/01/31 by [REDACTED_PERSON]

	#loc UE4DocJpn File updated against INT#3275878

Change 3278669 on 2017/01/31 by [REDACTED_PERSON]

	#loc UE4DocJpn File updated against INT#3275202

Change 3278673 on 2017/01/31 by [REDACTED_PERSON]

	Updated File against INTSourceChangelist:3242329

Change 3278678 on 2017/01/31 by [REDACTED_PERSON]

	Updated File against INTSourceChangelist:3178047

Change 3278681 on 2017/01/31 by [REDACTED_PERSON]

	Updated File against INTSourceChangelist:3178047

Change 3278682 on 2017/01/31 by [REDACTED_PERSON]

	#loc UE4DocJpn File updated against INT#3275202

Change 3278694 on 2017/01/31 by [REDACTED_PERSON]

	#loc UE4DocKOR

Change 3278718 on 2017/01/31 by [REDACTED_PERSON]

	Updated File against INTSourceChangelist:3150518

Change 3278812 on 2017/01/31 by [REDACTED_PERSON]

	#loc UE4DocKOR

Change 3278909 on 2017/01/31 by [REDACTED_PERSON]

	typo_correction

Change 3278910 on 2017/01/31 by [REDACTED_PERSON]

	#loc UE4DocKOR

Change 3279502 on 2017/01/31 by [REDACTED_PERSON]

	#ue4 docs

	Removed "docs" tag from page that is public so it shows up in the list.

Change 3280006 on 2017/01/31 by [REDACTED_PERSON]

	#UE4 Docs:Applying Peer Feedback.

Change 3280524 on 2017/01/31 by [REDACTED_PERSON]

	#loc UE4DocJpn File updated against INT#3275367

Change 3280530 on 2017/01/31 by [REDACTED_PERSON]

	#loc UE4DocJpn File updated against INT#3275254

Change 3280574 on 2017/01/31 by [REDACTED_PERSON]

	Updated File against  INTSourceChangelist:3277377

Change 3280575 on 2017/01/31 by [REDACTED_PERSON]

	Updated File against INTSourceChangelist:3275566

Change 3280592 on 2017/01/31 by [REDACTED_PERSON]

	#loc UE4DocKOR

Change 3280600 on 2017/01/31 by [REDACTED_PERSON]

	#loc UE4DocJpn File updated against INT#3275202

Change 3280677 on 2017/02/01 by [REDACTED_PERSON]

	Updated File against INTSourceChangelist:3275566

Change 3280700 on 2017/02/01 by [REDACTED_PERSON]

	#loc UE4DocJpn File updated against INT#3275202

Change 3280705 on 2017/02/01 by [REDACTED_PERSON]

	Updated File against INTSourceChangelist:3275257

Change 3280708 on 2017/02/01 by [REDACTED_PERSON]

	Updated File against INTSourceChangelist:3275208

Change 3280717 on 2017/02/01 by [REDACTED_PERSON]

	#loc UE4DocJpn File updated against INT#3275202

Change 3280718 on 2017/02/01 by [REDACTED_PERSON]

	Updated File against INTSourceChangelist:3275206

Change 3280721 on 2017/02/01 by [REDACTED_PERSON]

	typo correction

Change 3280723 on 2017/02/01 by [REDACTED_PERSON]

	Updated File against INTSourceChangelist:3275207

Change 3280724 on 2017/02/01 by [REDACTED_PERSON]

	Updated File against  INTSourceChangelist:3273564

Change 3280727 on 2017/02/01 by [REDACTED_PERSON]

	Updated File against INTSourceChangelist:3273564

Change 3280729 on 2017/02/01 by [REDACTED_PERSON]

	#loc UE4DocJpn File updated against INT#3275202

Change 3280733 on 2017/02/01 by [REDACTED_PERSON]

	Updated File against INTSourceChangelist:3215926

Change 3280746 on 2017/02/01 by [REDACTED_PERSON]

	#loc UE4DocJpn File updated against INT#3275202

Change 3280748 on 2017/02/01 by [REDACTED_PERSON]

	#loc UE4DocJpn File updated against INT#3279502

Change 3280749 on 2017/02/01 by [REDACTED_PERSON]

	Updated File against INTSourceChangelist:3238697

Change 3280750 on 2017/02/01 by [REDACTED_PERSON]

	#loc UE4DocJpn File updated against INT#3275202

Change 3280766 on 2017/02/01 by [REDACTED_PERSON]

	Updated File against INTSourceChangelist:3150518

Change 3280770 on 2017/02/01 by [REDACTED_PERSON]

	Updated File against INTSourceChangelist:3275293

Change 3280779 on 2017/02/01 by [REDACTED_PERSON]

	Updated File against INTSourceChangelist:3206215

Change 3280829 on 2017/02/01 by [REDACTED_PERSON]

	#loc UE4DocKOR

Change 3281495 on 2017/02/01 by [REDACTED_PERSON]

	Fixed according to SME review notes.

Change 3281540 on 2017/02/01 by [REDACTED_PERSON]

	Fixing parent paths and resetting documentation order on landing page.

	#UE4doc

Change 3282373 on 2017/02/01 by [REDACTED_PERSON]

	Updated File against INTSourceChangelist:3178047

Change 3282378 on 2017/02/01 by [REDACTED_PERSON]

	Updated File against INTSourceChangelist:3150518

Change 3282383 on 2017/02/01 by [REDACTED_PERSON]

	Updated File against INTSourceChangelist:3108692

Change 3282391 on 2017/02/01 by [REDACTED_PERSON]

	Updated File against INTSourceChangelist:3281540

Change 3282416 on 2017/02/01 by [REDACTED_PERSON]

	Updated File against INTSourceChangelist:3281495

Change 3282662 on 2017/02/02 by [REDACTED_PERSON]

	Updated File against INTSourceChangelist:3238697

Change 3282665 on 2017/02/02 by [REDACTED_PERSON]

	Updated File against INTSourceChangelist:3151855

Change 3282668 on 2017/02/02 by [REDACTED_PERSON]

	Updated File against INTSourceChangelist:3150518

Change 3282676 on 2017/02/02 by [REDACTED_PERSON]

	Updated File against INTSourceChangelist:3260884

Change 3282679 on 2017/02/02 by [REDACTED_PERSON]

	Updated File against INTSourceChangelist:3244370

Change 3282686 on 2017/02/02 by [REDACTED_PERSON]

	Updated File against INTSourceChangelist:3244370

Change 3282692 on 2017/02/02 by [REDACTED_PERSON]

	Updated File against INTSourceChangelist:3244370

Change 3282696 on 2017/02/02 by [REDACTED_PERSON]

	Updated File against  INTSourceChangelist:3244370

Change 3282701 on 2017/02/02 by [REDACTED_PERSON]

	Updated File against INTSourceChangelist:3244370

Change 3282705 on 2017/02/02 by [REDACTED_PERSON]

	Updated File against INTSourceChangelist:3244370

Change 3282716 on 2017/02/02 by [REDACTED_PERSON]

	Updated File against INTSourceChangelist:3244370

Change 3282750 on 2017/02/02 by [REDACTED_PERSON]

	Updated File against INTSourceChangelist:3244370

Change 3282837 on 2017/02/02 by [REDACTED_PERSON]

	#loc UE4DocKOR

Change 3283142 on 2017/02/02 by [REDACTED_PERSON]

	Moved "custom importer" information to the bottom of the document as SME recommended.

Change 3283566 on 2017/02/02 by [REDACTED_PERSON]

	Fixing cull volume information. No JIRA for this, based on a UDN post and code investigation.

Change 3284397 on 2017/02/02 by [REDACTED_PERSON]

	#UE4doc

	Repaired some capitalization errors throughout the document.

Change 3284565 on 2017/02/02 by [REDACTED_PERSON]

	#loc UE4DocJpn File updated against INT#3244012

Change 3284573 on 2017/02/02 by [REDACTED_PERSON]

	#loc UE4DocJpn File updated against INT#3244012

Change 3284586 on 2017/02/02 by [REDACTED_PERSON]

	#loc UE4DocJpn File updated against INT#3244012

Change 3284593 on 2017/02/02 by [REDACTED_PERSON]

	#loc UE4DocJpn File updated against INT#3244012

Change 3284632 on 2017/02/02 by [REDACTED_PERSON]

	#loc UE4DocJpn File updated against INT#3244012

Change 3284645 on 2017/02/02 by [REDACTED_PERSON]

	#loc UE4DocJpn File updated against INT#3244012

Change 3284656 on 2017/02/02 by [REDACTED_PERSON]

	#loc UE4DocJpn File updated against INT#3244012

Change 3284728 on 2017/02/02 by [REDACTED_PERSON]

	Updated File against INTSourceChangelist:3244370

Change 3284777 on 2017/02/03 by [REDACTED_PERSON]

	Updated File against INTSourceChangelist:3284397

Change 3284785 on 2017/02/03 by [REDACTED_PERSON]

	#loc UE4DocJpn File updated against INT#3275202

Change 3284795 on 2017/02/03 by [REDACTED_PERSON]

	Updated File against INTSourceChangelist:3150518

Change 3284797 on 2017/02/03 by [REDACTED_PERSON]

	Updated File against INTSourceChangelist:3150518

Change 3284810 on 2017/02/03 by [REDACTED_PERSON]

	Updated File against INTSourceChangelist:3150518

Change 3284820 on 2017/02/03 by [REDACTED_PERSON]

	Updated File against INTSourceChangelist:3150518

Change 3284829 on 2017/02/03 by [REDACTED_PERSON]

	Updated File against INTSourceChangelist:3150518

Change 3284855 on 2017/02/03 by [REDACTED_PERSON]

	Updated File against INTSourceChangelist:3253176

Change 3284861 on 2017/02/03 by [REDACTED_PERSON]

	Updated File against  INTSourceChangelist:3150518

Change 3284871 on 2017/02/03 by [REDACTED_PERSON]

	#loc UE4DocKor removed ??

Change 3284911 on 2017/02/03 by [REDACTED_PERSON]

	#loc UE4DocKOR

Change 3284940 on 2017/02/03 by [REDACTED_PERSON]

	#loc UE4DocCHN. do translation. ( fix missing pic, Updated File against INTSourceChangelist:3223841 )

Change 3285321 on 2017/02/03 by [REDACTED_PERSON]

	Fixing a small typo in Programming Quick Start.

Change 3285322 on 2017/02/03 by [REDACTED_PERSON]

	Updated metadata

Change 3286792 on 2017/02/04 by [REDACTED_PERSON]

	#loc UE4DocKOR

Change 3287218 on 2017/02/05 by [REDACTED_PERSON]

	Updated File against INTSourceChangelist:3150518

Change 3287220 on 2017/02/05 by [REDACTED_PERSON]

	Updated File against INTSourceChangelist:3150518

Change 3287224 on 2017/02/05 by [REDACTED_PERSON]

	#loc UE4DocJpn File updated against INT#3283566

Change 3287227 on 2017/02/05 by [REDACTED_PERSON]

	#loc UE4DocJpn File updated against INT#3285322

Change 3287231 on 2017/02/05 by [REDACTED_PERSON]

	Updated File against INTSourceChangelist:3150518

Change 3287236 on 2017/02/05 by [REDACTED_PERSON]

	Updated File against INTSourceChangelist:3150518

Change 3287252 on 2017/02/05 by [REDACTED_PERSON]

	Update File against INTSourceChangelist:3108692

Change 3287255 on 2017/02/05 by [REDACTED_PERSON]

	Updated File against INTSourceChangelist:3150518

Change 3287257 on 2017/02/05 by [REDACTED_PERSON]

	Updated File against INTSourceChangelist:3108692

Change 3287273 on 2017/02/06 by [REDACTED_PERSON]

	#loc UE4DocCHN moved or deleted

Change 3287280 on 2017/02/06 by [REDACTED_PERSON]

	#loc UE4DocJpn File updated against INT#3280006

Change 3287287 on 2017/02/06 by [REDACTED_PERSON]

	Fixed a typo.

Change 3287288 on 2017/02/06 by [REDACTED_PERSON]

	#loc UE4DocJpn File updated against INT#:3206237

Change 3287291 on 2017/02/06 by [REDACTED_PERSON]

	#loc UE4DocJpn File updated against INT#3266245

Change 3287301 on 2017/02/06 by [REDACTED_PERSON]

	#loc UE4DocJpn File updated against INT#3275202

Change 3287317 on 2017/02/06 by [REDACTED_PERSON]

	Updated File against INTSourceChangelist:3108692

Change 3287319 on 2017/02/06 by [REDACTED_PERSON]

	Fixed translation

Change 3287337 on 2017/02/06 by [REDACTED_PERSON]

	Updated File against INTSourceChangelist:3108692

Change 3287338 on 2017/02/06 by [REDACTED_PERSON]

	Updated File against INTSourceChangelist:3150518

Change 3287342 on 2017/02/06 by [REDACTED_PERSON]


diff --git a/Engine/Documentation/Source/Engine/AI/EnvironmentQuerySystem/QuickStart/7/7.INT.udn b/Engine/Documentation/Source/Engine/AI/EnvironmentQuerySystem/QuickStart/7/7.INT.udn
--- a/Engine/Documentation/Source/Engine/AI/EnvironmentQuerySystem/QuickStart/7/7.INT.udn
+++ b/Engine/Documentation/Source/Engine/AI/EnvironmentQuerySystem/QuickStart/7/7.INT.udn
@@ -28,7 +28,7 @@ The setup for this AI Controller is a single **Run Behavior Tree** node running
 
 1. Double-click the Runner_AI_CON Blueprint in the **Content Browser** to open it for editing.
 
-2. In the **Event Graph**, have **Event BeginPlay** execute a **Run Beahvior Tree** node using the RunnerBT as the **BTAsset**:
+2. In the **Event Graph**, have **Event BeginPlay** execute a **Run Behavior Tree** node using the RunnerBT as the **BTAsset**:
 
 	![](setup.png)
 
diff --git a/Engine/Documentation/Source/Engine/Actors/Actors.INT.udn b/Engine/Documentation/Source/Engine/Actors/Actors.INT.udn
--- a/Engine/Documentation/Source/Engine/Actors/Actors.INT.udn
+++ b/Engine/Documentation/Source/Engine/Actors/Actors.INT.udn
@@ -5,35 +5,8 @@ Description:A definition of Actors and how they are used in level design, includ
 Navigation:topic
 parent:Engine/Editor
 order:5
-
-[VAR:playlistnav]
-[OBJECT:PlaylistNavigation]
-	[PARAMLITERAL:id]
-		LevelDesign_Beginner
-	[/PARAMLITERAL]
-	[PARAM:previous]
-		[Next Topic](Engine/UI/LevelEditor/Viewports/Basics "%Engine/UI/LevelEditor/Viewports/Basics:title%")
-	[/PARAM]
-	[PARAM:next]
-		[Next Topic](Engine/Actors/Placement "%Engine/Actors/Placement:title%")
-	[/PARAM]
-	[PARAM:home]
-		[](Tracks/LevelDesign/Beginner)
-		%Tracks/LevelDesign/Beginner:menu%
-	[/PARAM]
-[/OBJECT]
-[OBJECT:PlaylistNavigation]
-	[PARAMLITERAL:id]
-		Programming_Beginner
-	[/PARAMLITERAL]
-	[PARAM:previous]
-		[Previous Topic](Engine/UI/LevelEditor "%Engine/UI/LevelEditor:title%")
-	[/PARAM]
-	[PARAM:home]
-		[](Tracks/Programming/Beginner)
-	[/PARAM]
-[/OBJECT]
-[/VAR]
+type:landing
+track:LevelDesign_Beginner|7
 
 [VAR:Topic]
 [OBJECT:Topic]
@@ -74,8 +47,6 @@ order:5
 [/OBJECT]
 [/VAR]
 
-%playlistnav%
-
 [REGION:banner]
 ![](actorBanner.png)
 [/REGION]
@@ -122,11 +93,3 @@ properties will be modified to cause them to look or behave appropriately.
 	* [](Engine/Actors/CameraActors "%Engine/Actors/CameraActors:description%")
 	[/PARAM]
 [/OBJECT]
-
-
-%playlistnav%
-
-
-
-
-
diff --git a/Engine/Documentation/Source/Engine/Actors/Brushes/Brushes.INT.udn b/Engine/Documentation/Source/Engine/Actors/Brushes/Brushes.INT.udn
--- a/Engine/Documentation/Source/Engine/Actors/Brushes/Brushes.INT.udn
+++ b/Engine/Documentation/Source/Engine/Actors/Brushes/Brushes.INT.udn
@@ -8,24 +8,6 @@ Related: Engine/Actors/Volumes
 version: 4.9
 
 
-[VAR:playlistnav]
-[OBJECT:PlaylistNavigation]
-	[PARAMLITERAL:id]
-		LevelDesign_Beginner
-	[/PARAMLITERAL]
-	[PARAM:previous]
-		[Previous Topic](Engine/Actors/StaticMeshActor "%Engine/Actors/StaticMeshActor:title%")
-	[/PARAM]
-	[PARAM:next]
-		[Next Topic](Engine/Actors/Volumes "%Engine/Actors/Volumes:title%")
-	[/PARAM]
-	[PARAM:home]
-		[](Tracks/LevelDesign/Beginner)
-		%Tracks/LevelDesign/Beginner:menu%
-	[/PARAM]
-[/OBJECT]
-[/VAR]
-
 [VAR:Topic]
 [OBJECT:Topic]
 	[PARAM:image]
@@ -65,8 +47,6 @@ version: 4.9
 [/OBJECT]
 [/VAR]
 
-%playlistnav%
-
 
 [TOC (end:2)]
 
@@ -454,10 +434,3 @@ The three types of solidity are explained below.
 * Do not create cuts in the surrounding Brushes.
 [/EXCERPT:Solidity]
 
-
-%playlistnav%
-
-
-
-
-
diff --git a/Engine/Documentation/Source/Engine/Actors/Grouping/Grouping.INT.udn b/Engine/Documentation/Source/Engine/Actors/Grouping/Grouping.INT.udn
--- a/Engine/Documentation/Source/Engine/Actors/Grouping/Grouping.INT.udn
+++ b/Engine/Documentation/Source/Engine/Actors/Grouping/Grouping.INT.udn
@@ -4,25 +4,6 @@ Crumbs:%ROOT%, Engine, Engine/Actors
 Description:Guide to creating and working with groups of actors within Unreal Editor.
 version: 4.9
 
-
-[VAR:playlistnav]
-[OBJECT:PlaylistNavigation]
-	[PARAMLITERAL:id]
-		LevelDesign_Beginner
-	[/PARAMLITERAL]
-	[PARAM:previous]
-		[Previous Topic](Engine/Actors/Mobility "%Engine/Actors/Mobility:title%")
-	[/PARAM]
-	[PARAM:next]
-		[Next Topic](Engine/Actors/Types "%Engine/Actors/Types:title%")
-	[/PARAM]
-	[PARAM:home]
-		[](Tracks/LevelDesign/Beginner)
-		%Tracks/LevelDesign/Beginner:menu%
-	[/PARAM]
-[/OBJECT]
-[/VAR]
-
 [VAR:Topic]
 [OBJECT:Topic]
 	[PARAM:image]
@@ -63,8 +44,6 @@ version: 4.9
 [/OBJECT]
 [/VAR]
 
-%playlistnav%
-
 
 [TOC(start:2)]
 
@@ -189,5 +168,3 @@ The **Add Actors to Group** option is only available when one group is selected
 * Group pivots can be set by selecting the desired group and then holding down **Alt** and clicking the **Middle Mouse Button** in the viewport for the pivot location.
 * It is not possible to group Actors that are in different levels.  Moving an Actor that is currently in a group from one level to another will remove it from the existing group.  However, entire groups can be moved across levels.
 
-
-%playlistnav%
diff --git a/Engine/Documentation/Source/Engine/Actors/Mobility/ActorMobility.INT.udn b/Engine/Documentation/Source/Engine/Actors/Mobility/ActorMobility.INT.udn
--- a/Engine/Documentation/Source/Engine/Actors/Mobility/ActorMobility.INT.udn
+++ b/Engine/Documentation/Source/Engine/Actors/Mobility/ActorMobility.INT.udn
@@ -4,27 +4,6 @@ Crumbs: %ROOT%, Engine, Engine/Actors
 Description:Setting that controls whether an Actor will be allowed to move or change in some way during gameplay.
 version: 4.9
 
-
-[VAR:playlistnav]
-[OBJECT:PlaylistNavigation]
-	[PARAMLITERAL:id]
-		LevelDesign_Beginner
-	[/PARAMLITERAL]
-	[PARAM:previous]
-		[Previous Topic](Engine/Actors/Transform "%Engine/Actors/Transform:title%")
-	[/PARAM]
-	[PARAM:next]
-		[Next Topic](Engine/Actors/Grouping "%Engine/Actors/Grouping:title%")
-	[/PARAM]
-	[PARAM:home]
-		[](Tracks/LevelDesign/Beginner)
-		%Tracks/LevelDesign/Beginner:menu%
-	[/PARAM]
-[/OBJECT]
-[/VAR]
-
-%playlistnav%
-
 [EXCERPT:Intro]
 ![](TransformMobility.png)
 
@@ -62,6 +41,3 @@ Stationary is used for **Light Actors** that do not move, but may update in some
 [/EXCERPT:Movable]
 
 -->
-
-
-%playlistnav%
diff --git a/Engine/Documentation/Source/Engine/Actors/Placement/PlacingActors.INT.udn b/Engine/Documentation/Source/Engine/Actors/Placement/PlacingActors.INT.udn
--- a/Engine/Documentation/Source/Engine/Actors/Placement/PlacingActors.INT.udn
+++ b/Engine/Documentation/Source/Engine/Actors/Placement/PlacingActors.INT.udn
@@ -5,24 +5,7 @@ Description:Shows how you can place Actors such as props, lights, cameras, etc.
 Related: Engine/UI/ClassViewer
 Related: GettingStarted\HowTo
 version: 4.9
-
-[VAR:playlistnav]
-[OBJECT:PlaylistNavigation]
-	[PARAMLITERAL:id]
-		LevelDesign_Beginner
-	[/PARAMLITERAL]
-	[PARAM:previous]
-		[Previous Topic](Engine/Actors "%Engine/Actors:title%")
-	[/PARAM]
-	[PARAM:next]
-		[Next Topic](Engine/Actors/Selection "%Engine/Actors/Selection:title%")
-	[/PARAM]
-	[PARAM:home]
-		[](Tracks/LevelDesign/Beginner)
-		%Tracks/LevelDesign/Beginner:menu%
-	[/PARAM]
-[/OBJECT]
-[/VAR]
+track:LevelDesign_Beginner|8
 
 [VAR:Topic]
 [OBJECT:Topic]
@@ -64,8 +47,6 @@ version: 4.9
 [/OBJECT]
 [/VAR]
 
-%playlistnav%
-
 At the most fundamental level, an **Actor** is any Object that you can place in a level and this page will show you the various methods in which you can place those Actors inside your levels. 
 
 [TOC(start:2)]
@@ -164,11 +145,3 @@ A more advanced method of placing Actors is by using the **Class Viewer**, which
 	![](Place8.png)
 
 For more information on working with the **Class Viewer**, refer to the [](Engine/UI/ClassViewer) documentation. 
-
-
-%playlistnav%
-
-
-
-
-
diff --git a/Engine/Documentation/Source/Engine/Actors/Selection/ActorSelection.INT.udn b/Engine/Documentation/Source/Engine/Actors/Selection/ActorSelection.INT.udn
--- a/Engine/Documentation/Source/Engine/Actors/Selection/ActorSelection.INT.udn
+++ b/Engine/Documentation/Source/Engine/Actors/Selection/ActorSelection.INT.udn
@@ -4,24 +4,6 @@ Crumbs: %ROOT%, Engine, Engine/Actors
 Description:Overview of methods available for selecting Actors in the Level Editor viewports.
 version: 4.9
 
-[VAR:playlistnav]
-[OBJECT:PlaylistNavigation]
-	[PARAMLITERAL:id]
-		LevelDesign_Beginner
-	[/PARAMLITERAL]
-	[PARAM:previous]
-		[Previous Topic](Engine/Actors "%Engine/Actors:title%")
-	[/PARAM]
-	[PARAM:next]
-		[Next Topic](Engine/Actors/Transform "%Engine/Actors/Transform:title%")
-	[/PARAM]
-	[PARAM:home]
-		[](Tracks/LevelDesign/Beginner)
-		%Tracks/LevelDesign/Beginner:menu%
-	[/PARAM]
-[/OBJECT]
-[/VAR]
-
 [VAR:Topic]
 [OBJECT:Topic]
 	[PARAM:image]
@@ -62,8 +44,6 @@ version: 4.9
 [/OBJECT]
 [/VAR]
 
-%playlistnav%
-
 
 [TOC(start:2)]
 
@@ -129,7 +109,3 @@ In addition to selecting related Actors, you can also do the following with the
 * Invert your current selection (in other words, deselect everything that is selected and instead select everything that was not selected) by clicking **Invert Selection**.
 * Select lights that affect your Actor by clicking **Select Relevant Lights**.
 
-
-%playlistnav%
-
-
diff --git a/Engine/Documentation/Source/Engine/Actors/StaticMeshActor/UserGuide.INT.udn b/Engine/Documentation/Source/Engine/Actors/StaticMeshActor/UserGuide.INT.udn
--- a/Engine/Documentation/Source/Engine/Actors/StaticMeshActor/UserGuide.INT.udn
+++ b/Engine/Documentation/Source/Engine/Actors/StaticMeshActor/UserGuide.INT.udn
@@ -4,27 +4,6 @@ Crumbs:%ROOT%, Engine, Engine/Actors
 Description:Guide to placing and working with StaticMeshActors for creating world geometry in your environments.
 version: 4.9
 
-[VAR:playlistnav]
-[OBJECT:PlaylistNavigation]
-	[PARAMLITERAL:id]
-		LevelDesign_Beginner
-	[/PARAMLITERAL]
-	[PARAM:previous]
-		[Previous Topic](Engine/Actors/Types "%Engine/Actors/Types:title%")
-	[/PARAM]
-	[PARAM:next]
-		[Next Topic](Engine/Actors/Brushes "%Engine/Actors/Brushes:title%")
-	[/PARAM]
-	[PARAM:home]
-		[](Tracks/LevelDesign/Beginner)
-		%Tracks/LevelDesign/Beginner:menu%
-	[/PARAM]
-[/OBJECT]
-[/VAR]
-
-%playlistnav%
-
-
 [TOC(start:2)]
 
 
@@ -131,6 +110,3 @@ By default, if a StaticMesh has Physics Bodies, either generated in your desired
 (see: [](Engine/Physics/Collision/Reference)), it will have collision and it will be set to **Block All**. See the 
 [](Engine/Physics/Collision) documentation for more information on Collision Channels and adjusting Collision settings.
 
-
-
-%playlistnav%
diff --git a/Engine/Documentation/Source/Engine/Actors/Transform/ActorTransform.INT.udn b/Engine/Documentation/Source/Engine/Actors/Transform/ActorTransform.INT.udn
--- a/Engine/Documentation/Source/Engine/Actors/Transform/ActorTransform.INT.udn
+++ b/Engine/Documentation/Source/Engine/Actors/Transform/ActorTransform.INT.udn
@@ -3,24 +3,6 @@ Title:Transforming Actors
 Crumbs: %ROOT%, Engine, Engine/Actors
 Description:Modifying the location, rotation, and scale of Actors placed in levels.
 
-[VAR:playlistnav]
-[OBJECT:PlaylistNavigation]
-	[PARAMLITERAL:id]
-		LevelDesign_Beginner
-	[/PARAMLITERAL]
-	[PARAM:previous]
-		[Previous Topic](Engine/Actors/Selection "%Engine/Actors/Selection:title%")
-	[/PARAM]
-	[PARAM:next]
-		[Next Topic](Engine/Actors/Mobility "%Engine/Actors/Mobility:title%")
-	[/PARAM]
-	[PARAM:home]
-		[](Tracks/LevelDesign/Beginner)
-		%Tracks/LevelDesign/Beginner:menu%
-	[/PARAM]
-[/OBJECT]
-[/VAR]
-
 [VAR:Topic]
 [OBJECT:Topic]
 	[PARAM:image]
@@ -61,8 +43,6 @@ Description:Modifying the location, rotation, and scale of Actors placed in leve
 [/OBJECT]
 [/VAR]
 
-%playlistnav%
-
 
 [TOC(start:2)]
 
@@ -246,9 +226,4 @@ By repeating this process, we can get nicely positioned meshes. This technique i
 
 
 
-%playlistnav%
-
-
-
-
 
diff --git a/Engine/Documentation/Source/Engine/Actors/Triggers/TriggerActors.INT.udn b/Engine/Documentation/Source/Engine/Actors/Triggers/TriggerActors.INT.udn
--- a/Engine/Documentation/Source/Engine/Actors/Triggers/TriggerActors.INT.udn
+++ b/Engine/Documentation/Source/Engine/Actors/Triggers/TriggerActors.INT.udn
@@ -4,26 +4,6 @@ Crumbs: %ROOT%, Engine, Engine/Actors
 Description:Actor that can be activated and cause events to occur in the level.
 version: 4.9
 
-[VAR:playlistnav]
-[OBJECT:PlaylistNavigation]
-	[PARAMLITERAL:id]
-		LevelDesign_Beginner
-	[/PARAMLITERAL]
-	[PARAM:previous]
-		[Previous Topic](Engine/Actors/Volumes "%Engine/Actors/Volumes:title%")
-	[/PARAM]
-	[PARAM:next]
-		[Next Topic](Engine/Rendering/Materials/IntroductionToMaterials "%Engine/Rendering/Materials/IntroductionToMaterials:title%")
-	[/PARAM]
-	[PARAM:home]
-		[](Tracks/LevelDesign/Beginner)
-		%Tracks/LevelDesign/Beginner:menu%
-	[/PARAM]
-[/OBJECT]
-[/VAR]
-
-%playlistnav%
-
 [TOC(start:2)]
 
 [EXCERPT:Intro]
@@ -69,4 +49,3 @@ overlaps (or runs through) the Trigger:
 
 <!-- Add GIF/slider showing overlap during play and pin firing in graph -->
 
-%playlistnav%
diff --git a/Engine/Documentation/Source/Engine/Actors/Types/CommonActorTypes.INT.udn b/Engine/Documentation/Source/Engine/Actors/Types/CommonActorTypes.INT.udn
--- a/Engine/Documentation/Source/Engine/Actors/Types/CommonActorTypes.INT.udn
+++ b/Engine/Documentation/Source/Engine/Actors/Types/CommonActorTypes.INT.udn
@@ -3,27 +3,6 @@ Title: Common Actor Types
 Crumbs: %ROOT%, Engine, Engine/Actors
 Description:Listing and descriptions of common types of Actors used in levels.
 
-[VAR:playlistnav]
-[OBJECT:PlaylistNavigation]
-	[PARAMLITERAL:id]
-		LevelDesign_Beginner
-	[/PARAMLITERAL]
-	[PARAM:previous]
-		[Previous Topic](Engine/Actors/Grouping "%Engine/Actors/Grouping:title%")
-	[/PARAM]
-	[PARAM:next]
-		[Next Topic](Engine/Actors/StaticMeshActor "%Engine/Actors/StaticMeshActor:title%")
-	[/PARAM]
-	[PARAM:home]
-		[](Tracks/LevelDesign/Beginner)
-		%Tracks/LevelDesign/Beginner:menu%
-	[/PARAM]
-[/OBJECT]
-[/VAR]
-
-%playlistnav%
-
-
 Below is a quick rundown of some of the Actors you will most commonly use or come across while developing games with 
 Unreal Engine to familiarize you with what they are and why you would use them.
 
@@ -57,6 +36,3 @@ Unreal Engine to familiarize you with what they are and why you would use them.
 [/REGION]
 [/EXCERPT:Trigger]
 -->
-
-
-%playlistnav%
diff --git a/Engine/Documentation/Source/Engine/Actors/Volumes/Volumes.INT.udn b/Engine/Documentation/Source/Engine/Actors/Volumes/Volumes.INT.udn
--- a/Engine/Documentation/Source/Engine/Actors/Volumes/Volumes.INT.udn
+++ b/Engine/Documentation/Source/Engine/Actors/Volumes/Volumes.INT.udn
@@ -3,27 +3,6 @@ Title:Volumes User Guide
 Crumbs:%ROOT%, Engine, Engine/Actors
 Description:Guide to placing and working with Volumes to augment the behavior and setup of your levels.
 
-[VAR:playlistnav]
-[OBJECT:PlaylistNavigation]
-	[PARAMLITERAL:id]
-		LevelDesign_Beginner
-	[/PARAMLITERAL]
-	[PARAM:previous]
-		[Previous Topic](Engine/Actors/Brushes "%Engine/Actors/Brushes:title%")
-	[/PARAM]
-	[PARAM:next]
-		[Next Topic](Engine/Actors/Triggers "%Engine/Actors/Triggers:title%")
-	[/PARAM]
-	[PARAM:home]
-		[](Tracks/LevelDesign/Beginner)
-		%Tracks/LevelDesign/Beginner:menu%
-	[/PARAM]
-[/OBJECT]
-[/VAR]
-
-
-%playlistnav%
-
 
 [TOC(start:2 end:2)]
 
@@ -183,7 +162,7 @@ For more on Level Streaming Volumes, see [](Engine/LevelStreaming/StreamingVolum
 
 ### Cull Distance Volume
 
-**Cull Distance Volumes** are optimization tools that allow for objects at or below a certain size to be culled (or not drawn to the screen) based on that object's distance from the camera. This can help optimize your scene by not drawing objects of a given size. Size is calculated by the bounding box along its longest dimension. 
+**Cull Distance Volumes** are optimization tools that allow for objects to be culled (not drawn to the screen) based on that object's distance from the camera and its size. This can help optimize your scene by not drawing objects when they are small enough to be considered unimportant. Size is calculated by the bounding box along its longest dimension, and the cull distance chosen is the one closest to that size.
 
 Cull Distance Volume setup is dependent on the **Cull Distances** property, shown in the **Details** panel below.
 
@@ -191,9 +170,9 @@ Cull Distance Volume setup is dependent on the **Cull Distances** property, show
 
 In the image above, the properties define the following behavior:
 
-* Objects within the volume that have a size less than 50 units will be culled (disappear) when they are 500 units away from the camera or farther.
-* Objects within the volume that have a size less than 120 units will be culled (disappear) when they are 1000 units away from the camera or farther.
-* Objects within the volume that have a size less than 300 units will never be culled, since 0 in this case is considered to be infinity, meaning that the camera can never get far enough away.
+* Objects within the volume that have a size closest to 50 units (85 units or less) will be culled (disappear) when they are 500 units away from the camera or farther.
+* Objects within the volume that have a size closest to 120 units (85 to 210 units) will be culled (disappear) when they are 1000 units away from the camera or farther.
+* Objects within the volume that have a size closest to 300 units (210 units or more) will never be culled, since 0 in this case is considered to be infinity, meaning that the camera can never get far enough away.
 
 Setup starts with adding a new entry into the Cull Distances array, which is done by clicking the ![](button_Plus.png) button. Next, simply fill in the size of a given object and the desired distance at which you want objects of that size or smaller to be culled. These properties do not have to be created in a specific order. 
 
diff --git a/Engine/Documentation/Source/Engine/Animation/AnimBlueprints/VimBlueprints.INT.udn b/Engine/Documentation/Source/Engine/Animation/AnimBlueprints/VimBlueprints.INT.udn
--- a/Engine/Documentation/Source/Engine/Animation/AnimBlueprints/VimBlueprints.INT.udn
+++ b/Engine/Documentation/Source/Engine/Animation/AnimBlueprints/VimBlueprints.INT.udn
@@ -23,17 +23,17 @@ tags:State Machine
 
 
 [EXCERPT:Intro]
-An **Animation Blueprint** is essentially a specialized [Blueprint](Engine/Blueprints) whose graphs control the animation of a [Skeletal Mesh.](Engine/Content/Types/SkeletalMeshes) 
-Graphs are edited inside the **Animation Blueprint Editor** where you can perform blending of animations, directly control the bones of a Skeleton, or set up logic that will ultimately define the final animation pose for a Skeletal Mesh to use each frame. 
+An **Animation Blueprint** is a specialized [Blueprint](Engine/Blueprints) that controls the animation of a [Skeletal Mesh.](Engine/Content/Types/SkeletalMeshes) 
+Graphs are edited inside of the **Animation Blueprint Editor**, where you can perform animation blending, directly control the bones of a Skeleton, or setup logic that will ultimately define the final animation pose for a Skeletal Mesh to use per frame. 
 [/EXCERPT:Intro]
 
 [REGION:note]
 An example of an **Animation Blueprint** is demonstrated on the [Animation Content Examples](Resources\ContentExamples\Animation) page under section 1.2.
 [/REGION]
 
-## Creating & Navigating Animation Blueprints
+## Creating and Navigating Animation Blueprints
 
-Before jumping into how Animation Blueprints work, it's best to understand how to create them and start to familiarize yourself with the Animation Blueprint Interface. 
+Before jumping into how Animation Blueprints work, it's best to understand how to create them. Additionally, you should take some time to familiarize yourself with the Animation Blueprint Interface. 
 
 [REGION:topics half]
 %Engine\Animation\AnimBlueprints\Creation:TopicCompact%
@@ -42,13 +42,13 @@ Before jumping into how Animation Blueprints work, it's best to understand how t
 
 ## Working with Animation Blueprints
 
-Once you have created your Animation Blueprint and can navigate the UI, it's time to consider which aspect you want to work on first.
+Now that you've created your Animation Blueprint (also, having learned the Animation Blueprint Interface), it's time to consider what you want to work on first.
 
 [EXCERPT:Graphs]
-There are two main components in an Animation Blueprint that work in conjunction to create the final animation 
-for each frame. The **EventGraph** can perform updates to values which can then be used in the **AnimGraph** 
-to drive State Machines, BlendSpaces, or other nodes allowing blending between multiple AnimationSequences or poses 
-that can fire off notifications to other systems to enable dynamic animation-driven effects to occur.
+There are two main components in an Animation Blueprint that work together to create the final animation 
+for each frame. Specifically, the **EventGraph** performs updates to values that are then used in the **AnimGraph** 
+to drive State Machines, BlendSpaces or other nodes; allowing blending between multiple AnimationSequences or poses 
+that can fire off notifications to other systems, enabling dynamic animation-driven effects.
 [/EXCERPT:Graphs]
 
 [REGION:topics half]
@@ -59,13 +59,13 @@ that can fire off notifications to other systems to enable dynamic animation-dri
 
 ## Working with Sub Animation Blueprint Instances
 
-When creating Animation Blueprints for various characters, it is not uncommon to share animation logic between those characters. 
-For example, suppose you have physical animation or [AnimDynamics](Engine\Animation\AnimHowTo\AnimDynamics) applied to one character and you want to apply it to another. 
-With the **Sub Anim Instance** node, you can create a reference to a separate Animation Blueprint to access and include all of its logic within another Animation Blueprint (sort of like a function call). 
+When creating Animation Blueprints for various characters, it's not uncommon to share animation logic between those characters. 
+For example, suppose you have physical animation, or [AnimDynamics](Engine\Animation\AnimHowTo\AnimDynamics), applied to one character and you want to apply it to another. 
+With the **Sub Anim Instance** node, you can create a reference to a separate Animation Blueprint to access and include all of its logic within another Animation Blueprint. 
 
-This can be used to break up large Animation Blueprints into separate, smaller Animation Blueprints, that can be included within other Animation Blueprints. 
-For example, you could have one Animation Blueprint for purely "Locomotion" and a Sub Anim Instance that handles any physics-based animation. 
-Member variables of a Sub Animation Blueprint can also be exposed as input pins inside their parent Animation Blueprint and **Sub-Graph Input** nodes can be used to expose **In Poses** on the Sub Anim Instance (shown below).
+This can be used to break up large Animation Blueprints into separate, smaller blueprints that can then be included in other Animation Blueprints. 
+For example, you could have one Animation Blueprint strictly for "Locomotion", and a Sub Anim Instance that handles physics-based animation. 
+Member variables of a Sub Animation Blueprint can also be exposed as input pins inside of their parent Animation Blueprint, enabling **Sub-Graph Input** nodes to expose **In Pose** to the Sub Anim Instance (see the example shown below).
 
 ![](SubAnimInstance.png)
 
@@ -80,14 +80,31 @@ See [Using Sub Anim Instances](https://docs.unrealengine.com/latest/INT/Engine/A
 
 ## Working with Child Animation Blueprints
 
-Also when creating and setting up your characters and Animation Blueprints, there may be instances where you have multiple characters that are similar but have entirely different animations for different states. Rather than create Animation Blueprints for each of those characters, you may be able to use a **Child Animation Blueprint** of a pre-existing Animation Blueprint and override the animation assets that you want to replace. The Child Animation Blueprint will inherit everything from the parent, but it will also allow you to specify which animations to override via the **Asset Override Editor**. 
+When creating and setting up characters and Animation Blueprints, there may be situations where you have multiple characters that are similar, but have entirely different animations for different states. 
+Rather than creating Animation Blueprints for each of those characters, you may be able to use a **Child Animation Blueprint** of a pre-existing Animation Blueprint and override the animation assets that you want to replace. 
+The Child Animation Blueprint will inherit everything from the parent, enabling you to specify which animations should be overridden via the **Asset Override Editor**. 
 
-For example, suppose you have two humanoid type characters where one is a human and another is a werewolf. Both characters can walk, run, jump, crawl, etc. However, their motions may be entirely different. The logic of moving between those motions is the same, but the motions themselves are different. It's perfectly fine to create two different Animation Blueprints but to save some time. you can create a Child Animation Blueprint and selectively choose which animations to override. 
+For example, suppose you have two humanoid characters, where one is a human and the other is a werewolf. Both characters can walk, run, jump, and crawl. 
+However, their motions may be entirely different. The logic of moving between those motions is the same, but the motions themselves are different. 
+It's perfectly fine to create two different Animation Blueprints, but to save some time, you can create a Child Animation Blueprint and selectively choose which animations to override. 
 
 [REGION:note]
-See [Animation Blueprint Override](Engine/Animation/AnimHowTo/AnimBPOverride) for an example walk-through of creating and overriding an Animation Blueprint with a Child Animation Blueprint.
+See [Animation Blueprint Override](Engine/Animation/AnimHowTo/AnimBPOverride) for an example of creating and overriding an Animation Blueprint with a Child Animation Blueprint.
 [/REGION]
 
+## Assigning Post Process Animation Blueprints
+
+Animation Blueprints can also be assigned directly to a Skeletal Mesh through the [](Engine/Animation/Persona/MeshDetails) panel as a **Post Process Anim Blueprint**. 
+This allows you to assign an Animation Blueprint to a Skeletal Mesh that will always run after any Animation Blueprint was assigned as the Main Anim Instance for a Skeletal Mesh Component. 
+By assigning a Post Process Anim Blueprint to a Skeletal Mesh, you can create AnimDynamics, Skeletal Controls, or other Anim Blueprint logic without having to copy that logic to other Animation Blueprints.
+
+![](ExamplePostProcess.png)
+
+[REGION:note]
+Please see the [Post Process Anim Blueprint](https://docs.unrealengine.com/latest/INT/Engine/Animation/Persona/MeshDetails/index.html#skeletalmesh) section of the Mesh Details page for more information.
+[/REGION]
+
+
 
 
 
diff --git a/Engine/Documentation/Source/Engine/Animation/AnimHowTo/AdditiveAnimations/01_BlendSetup/BlendSetup.INT.udn b/Engine/Documentation/Source/Engine/Animation/AnimHowTo/AdditiveAnimations/01_BlendSetup/BlendSetup.INT.udn
--- a/Engine/Documentation/Source/Engine/Animation/AnimHowTo/AdditiveAnimations/01_BlendSetup/BlendSetup.INT.udn
+++ b/Engine/Documentation/Source/Engine/Animation/AnimHowTo/AdditiveAnimations/01_BlendSetup/BlendSetup.INT.udn
@@ -4,9 +4,8 @@ Crumbs: %ROOT%, Engine, Engine\Animation, Engine\Animation\AnimHowTo
 Description:In this step, we update our player character and create an AnimMontage for the animation we want to layer in on top of our movement animations.
 version: 4.14
 skilllevel: Advanced
-type:howto
+type:how-to
 tags:Animation
-tags:How-To
 
 [VAR:Steps]
 [OBJECT:Navigation]
diff --git a/Engine/Documentation/Source/Engine/Animation/AnimHowTo/AdditiveAnimations/02_CharacterBlueprintSetup/CharacterBPSetup.INT.udn b/Engine/Documentation/Source/Engine/Animation/AnimHowTo/AdditiveAnimations/02_CharacterBlueprintSetup/CharacterBPSetup.INT.udn
--- a/Engine/Documentation/Source/Engine/Animation/AnimHowTo/AdditiveAnimations/02_CharacterBlueprintSetup/CharacterBPSetup.INT.udn
+++ b/Engine/Documentation/Source/Engine/Animation/AnimHowTo/AdditiveAnimations/02_CharacterBlueprintSetup/CharacterBPSetup.INT.udn
@@ -4,9 +4,8 @@ Crumbs: %ROOT%, Engine, Engine\Animation, Engine\Animation\AnimHowTo
 Description:In this step, we create some script inside our Character Blueprint to allow us to simulate firing a weapon. 
 version: 4.14
 skilllevel: Advanced
-type:howto
+type:how-to
 tags:Animation
-tags:How-To
 
 [VAR:Steps]
 [OBJECT:Navigation]
diff --git a/Engine/Documentation/Source/Engine/Animation/AnimHowTo/BlendSpace/BlendspaceHowTo.INT.udn b/Engine/Documentation/Source/Engine/Animation/AnimHowTo/BlendSpace/BlendspaceHowTo.INT.udn
--- a/Engine/Documentation/Source/Engine/Animation/AnimHowTo/BlendSpace/BlendspaceHowTo.INT.udn
+++ b/Engine/Documentation/Source/Engine/Animation/AnimHowTo/BlendSpace/BlendspaceHowTo.INT.udn
@@ -1,12 +1,11 @@
 Availability:Public
-Title:Creating a Blend Space
+Title:Locomotion Based Blending
 Crumbs: %ROOT%, Engine, Engine\Animation, Engine\Animation\AnimHowTo
-Description:Example of a Blend Space used to toggle a character from walking/jogging animations smoothly. 
-version: 4.14
+Description:Example of a Blend Space used to blend a character's animations from walking and jogging. 
+version: 4.15
 skilllevel: Intermediate
-type:howto
+type:how-to
 tags:Animation
-tags:How-To
 tags:Blend Space
 Related: Engine/Animation/Blendspaces
 Related: Engine/Animation/AimOffset
@@ -17,13 +16,31 @@ Related: Engine/Animation/Skeleton
 Simple blending between two animations based on a single input can be performed using one of the standard [Blend Nodes](Engine/Animation/NodeReference/Blend) available in Animation Blueprints. 
 Blend Spaces provide a means of doing more complex blending between multiple animations based on multiple values (currently limited to two).
 
-In this How-to, we use a Blend Space to blend between walking and jogging animations based on the character's movement speed and movement direction. 
+In this How-to, we use a Blend Space to blend between walking and jogging animations based on the character's movement speed and movement direction similar to below. 
 
-
-[INCLUDE:Engine\Animation\AnimHowTo\AdditiveAnimations\01_BlendSetup#intro]
+[OBJECT:EmbeddedVideo]
+[PARAMLITERAL:width]
+640
+[/PARAMLITERAL]
+[PARAMLITERAL:height]
+360
+[/PARAMLITERAL]
+[PARAMLITERAL:videoid]
+jVtczzmNwQ4
+[/PARAMLITERAL]
+[/OBJECT]
 
 ## Steps
 
+[REGION:note]
+For this guide, we are using the **Blueprint Third Person Template** project and have added the **Animation Starter Pack** to the project. 
+
+![](AnimationAssetPack.png)(w:800)
+
+You can download the Animation Starter pack for free through the **Marketplace** from the Epic Launcher. 
+
+[/REGION]
+
 1.  In your project in the **Content/ThirdPersonBP/Blueprints** folder, open the **ThirdPersonGameMode** and set the **Default Pawn Class** to **Ue4ASP_Character**. 
 
 	![](BlendSpace1.png)
@@ -36,46 +53,61 @@ In this How-to, we use a Blend Space to blend between walking and jogging animat
 
 	We need to delete this character from the level in order to spawn in as the character from the previous step. 
 
-1.	In the **Content/AnimStarterPack/UE4_Mannequin/Mesh** folder, **Right-click** on the **UE4_Mannequin** and under **Create**, select **Blend Space**. 
+1.	In the **Content/AnimStarterPack/UE4_Mannequin/Mesh** folder, **Right-click** on the **UE4_Mannequin** and under **Create**, select **Blend Space** and give it a name.
 
 	![](BlendSpace3.png)
 
-	Blend Space assets need a Skeleton asset to reference, here we are using a shortcut to create a Blend Space from this Skeleton asset. 
+	Blend Space assets need a Skeleton asset to reference, here we are using a shortcut to create a Blend Space from this Skeleton asset which we called **MyBlendSpace**.
 
-1.	In the **Parameters** section, apply the values below then press the **Apply Parameter Changes** button. 
+1.	Inside the Blend Space, in the **Asset Details** panel under **Axis Settings**, add the following values: 
 
 	![](BlendSpace4.png)
 
-	* Set the **X Axis Label** to **Direction**
-	* Set the **X Axis Range** to **-180 - 180**
-	* Set the **Y Axis Label** to **Speed**
-	* Set the **Y Axis Range** to **150 - 260**
+	The Horizontal Axis will determine the **Direction** we are moving in degrees: 
+
+	|Horizontal Axis Settings | Value|
+	|---|---|
+	|**Name**| Direction |
+	|**Minimum Axis Value**| -180 |
+	|**Maximum Axis Value**| 180 |
+	|**Number of Grid Divisions**| 4 |
+	|**Interpolation Time**| 0 |
+	|**Interpolation Type**| Averaged Interpolation |
 
-	This Blend Space using the parameters **Direction** and **Speed** to determine which pose we should be using or blending to or from. 
+	The Vertical Axis will determine the **Speed** in which our character is moving: 
+
+	|Vertical Axis Settings | Value|
+	|---|---|
+	|**Name**| Speed |
+	|**Minimum Axis Value**| 0 |
+	|**Maximum Axis Value**| 250 |
+	|**Number of Grid Divisions**| 4 |
+	|**Interpolation Time**| 0 |
+	|**Interpolation Type**| Averaged Interpolation |
 
 1.	In the **Asset Browser**, drag the **Walk_Fwd_Rifle_Ironsights** animation to the center-bottom position on the grid. 
 
 	![](BlendSpace5.png)
 
-	Placing this animation here indicates that when our Speed is 150 and are direction is not left or right, we should be walking forward with our rifle up at ironsights level. 
+	This will be are starting position when we enter this Blend Space. 
 
 1.	Drag the **Walk_Lt_Rifle_Ironsights** animation to location **1** and **Walk_Rt_Rifle_Ironsights** animation to location **2** below. 
 
 	![](BlendSpace6.png)
 
-	Here we place the animations to use when moving left or right and our speed is 150. 
+	Here we place the animations to use when moving left or right. 
 
 1.	Drag the **Walk_Bwd_Rifle_Ironsights** animation to locations **1** and **2** below.
 
 	![](BlendSpace7.png)
 
-	Here we place the animation to use when moving backward and our speed it 150. 
+	Here we place the animation to use when moving backward. 
 
 1.	From the **Asset Browser**, drag the **Jog_Fwd_Rifle** animation to the top-center position on the grid. 
 
 	![](BlendSpace8.png)
 
-	This allows us to use a different animation when our speed is at max 270 and we are not moving left or right, we are jogging forward. 
+	This allows us to use a different animation when our speed is at max 250 and we are not moving left or right, we are jogging forward. 
 
 1.	Drag the **Jog_Lt_Rifle** animation to location **1** and **Jog_Rt_Rifle** animation to location **2** below. 
 
@@ -95,7 +127,7 @@ In this How-to, we use a Blend Space to blend between walking and jogging animat
 
 	This [State Machine](Engine/Animation/StateMachines) already has a Jog State and in it, a Blend Space for different Jogging motions. 
 
-1.	From the **Asset Browser**, drag the **Walk_Jog_BS** into the graph and replace the **BS_Jog** Blend Space as shown below. 
+1.	From the **Asset Browser**, drag the **MyBlendSpace** asset into the graph and replace the **BS_Jog** Blend Space as shown below. 
 
 	![](BlendSpace13.png)
 
@@ -103,11 +135,11 @@ In this How-to, we use a Blend Space to blend between walking and jogging animat
 
 1.	In the **Content/AnimStarterPack** folder, open the **Ue4ASP_Character**. 
 
-1.	Locate the section marked **Crouch** in the graph and replace the **InputAction Crouch** node with a **C** Key Event. 
+1.	Locate the section marked **Crouch** in the graph and delete it to remove the warning error. 
 
 	![](BlendSpace14.png)
 
-	By default, the Animation Starter Pack has bindings created for crouch which we will replace with a normal **C** press (or whatever key you want to crouch with). 
+	By default, the Animation Starter Pack has bindings created for crouch which our project does not. For this example, we will remove the ability to crouch walk. 
 
 1.	**Right-click** in the graph and add a **Left Shift** Key Event, then drag the **CharacterMovement** into the graph from the **Components** tab. 
 
@@ -115,23 +147,17 @@ In this How-to, we use a Blend Space to blend between walking and jogging animat
 
 	When we press and hold the Left Shift button, that will toggle us from walking into a jog. 
 
-1.	Hold the **B** key and **Left-click** in the graph and create two **Branch** nodes, then drag in the **Crouch Button Down** variable and connect as shown below. 
-
-	![](BlendSpace16.png)
-
-	We need to check and make sure that we are not crouching while holding shift (as for this example, sprinting can only occur when not crouched). 
-
-1.	Drag off the **CharacterMovement** node and **Set Max Walk Speed** for both **Branch** nodes with the values **260** and **150** respectively. 
+1.	Drag off the **CharacterMovement** node and **Set Max Walk Speed** for **Pressed** and **Released** with the values **250** and **100** respectively. 
 
 	![](BlendSpace17.png)
 
 	We are changing our movement speed based on if Left Shift is held and will alter our character's movement speed in game. 
 
-1.	**Right-click** in the graph and add an **Event Begin Play** node and connect a **CharacterMovement** component to **Set Max Walk Speed** to **150**. 
+1.	Click the **CharacterMovement** component, then in the **Details** panel set the **Max Walk Speed** to **100**. 
 
 	![](BlendSpace20.png)
 
-	To start the game, we need to update our character's default movement speed so they start out walking. 
+	We need to update our character's default movement speed so they start out walking. 
 
 1.	Click the **CapsuleComponent** in the **Components** tab and in the **Details** panel, check the **Hidden in Game** option. 
 
@@ -139,12 +165,6 @@ In this How-to, we use a Blend Space to blend between walking and jogging animat
 
 	This will turn off the collision capsule debug display in game. 
 
-1.	Click the **CharacterMovement** component, then in the **Details** panel set the **Max Walk Speed** to **260**. 
-
-	![](BlendSpace19.png)
-
-	Here we are capping the maximum movement speed to reflect what is in our Blend Space. 
-
 1.	Click **Compile** then click the **Play** in editor button from the Toolbar. 
 
 
@@ -158,13 +178,13 @@ In this How-to, we use a Blend Space to blend between walking and jogging animat
 360
 [/PARAMLITERAL]
 [PARAMLITERAL:videoid]
-vn0_YKcu1RU
+jVtczzmNwQ4
 [/PARAMLITERAL]
 [/OBJECT]
 
 When playing in the Editor, using **WASD** to move around, the character will now walk by default and blend between walking animations when changing directions. 
 While standing, pressing and holding **Left Shift** while using WASD will cause the character to blend from the walk/aim animation to a lowered/aim jog animation, essentially creating a sprint for the character. 
-You can also press and hold **C**, to crouch walk using the Blend Space that was created as part of the Animation Starter Pack project, blending between crouch walking animations. 
+
 
 
 	
@@ -183,3 +203,4 @@ You can also press and hold **C**, to crouch walk using the Blend Space that was
 
 
 
+
diff --git a/Engine/Documentation/Source/Engine/Animation/AnimMontage/AnimMontage.INT.udn b/Engine/Documentation/Source/Engine/Animation/AnimMontage/AnimMontage.INT.udn
--- a/Engine/Documentation/Source/Engine/Animation/AnimMontage/AnimMontage.INT.udn
+++ b/Engine/Documentation/Source/Engine/Animation/AnimMontage/AnimMontage.INT.udn
@@ -1,6 +1,6 @@
 Availability:Public
 Title:AnimMontage
-Crumbs: %ROOT%, Engine, Engine/Animation
+Crumbs:
 Description:Info detailing what AnimMontages are and how they are used in UE4's animation system.
 Related: Engine/Content/FBX/Animations
 Related: Engine/Content/Types/SkeletalMeshes
@@ -8,16 +8,22 @@ Related: Engine/Animation/PhysicallyDrivenAnimation
 Related: Engine/Animation/StateMachines
 Related: Gameplay/AI
 Related: Engine/Content/Tools/MayaRiggingTool
+Related: Engine/Animation/RootMotion
 Related: Gameplay/Input
 version: 4.14
+Parent:Engine/Animation
+Order:
+tags:Animation Montage
+tags:Animation
+type:overview
 
-[TOC(start:2 end:2)]
+[TOC(start:2 end:3)]
 
 [EXCERPT:Intro]
-**Animation Montages** (or just Montages for short) are a multipurpose tool that allows for a wide variety of animation effects, primarily related to exposing animation controls within code or [Blueprint Visual Scripting](Engine/Blueprints). It can also be used to create a wide variety of animation effects including intelligent loops of animation, logic-based animation switching and much more.
+**Animation Montages** (or, "Montages", for short) are a multipurpose tool that allows for a wide variety of animation effects, primarily related to exposing animation controls within code, or [Blueprint Visual Scripting](Engine/Blueprints). It can also be used to create a wide variety of animation effects, including intelligent loops of animation, logic-based animation switching, and much more.
 [/EXCERPT:Intro]
 
-Montages are animation assets that can be created and exist within the **Content Browser**. These assets can then be plugged directly into the AnimGraph in run-time and you can modify any state of it. For example, you can jump to different sections or you can re-link different sections. This is mostly for code driven animations or one-off animations such as melee attacks, allowing you to control triggers, stop when you want, or change between states (looping or non-looping). 
+Montages are animation assets that can be created, and exist, within the **Content Browser**. These assets can then be plugged directly into the AnimGraph in run-time, and you can modify any state of it. For example, you can jump to different sections, or you can re-link different sections. This is mostly for code driven animations (or one-off animations) such as melee attacks, allowing you to control triggers, stop when you want, or change between states (looping or non-looping). 
 
 Montages also provide a way to replicate [](Engine\Animation\RootMotion\) across networked games if your animation has Root Motion enabled. 
 
@@ -25,32 +31,32 @@ Montages also provide a way to replicate [](Engine\Animation\RootMotion\) across
 
 ![](Montage_Screen2.png)
 
-The image above is a melee attack with 3 sections, labeled _Start_, _Loop_, and _End_; the divisions between the sections are marked in yellow. When the player presses the attack button, the animation is played by the game, and it begins at the _Start_ section. The looping middle section, called _Loop_, immediately follows the _Start_ section and will be programmed to repeat indefinitely. If the player releases the attack button, the animation will stop looping and eventually reach the _End_ section and transition out.
+The image above is a melee attack with 3 sections, labeled _Start_, _Loop_, and _End_; the divisions between the sections are marked in yellow. When the player presses the attack button, the animation is played by the game, and it begins at the _Start_ section. The looping middle section, called _Loop_, immediately follows the _Start_ section and will be programmed to repeat indefinitely. If the player releases the attack button, the animation will stop looping and eventually reach the _End_ section to transition out of the animation.
 
 Some additional uses for Montages include:
 * The ability to play an animation from within an AnimBlueprint's **EventGraph**.
 * Chaining together a complex sequence of animations that you want to think of as a single animation.
-* Looping only a specific portion of animation or animations based on code or Blueprint script.
-* Handling event-based switching of multiple animations based on code or Blueprint script.
-* The ability to assign complex animation sequences to named slots that can be switched between in code or Blueprints.
-* Precise switching between various AnimSequences based on code or Blueprint script.
+* Looping only a specific portion of animation (or animations) based on code, or a Blueprint script.
+* Handling event-based switching of multiple animations based on code, or a Blueprint script.
+* The ability to assign complex animation sequences to named slots that can be switched between in code, or with Blueprints.
+* Precise switching between various AnimSequences based on code, or with a Blueprint script.
 
 [REGION:note]
-A Montage applied to a character can also be seen on the [Animation Content Examples](Resources\ContentExamples\Animation) page under section 1.5.
+A **Montage** applied to a character can also be seen on the [Animation Content Examples](Resources\ContentExamples\Animation) page.
 [/REGION]
 
 ## Montage Properties
 
-The following is a breakdown of the Montage asset properties. These are available in **Persona** when looking at a Montage, and can also be accessed by **Right-clicking** on a Montage within the **Content Browser** and choosing Properties from the context menu.
+The following is a breakdown of the Montage asset properties. These are available in **Persona** when looking at a Montage, and can also be accessed by right-clicking on a Montage within the **Content Browser**, choosing **Properties** from the context menu.
 
 ![](MontageProperties.png)
 
 |Montage Properties||
 |---|---|
 |[REGION:tablesection]Montage[/REGION]||
-|**Blend In Time**| The amount of time at the beginning of Montage playback during which the character will blend in from its current pose.|
-|**Blend Out Time**|The amount of time at the end of Montage playback during which the character will blend back to its original pose.|
-|**Blend Out Trigger Time**| Time from the Sequence End to trigger blend out. < 0 uses BlendOutTime, so BlendOut finishes as the Montage Ends. > = 0 uses SequenceEnd-BlendOutTriggerTime to trigger blend out. |
+|**Blend In Time**| The amount of time at the beginning of Montage playback, during which the character will blend in from its current pose.|
+|**Blend Out Time**|The amount of time at the end of Montage playback, during which the character will blend back to its original pose.|
+|**Blend Out Trigger Time**| Time from the Sequence End to trigger blend out. < 0 uses BlendOutTime, so BlendOut finishes as the Montage Ends. >= 0 uses SequenceEnd-BlendOutTriggerTime to trigger blend out. |
 |[REGION:tablesection]Additive Settings[/REGION]||
 |**Preview Base Pose**|Sets a base preview pose used for additive Blend Spaces.|
 |[REGION:tablesection]Animation[/REGION]||
@@ -65,7 +71,7 @@ The following is a breakdown of the Montage asset properties. These are availabl
 
 ## Montage UI
 
-When looking at a Montage in Persona, it is useful to know what each area is and what it does:
+When looking at a Montage in Persona, it is useful to know what each area is, and what it does:
 
 [REGION:raw]
 ![](MontageUI.png)
@@ -89,24 +95,24 @@ The Montage Area breaks down as follows:
 1. **Slot Track** - Shows the current [Slot](#Slots), along with the Slot name on the right. This Slot may be filled with as many animations as you desire; they will be played back in order. Notice that multiple animations will have alternating positions in the Slot Track - top, then bottom, then top again, and so on. This is to help you differentiate between different animations.
 
 [REGION:note]
-You may have as many Slot Tracks as you like within a single Montage, each with their own name and containing their own unique animations. However, you may only have one Section Track for each Montage.
+You may have as many **Slot Tracks** as you like within a single **Montage**, each with their own name, containing their own unique animations. However, you may only have one **Section Track** for each Montage.
 [/REGION]
 
 ### Sections
 
-Montage Sections provide a way to break a Slot up into multiple portions of animation. Each Section has a name and location in the Slot's timeline. Using the name, you can either jump directly to a particular Section or queue it to play next, when the current segment is complete. In Blueprint, you can query the current Section, jump to a Section, or set the next Section that will play.
+Montage Sections provide a way to break a Slot into multiple portions of animation. Each Section has a name and location in the Slot's timeline. Using the name, you can either jump directly to a particular Section, or queue it to play next, when the current segment is complete. In Blueprint, you can query the current Section, jump to a Section, or set the next Section that will play.
 
-It may help to think of Sections like songs on a music playlist, with Slots being the albums. Just like with many modern media players, you can choose which song will play next when the current one finishes, or just jump to the one you want to hear right now.
+It may help to think of Sections like songs on a music playlist, with Slots being the albums. Just like with many modern media players, you can choose which song will play next, when the current one finishes, or just jump to the one you want to hear right now.
 
-Sections are created by **Right-clicking** on the Section track and choosing **New Montage Section**. 
+Sections are created by right-clicking on the Section track and choosing **New Montage Section**. 
 
 
 ### Slots
 
 [EXCERPT:Slots]
-Within a Montage, a **Slot** is just a single track that can hold any number of animations. you can pick any name for a Slot and then use that name to blend to those specific animations. A great example is having a character with a weapon reload animation. You may have different versions of the reload for when the player is standing, crouching, and lying prone. So long as all 3 of the animations used the same timing, you could place each one within a separate Slot in your Montage; the Slots could be named Standing, Crouching, and Prone. In your Animation Blueprint's **AnimGraph**, you can use the Slot node to determine which one you want to play based on your character's current state. When they are standing, you can use the result of the animation in the Standing Slot. When they are prone, you can see the result of the Prone Slot.
+Within a Montage, a **Slot** is just a single track that can hold any number of animations. you can pick any name for a Slot and then use that name to blend to those specific animations. A great example is having a character with a weapon reload animation. You may have different versions of the reload for when the player is standing, crouching, and lying prone. As long as all 3 of the animations use the same timing, you could place each one within a separate Slot in your Montage; the Slots could be named Standing, Crouching, and Prone. In your Animation Blueprint's **AnimGraph**, you can use the Slot node to determine which one you want to play, based on your character's current state. When they are standing, you can use the result of the animation in the Standing Slot. When they are prone, you can see the result of the Prone Slot.
 
-It is an important point to keep in mind that although much of your Montage control will take place in the Animation Blueprint's **Event Graph**, Slots are actually handled within the **Anim Graph**. This is done by way of the Slot node, which simply takes in the name for a Slot. By positioning this node at a strategic point along your AnimGraph's execution, you can have multiple Montages that utilize the same Slot name.
+It is an important point to keep in mind that although much of your Montage control will take place in the Animation Blueprint's **Event Graph**, Slots are actually handled within the **Anim Graph**. This is done by way of the Slot node, which takes in the name for a Slot. By positioning this node at a strategic point along your AnimGraph's execution, you can have multiple Montages that utilize the same Slot name.
 [/EXCERPT:Slots]
 
 <!---------------- Branch moved to Notify page.
@@ -134,8 +140,8 @@ The Sections area is where you can establish relationships between the Sections
 
 
 1. **Create Default and Clear Buttons** - **Create Default** creates default associations between all Sections, stringing them together one after the other. **Clear** wipes out all associations.
-1. **Section Buttons** - In this area, you will see a button for each of the Sections you define in the Montage area. By choosing an existing Section and then clicking the one of these buttons, you associate the Section for that button with the selected track. For instance, in the image above, we have associated Swing2 with Swing1. The order actually goes Swing1, Swing2, and then Swing1 again, which causes a loop. See the [Looping](#Looping) section below for details.
-1. **Section Association Tracks** - This is where you can visualize and preview the relationships between animation Sections. By clicking on the **Preview** buttons, you can see the result of each individual track, or click the **Preview All Sections** button and see all Sections play back one after the other.
+1. **Section Buttons** - In this area, you will see a button for each of the Sections you define in the Montage area. By choosing an existing Section, then clicking one of these buttons, you associate the Section for that button with the selected track. For instance, in the image above, we have associated Swing2 with Swing1. The order actually goes Swing1, Swing2, and then Swing1 again, which causes a loop. See the [Looping](#Looping) section below for details.
+1. **Section Association Tracks** - This is where you can visualize and preview the relationships between animation Sections. By clicking on the **Preview** buttons, you can see the result of each individual track, or click the **Preview All Sections** button, and see all Sections playing back, one after the other.
 
 ### Looping
 
@@ -147,20 +153,20 @@ The Element Timing area takes information from the Montage area and the Notifies
 
 ![](ElementTiming1.png)
 
-Each node in the track is given a number denoting the trigger order of that object in the entire montage, and is colored to match what they describe. The default colors are:
+Each node in the track is given a number denoting the trigger order of that object in the entire montage. It is colored to match what they describe. The default colors are:
 * Red - Notifies, Notify states and notify state ending markers
-* Light Blue - Branching point notifies
+* Light Blue - Branching points
 * Green - Montage sections
 
 [REGION:note]
-These colors can be edited in the **Animation Editor** preferences, which can be found in the **Editor Preferences**.
+These colors can be edited in the **Animation Editor** preferences, which can be found in **Editor Preferences**.
 [/REGION]
 
 Hovering over any of the nodes in the track will cause a tooltip to appear that explains what that node represents.
 
 ![](ElementTiming2.png)
 
-There is dropdown menu on the right side of the track to toggle different elements' visibility.
+There is a dropdown menu on the right side of the track to toggle the visibility of different elements.
 
 ![](ElementTiming3.png)
 
@@ -186,7 +192,7 @@ Please see [](Engine/Animation/Sequences\Curves) for more information.
 
 
 ## Creating Variant Montages
-**Child Anim Montages** provide a convenient way to create animation variants without affecting gameplay, and without needing to recreate and maintain notifications or other details across multiple assets. An example use case for this might be varying the auto-attack animations on units in an RTS game, or changing a character's idle animation in an RPG or action game based on what kind of weapons or armor they are holding or wearing. A Child Animation Montage can be created from an existing Animation Montage via the "Create Child Montage" context-menu option.
+**Child Anim Montages** provide a convenient way to create animation variants without affecting gameplay, and without needing to recreate and maintain notifications (or other details) across multiple assets. An example use case for this might be varying the auto-attack animations on units in a real-time strategy (RTS) game, or changing a character's idle animation in a role-playing game (RPG), or an action game, based on what kind of weapons or armor they are holding or wearing. A Child Animation Montage can be created from an existing Animation Montage via the "Create Child Montage" context-menu option.
 
 ![](CreateChildMontage.png)
 
@@ -198,7 +204,7 @@ The new montage will have the same name as the original, but with "_Child" added
 Although Child Anim Montages are not a separate class from standard Anim Montages (both are `UAnimMontage`), Child Anim Montages cannot be used as parents for other Child Anim Montages.
 [/REGION]
 
-When editing a Child Anim Montage, all editable parts of the interface will be dimmed and deactivated except for the Slot Track. The Slot Track is in a special mode which enables animations to be overridden. New animations cannot be added, and the details of an animation (such as play rate) will not be editable. This helps to create varied animations with minimal effort and without affecting gameplay, since you can randomly pick between the original montage and its children freely, knowing that the timing of the animations themselves, and all notification events, will be identical. When editing a Child Anim Montage, animations that have not been changed from the parent will be shown in green, and animations that have been overridden will be shown in yellow.
+When editing a Child Anim Montage, all editable parts of the interface will be dimmed and deactivated (except for the Slot Track). The Slot Track is in a special mode which enables animations to be overridden. New animations cannot be added, and the details of an animation (such as play rate) will not be editable. This helps to create varied animations with minimal effort, without affecting gameplay. Since you can randomly pick between the original montage and its children freely, knowing that the timing of the animation (and all notification) events, will be identical. When editing a Child Anim Montage, animations that have not been changed from the parent will be shown in green, and animations that have been overridden will be shown in yellow.
 
 ![](AnimMontageSlot.png)
 
@@ -206,34 +212,34 @@ When editing a Child Anim Montage, all editable parts of the interface will be d
 
 The animation substitution in the picture above is indicated by the yellow box in the Child Anim Montage's editor window. 
 
-It is possible to substitute any animation for any other, even if they are of different lengths. The engine will automatically adjust the play rate, start time, and end time of the replacement animation so that it uses the equivalent portion and runs in the same time as the original would have. For example, if a one-second animation with a start time of 0.5 and an end time of 1.0 is replaced with a two-second animation, the two-second animation will play with a start time of 1.0, an end time of 2.0, and will play at double the speed of the one-second animation. The replacement animation will always fit perfectly into the original slot.
+It is possible to substitute any animation for any other, even if they are of different lengths. The engine will automatically adjust the play rate, start time, and end time, of the replacement animation so that it uses the equivalent portion and runs in the same time as the original. For example, if a one-second animation with a start time of 0.5, and an end time of 1.0, is replaced with a two-second animation, the two-second animation will play with a start time of 1.0, with an end time of 2.0, and will play at double the speed of the one-second animation. The replacement animation will always fit perfectly into the original slot.
 
 ## Playing Montages
 
 Montages can be played in a number of different ways: 
 
 * You can assign them as the **Anim to Play** for a **Skeletal Mesh**.
-* You can call them directly through **Blueprint** script or C++.
+* You can call them directly through a **Blueprint** script, or C++.
 * You can incorporate them into the **Animation Graph** as part of your **Animation Blueprint**.
 
 [REGION:note]
-Please see [](Engine/Animation/AnimHowTo/AdditiveAnimations) for an example of how a Montage can be used inside your Animation Blueprints. This example illustrates a character that can perform two animations at the same time independently of one another (in this case the character can move and shoot at the same time).
+Please see [](Engine/Animation/AnimHowTo/AdditiveAnimations) for an example of how a Montage can be used inside of your Animation Blueprints. This example illustrates a character that can perform two animations at the same time independently of one another (in this case, the character can move and shoot at the same time).
 [/REGION]
 
 ### Playing Multiple Montages
 
-Prior to Unreal Engine version 4.6, you could only play a single Montage at a time. However, as of 4.6 or later it is now possible to play as many Montages at you would like through the use of **Anim Slots** and **Anim Groups**.
+Prior to Unreal Engine version 4.6, you could only play a single Montage at a time. However, as of 4.6 or later, it is now possible to play as many Montages at you would like through the use of **Anim Slots** and **Anim Groups**.
 
 Here's a quick rundown of how it works: 
 
 1. Animations are organized into Anim Groups. You can play a single Montage per group. If you want to play another Montage at the same time, you can create a new Anim Group for it.
 1. Anim Groups and Anim Slots are centralized into the **Anim Slot Manager** tab inside **Persona** (see [Anim Slot Manager](#animslotmanager) below).
-1. You can use the drop-down UI to select between valid slots.
-1. When creating a new Montage, Slot nodes start out with a single **Default Slot** so they are immediately functional.
+1. You can use the drop-down User Interface (UI) to select between valid slots.
+1. When creating a new Montage, Slot nodes start out with a single **Default Slot** so that they are immediately functional.
 
 ### Anim Slot Manager
 
-The **Anim Slot Manager** is where you will manage your Anim Slot and Anim Group names. It displays a hierarchical view and lets you make changes. You can open the Anim Slot Manager inside a Montage in **Persona**, by clicking the magnifying glass icon (pictured below).
+The **Anim Slot Manager** is where you will manage your Anim Slot and Anim Group names. It displays a hierarchical view and lets you make changes. You can open the Anim Slot Manager inside of a Montage in **Persona** by clicking the magnifying glass icon (pictured below).
 
 ![](OpenAnimSlotManger.png)
 
@@ -241,9 +247,9 @@ This will open the Anim Slot Manager.
 
 ![](AnimSlotManager.png)
 
-Groups and Slots are saved in the **Skeleton** asset so all Animation Blueprints and Montages from the same family can use the same settings. The **Save** button inside the Anim Slot Manager is there so that you can easily save the changes to the Skeleton asset whenever you make any changes.
+Groups and Slots are saved in the **Skeleton** asset so that all Animation Blueprints and Montages from the same family can use the same settings. The **Save** button inside of the Anim Slot Manager is there so that you can easily save the changes to the Skeleton asset whenever you make any changes.
 
-When **Right-clicking** on a Slot name, a context menu is presented which will also allow you to move it to a different group if you desire. This makes it easier to manage the Slots you have created and place and move them into different groups and arrange them as you see fit.
+When right-clicking on a Slot name, a context menu is presented, which will also allow you to move it to a different group (if you desire). This makes it easier to manage the Slots you have created and place and move them into different groups and arrange them as you see fit.
 
 ![](SlotClicking.png)
 
@@ -259,21 +265,21 @@ The status bar across the top will tell you which Group the Montage belongs to.
 
 ![](SlotExample2.png)
 
-If you have multiple Slots, the status bar will also tell you if there are any issues with the setup. For example, if you mix and match Slots that belong to different groups, a warning sign will appear to which you can hover over and see more information about the problem.
+If you have multiple Slots, the status bar will also tell you if there are any issues with the setup. For example, if you mix and match Slots that belong to different groups, a warning sign will appear, allowing you to see more information about the problem.
 
 ![](SlotExample3.png)
 
 ### Editing Slots in Animation Blueprints
 
-Inside the AnimGraph of an Animation Blueprint, you can add a **Slot** node that when selected, allows you to see the same options as shown in the Montage Track view. 
+Inside the AnimGraph of an Animation Blueprint, you can add a **Slot** node that, when selected, allows you to see the same options as shown in the Montage Track view. 
 
 ![](SlotNode.png)
 
-From here you can use the drop-down box to pick an existing Slot name and the magnifying glass icon to open the Anim Slot Manager. 
+From here, you can use the drop-down box to pick an existing Slot name, and the magnifying glass icon, to open the Anim Slot Manager. 
 
-Additional notes: 
+###Final Notes
 
-* This is backwards compatible with previous setups. You should not notice any changes other than your existing Slots being part of a "Default Group". 
+* This is backwards compatible with previous setups. You should not notice any changes... other than your existing Slots being part of a "Default Group". 
 * The Montage Blueprint and C++ API have been enhanced to take Montage asset references. This is how you can control different Montages playing at the same time. For example, see the **Montage Set Play Rate** function. 
 
 
diff --git a/Engine/Documentation/Source/Engine/Animation/AnimPose/AnimBPPoseNodes/AnimBPPoseNodes.INT.udn b/Engine/Documentation/Source/Engine/Animation/AnimPose/AnimBPPoseNodes/AnimBPPoseNodes.INT.udn
--- a/Engine/Documentation/Source/Engine/Animation/AnimPose/AnimBPPoseNodes/AnimBPPoseNodes.INT.udn
+++ b/Engine/Documentation/Source/Engine/Animation/AnimPose/AnimBPPoseNodes/AnimBPPoseNodes.INT.udn
@@ -71,8 +71,8 @@ If you are using a **Pose Asset** with **Additive** enabled, you will also need
 
 ## Pose Driver Node
 
-Available within the AnimGraph of an Animation Blueprint is the **Pose Driver** node which lets you drive a curve value based on a bone's movement. 
-One use for this node would be to create a pose space deformer to drive corrective blend shapes based on a bone's orientation. 
+Available within the AnimGraph of an Animation Blueprint is the **Pose Driver** node which allows a bone to drive other aspects of animation, based on a set of 'example poses' or lets you drive a curve value based on a bone's movement. 
+One use for this node would be to create a pose space deformer to drive corrective blend shapes based on a bone's orientation or translation. 
 For example, in the clip below, clavicle bone poses are being used to move the shoulder relative to the arm bone's movement. 
 
 [OBJECT:EmbeddedVideo]
@@ -87,7 +87,7 @@ tJSYgPmzoKc
 [/PARAMLITERAL]
 [/OBJECT]
 
-To achieve this, an RBF (Radial Basis Function) interpolates driven values based on the orientation of a target bone. 
+To achieve this, an **RBF (Radial Basis Function)** interpolates driven values based on the orientation of a target bone. 
 You will need to create a **Pose Asset** and define the target poses for the bone and the desired curve values at each pose. 
 
 On the **Pose Driver** node itself, you have the following options: 
@@ -100,9 +100,19 @@ On the **Pose Driver** node itself, you have the following options:
 | **Radial Scaling** | Scaling of radial basis, applied to max distance between poses.| 
 | **Pose Source** | The **Pose Asset** used as the source of pose data. | 
 | **Include Ref Pose as Neutral Pose** | Should the mesh ref pose of Source Bone be considered as a "neutral pose" (zero curves). | 
-| **Type** | **Swing and Twist** (Consider full rotation for interpolation) or **Swing Only** (Consider only swing for interpolation). | 
+| **Type** | [INCLUDE:Engine\Animation\AnimPose\AnimBPPoseNodes\#typetable] | 
 | **Twist Axis** | The axis to use when **Type** is set to **SwingOnly** (X, Y or Z). | 
 
+<!--
+[EXCERPT:typetable]
+| Option | Description |
+| ----------- |------|
+| **Swing and Twist** | Consider full rotation for interpolation. | 
+| **Swing Only** | Consider only swing for interpolation. | 
+| **Translation** | Consider translation relative to parent. | 
+[/EXCERPT:typetable]
+-->
+
 For additional information see the Unreal Engine Live Stream on Animation Topics which include a **Pose Driver** example:
  
 [OBJECT:EmbeddedVideo]
diff --git a/Engine/Documentation/Source/Engine/Animation/Blendspaces/Blendspaces.INT.udn b/Engine/Documentation/Source/Engine/Animation/Blendspaces/Blendspaces.INT.udn
--- a/Engine/Documentation/Source/Engine/Animation/Blendspaces/Blendspaces.INT.udn
+++ b/Engine/Documentation/Source/Engine/Animation/Blendspaces/Blendspaces.INT.udn
@@ -5,12 +5,11 @@ Description:Blend Spaces are assets that allow any number of animations to be bl
 Related: Engine/Animation/AnimBlueprints
 Related: Engine/Content/FBX/Animations
 Related: Engine/Animation/AnimMontage
-Related: Engine/Animation/PhysicallyDrivenAnimation
+Related: Engine\Animation\AnimHowTo\BlendSpace
 Related: Engine/Animation/NodeReference/Blend
 Related: Engine/Animation/StateMachines
 Related: Engine/Animation/AnimBlueprints/AnimGraph
-Related: Engine/Content/Tools/MayaRiggingTool
-version: 4.9
+version: 4.15
 
 [VAR:Topic]
 [OBJECT:Topic]
@@ -32,15 +31,11 @@ version: 4.9
 [/OBJECT]
 [/VAR]
 
-[TOC (start:2 end:2)]
-
-
-
 [EXCERPT:Intro]
 **Blend Spaces** are special assets which can be sampled in AnimGraphs that allow for blending of 
 animations based on the values of two inputs. Simple blending between two animations based on a 
-single input can be performed using one of the standard [blend nodes](Engine/Animation/NodeReference/Blend) 
-available in Animation Blueprints. Blend Spaces provide a means of doing more complex blending 
+single input can be performed using one of the standard [Blend nodes](Engine/Animation/NodeReference/Blend) 
+available in [Animation Blueprints](Engine/Animation/AnimBlueprints). Blend Spaces provide a means of doing more complex blending 
 between multiple animations based on multiple values (currently limited to two).
 
 The goal of Blend Spaces is to reduce the need for creating individual, hard-coded nodes to perform 
@@ -49,154 +44,10 @@ the inputs, the animations, and how the inputs are used to blend between animati
 type of blending can be performed using the generic Blend Space.
 [/EXCERPT:Intro]
 
-
-
-
-[REGION:topics third]
+[REGION:topics half]
+%Engine\Animation\Blendspaces\Overview:topic%
 %Engine/Animation/Blendspaces/Creation:topic%
 %Engine/Animation/Blendspaces/Editor:topic%
 %Engine/Animation/Blendspaces/UserGuide:topic%
 [/REGION]
 
-[REGION:tip]
-An example of a Locomotion **Blendspace** is also shown in the [Animation Content Examples](Resources\ContentExamples\Animation) page under section 1.3.
-[/REGION]
-
-## How Blend Spaces Work
-
-The goal of Blend Spaces is to reduce the need for creating individual, hard-coded nodes to perform 
-blending based on specific properties or conditions. By allowing the animator or programmer to specify 
-the inputs, the animations, and how the inputs are used to blend between animations, virtually any 
-type of blending can be performed using the generic Blend Space. This is completely contrary to how 
-AnimTrees in Unreal Engine 3 handled the same tasks. Any complex blending required the creation of a 
-new node to handle that blending. Blend Spaces are completely generic allowing the factors that determine 
-the blending to be specified on each individual Blend Space. Each Blend Space has inputs that simply take 
-in values. These values can be calculated during the update cycle of an Animation Blueprint (via the **EventGraph**), 
-driven by gameplay code, or any other means. This makes them extremely flexible and puts the power into 
-the hands of the person creating the AnimGraph to blend animations in any way they see fit.
-
-![Input Flow](bs_flow.png)
-
-You can think of a Blend Space like a 2D graph with each input value along one axis and animations plotted 
-on the graph at various points. The blended animation is calculated by blending between the animations 
-on the graph based on the location designated by the current input values.
-
-![](BlendSpaceLayout_2.png)
-
-1. Animation Sequence 1
-1. Animation Sequence 2
-1. Animation Sequence 3
-1. Input 1 Axis (X-Axis)
-1. Input 2 Axis (Y-Axis)
-
-
-As an example, take the typical blending between directional run animations and an idle animation based on 
-the direction and speed of the movement of the player.
-
-In Unreal Engine 3, this was performed using the directional blending node in combination with a speed 
-blending node. Each of these nodes was hard-coded to only do that specific blend. The values used to 
-perform the blending were hidden away in the code, so no modifications could be made to tweak the blend 
-except by a programmer; and if you modified the value in the code, it affected every instance of the node 
-in every AnimTree making use of it. This is far from ideal.
-
-With Blend Spaces, this same blend can be performed using a single Blend Space. The direction (e.g., -180 to 
-180 representing degrees) and speed (e.g., 0 to 250 representing units per second) of the player would be 
-values passed to the Blend Space as inputs and the animations would be set up to blend appropriately based 
-on the values of these inputs. 
-
-![](BlendSpaceDirection.png)(w:720)
-
-In addition to only requiring a single node, the values used as inputs can be calculated directly in the 
-Animation Blueprint's **EventGraph** and then passed to the Blend Space in the **AnimGraph**. The input values can be 
-tweaked easily. The thresholds for animation blending can be quickly and intuitively adjusted using a 
-graphical editor. The flexibility and ease of use represent a huge win over previous systems.
-
-## One-Dimensional Blend Spaces
-
-Blend Spaces can also be created in a one-dimensional format, known as a **Blend Space 1D**. These can blend between any number of poses or animations, but do so based on a single input value. An example that we use is to make a character lean when sprinting. We have two poses, one leaning left and one leaning right, and use an input driven by how much the player is turning the character while sprinting. By blending into the additive leaning poses, we can give the running animation a bit more life as the character seems to lean into hard turns.
-
-Aside from having only one axis instead of two, a Blend Space 1D is exactly the same as a standard two-axis Blend Space. As such, this document will focus on the 2D version. But you should know that all the rules and properties still carry over.
-
-[OBJECT:EmbeddedVideo]
-[PARAMLITERAL:width]
-640
-[/PARAMLITERAL]
-[PARAMLITERAL:height]
-360
-[/PARAMLITERAL]
-[PARAMLITERAL:videoid]
-2mkn8FZL2KA
-[/PARAMLITERAL]
-[/OBJECT]
-
-## Architecture
-
-The underlying architecture of Blend Spaces is fairly simple. The main parts were touched on previously: 
-inputs values and sampled animations.
-
-### Skeleton
-
-Each Blend Space, as with other animation asset types, is targeted to a specific **Skeleton**. The Skeleton 
-associated with the Blend Space is set when the Blend Space is created and determines what animations can be 
-sampled within the Blend Space. It also determines which Animation Blueprints the Blend Space can be sampled within 
-as they both must target the same Skeleton.
-
-### Blend Parameters
-
-[EXCERPT:BlendParameters]
-The **Blend Parameters** of a Blend Space represent the input values used to drive the blending between 
-sample poses. Each Blend Space has two Blend Parameters (**X** and **Y**) and each Blend Parameter has 
-the following properties:
-
-| Property | Description |
-| -------- | ----------- |
-| **Label** | The readable name displayed in the **BlendSpace Editor** along the axis for this input and on the data pin for this input on the BlendSpace node in the AnimGraph of a Animation Blueprint. |
-| **Range** | The minimum and maximum values to which any input data will be clamped. |
-| **Divisions** | The number of sampled points from the animation data used for interpolating. Higher values give more precise results. Displayed as grid lines in the **BlendSpace Editor** along the axis for this input. |
-[/EXCERPT:BlendParameters]
-
-### Sample Data
-
-[EXCERPT:SampleData]
-The **Sample Data** for the Blend Space is a collection of animations and sample values. The sample values are 
-used to determine the weight of that particular animation. For instance, using the directional movement example 
-above, a **_Run_Fwd_** animation with a sample value of (0.0, 250.0, 0.0) would be fully blended to when the 
-direction input was 0.0 and the speed input was 250.0.
-[/EXCERPT:SampleData]
-
-Each item in the Sample Data array has the following properties:
-
-| Property | Description |
-| -------- | ----------- |
-| **Animation** | The AnimationSequence asset to blend to. |
-| **Sample Value** | The input values (X and Y corresponding to the Blend Parameters) where the Animation has full influence. |
-
-
-## Anim Asset Details
-
-When setting up your Blend Space asset, there are additional options that can be set via the **Anim Asset Details** panel inside of **Persona**. By default, this window is located in the lower left-hand corner of the **Animation** tab of the Persona Editor and contains the following properties: 
-
-![](BlendSpaceDetails.png)
-
-| Property | Description |
-|---|---|
-| [REGION:tablesection]Blend Space[/REGION] ||
-|**Display Editor Vertically**| Only available for 1D Blend Space assets, toggles the Blend Space graph display in either vertical or horizontal display modes.|
-| [REGION:tablesection]Input Interpolation[/REGION] ||
-|**Axis to Scale Animation**| Not available for 	1D Blend Space assets, this allows you (if you have input interpolation) to determine which axis to drive animation speed (scale). For example, for locomotion animation, speed axis will drive animation speed (thus scale).|
-|**Interpolation Params**| Input interpolation parameter for all 3 axis, for each axis input, decide how you'd like to interpolate input to by setting Interpolation Time and Interpolation Type. |
-| [REGION:tablesection]Sample Interpolation[/REGION] ||
-|**Target Weight Interpolation Speed Per Sec**| When target samples are set, how fast you would like to get to target which improves target blending. For example, if you interpolate input, when you move from left to right rapidly, you will interpolate through forward, but if you use target weight interpolation, you will skip forward but interpolate between left and right.|
-|**Per Bone Blend**| Define target weight interpolation per bone. This will blend in different speed per each bone setting.|
-| [REGION:tablesection]Animation Notifies[/REGION] ||
-|**Notify Trigger Mode**| Assign the current mode used by the Blend Space to decide which animation notifies to fire: All Animation (all notify events will fire), HighestWeightedAnimation (Notify events will only fire from the highest weighted animation) or None (no notify events will fire from any animations).|
-| [REGION:tablesection]Additive Settings[/REGION] ||
-|**Preview Base Pose**| Preview Base pose for additive Blend Space.|
-| [REGION:tablesection]Animation[/REGION] ||
-|**Skeleton**| Pointer to the Skeleton this asset can be played on (unadjustable).|
-| [REGION:tablesection]Meta Data[/REGION] ||
-|**Meta Data**| This is Meta Data that can be saved with the asset. The meta data is a Blueprintable class derived from the Anim Meta Data class. This allows you to add custom meta data to animation assets (Anim Sequence, Anim Montage, Anim Composite, and Blendspace are all supported). You can query the data from an animation asset in C++ by calling the `GetMetaData()` method, or by using `GetSectionMetaData()` on Anim Montages. |
-| [REGION:tablesection]Thumbnail[/REGION] ||
-|**Orbit Pitch**| The pitch of the orbit camera around the asset. |
-|**Orbit Yaw**| The yaw of the orbit camera around the asset. |
-|**Orbit Zoom**| The offset from the bounds sphere distance from the asset.|
diff --git a/Engine/Documentation/Source/Engine/Animation/Blendspaces/Creation/BlendspaceCreation.INT.udn b/Engine/Documentation/Source/Engine/Animation/Blendspaces/Creation/BlendspaceCreation.INT.udn
--- a/Engine/Documentation/Source/Engine/Animation/Blendspaces/Creation/BlendspaceCreation.INT.udn
+++ b/Engine/Documentation/Source/Engine/Animation/Blendspaces/Creation/BlendspaceCreation.INT.udn
@@ -1,51 +1,48 @@
 Availability:Public
 Title: Creating Blend Spaces
 Crumbs: %ROOT%, Engine, Engine/Animation, Engine/Animation/Blendspaces
-Description:Guide to creating new BlendSpace multi-input animation blending assets for use in AnimGraphs of Animation Blueprint assets.
-version: 4.9
-
-[VAR:Topic]
-[OBJECT:Topic]
-	[PARAM:image]
-		![%Engine/Animation/Blendspaces/Creation:title%](Engine/Animation/Blendspaces/create.png)
-	[/PARAM]
-	[PARAM:icon]
-		![](%ROOT%/content_icon.png)(convert:false)
-	[/PARAM]
-	[PARAM:title]
-		%Engine/Animation/Blendspaces/Creation:title%
-	[/PARAM]
-	[PARAM:description]
-		%Engine/Animation/Blendspaces/Creation:description%
-	[/PARAM]
-	[PARAM:path]
-		[RELATIVE:Engine/Animation/Blendspaces/Creation]
-	[/PARAM]
-[/OBJECT]
-[/VAR]
-
-
-Blend Spaces are regular assets and must be created through the **Content Browser** like any other asset.
-
-To create a new Blend Space:
-
-1. In the **Content Browser**, navigate to the folder within which you want to create your new Blend Space, then click ![](Shared/EditorButtons/button_NewAsset.png) or **Right-click** and choose **Animation** > **Blend Space**:
+Description:Overview of creating a Blend Space and Blend Space 1D asset for use in Animation Blueprints. 
+Related: Engine\Animation\AnimHowTo\BlendSpace
+Related: Engine/Animation/Blendspaces/Editor
+Related: Engine/Animation/StateMachines
+Related: Engine/Animation/AnimBlueprints
+Related: Engine/Animation/AnimBlueprints/AnimGraph
+version: 4.15
+type:how-to
+tags:Animation
+tags:Animation Blueprint
+tags:Blend Space
+topic-image:createblendspacetopic.png
+
+This page covers how to create both **Blend Space** and **Blend Space 1D** assets as you will need to select a [Skeleton](Engine/Animation/Skeleton) to associate the Blend Space with.
+
+[REGION:note]
+Creating a Blend Space and Blend Space 1D are generally the same, however, [](Engine/Animation/Blendspaces/Editor) slightly differs based on which asset you are working with. 
+[/REGION]
+
+1. In the **Content Browser**, click the **Add New** button (or **Right-click** in the Content Browser) and choose **Animation** > **Blend Space**.
 
 	![](NewBlendSpace.png)
 
-1. In the **Pick Skeleton** dialog, select the Skeleton the Blend Space should target:
+	If you would like to create a **Blend Space 1D**, select the option for Blend Space 1D instead. 
+
+	![](blendspace1dcreate.png)
+
+1. In the **Pick Skeleton** dialog, select the Skeleton the Blend Space should target.
 
 	![](SkeletonList.png)
 
-1. The new Blend Space asset is displayed in the **Content Browser** at the location specified:
+	Your list of **Skeleton Assets** may look different based on the number of Skeleton Assets you have in your project. 
+
+1.  Enter a name for your new Blend Space asset. 
 
 	![Create Blueprint - Asset Created](bs_create_asset.png)
 
 	The Blend Space is also available to be used in the **Anim Graphs** of Animation Blueprints targeting the same Skeleton as the Blend Space:
 
-	[REGION:imagetable]
-	| ![](AssetBrowserBlendList.png)(h:400) | ![](ContextMenuBlendList.png)(h:400) |
-	| :---: | :---: |
-	| **Asset Browser** | **Context Menu** |
-	[/REGION]
+	![](AssetBrowserBlendList.png)
+		
+
+
+
 
diff --git a/Engine/Documentation/Source/Engine/Animation/Blendspaces/Editor/BlendspaceEditor.INT.udn b/Engine/Documentation/Source/Engine/Animation/Blendspaces/Editor/BlendspaceEditor.INT.udn
--- a/Engine/Documentation/Source/Engine/Animation/Blendspaces/Editor/BlendspaceEditor.INT.udn
+++ b/Engine/Documentation/Source/Engine/Animation/Blendspaces/Editor/BlendspaceEditor.INT.udn
@@ -1,188 +1,176 @@
 Availability:Public
 Title: Editing Blend Spaces
 Crumbs: %ROOT%, Engine, Engine/Animation, Engine/Animation/Blendspaces
-Description:Guide to setting up sample parameter and animation samples in Blendspaces for complex, multi-input animation blending.
-version: 4.9
-
-
-[VAR:topic]
-[OBJECT:Topic]
-	[PARAM:image]
-		![%Engine/Animation/Blendspaces/Editor:title%](Engine/Animation/Blendspaces/edit.png)
-	[/PARAM]
-	[PARAM:icon]
-		![](%ROOT%/editor_icon.png)(convert:false)
-	[/PARAM]
-	[PARAM:title]
-		%Engine/Animation/Blendspaces/Editor:title%
-	[/PARAM]
-	[PARAM:description]
-		%Engine/Animation/Blendspaces/Editor:description%
-	[/PARAM]
-	[PARAM:path]
-		[RELATIVE:Engine/Animation/Blendspaces/Editor]
-	[/PARAM]
-[/OBJECT]
-[/VAR]
+Description:Outlines how you can edit and set up the parameters that drive your Blend Space asset.
+Related: Engine\Animation\AnimHowTo\BlendSpace
+Related: Engine\Animation\Blendspaces\Creation
+Related: Engine/Animation/StateMachines
+Related: Engine/Animation/AnimBlueprints
+Related: Engine/Animation/AnimBlueprints/AnimGraph
+version: 4.15
+type:reference
+topic-image:editblendspacetopic.png
+tags:Animation
+tags:Animation Blueprint
+tags:Blend Space
 
+After [Creating a Blend Space](Engine\Animation\Blendspaces\Creation) and opening the asset, you will need to define the properties that drive the blending as well as specify what animations you want to sample from. 
 
+[TOC(start:2 end:2)]
 
-[REGION:banner]
-![Blendspace Editor](bs_editor_banner.png)(convert:false)
-[/REGION]
+## Axis Settings
 
-[EXCERPT:overview]
-The **Blend Space Editor** provides the ability to set up the Blend Parameters of the BlendSpace as 
-well as graphically edit the positions of Sample poses and interactively see the results of blending 
-between those Sample poses.
-[/EXCERPT:overview]
+The first thing you will want to do is set up your grid by defining the **Axis Settings** from the [Asset Details](Engine/Animation/Persona/AnimAssetDetails) panel.  
 
-[TOC(start:2)]
+![](AxisSettings.png)
 
-##Interface
+For the normal Blend Space asset, you can define the **Horizontal Axis** and **Vertical Axis** as your inputs while Blend Space 1D assets use the horizontal axis only. 
 
-![](TestBlendShape.png)(w:720)
+For either axis, the properties that can be defined are the same: 
 
-1. **Parameters** - Blend Parameters property section for setting up the BlendSpace.
-1. **Samples** - Editable list of all Samples in the BlendSpace.
-1. **Options** - Preview options for the BlendSpace.
-1. **Graph** - 2D graph for placing Samples and previewing blending between them.
+[EXCERPT:axissettings]
+|Axis Settings||
+|---|---|
+|**Name**| The name to display for this axis on the Blend Space Editor grid (also the name of the exposed pin on the Blend Space node inside an Animation Blueprint). |
+|**Minimum Axis Value**| Minimum value for this parameter to which any input data will be clamped. |
+|**Maximum Axis Value**| Maximum value for this parameter to which any input data will be clamped. |
+|**Number of Grid Divisions**| The number of sampled points from the animation data used for interpolating. Higher values give more precise results. This is displayed as gridlines in the Blend Space Editor grid along the axis for this input. |
+|**Interpolation Time**| Interpolation Time for input, when it gets input, it will use this time to interpolate to target. Used for smoother interpolation. |
+|**Interpolation Type**| Type of interpolation used for filtering the input value to decide how to get to the target. |
+[/EXCERPT:axissettings]
 
+Typically for multi-direction locomotion, you would define one axis as the **Direction** in degrees which you are moving and the other axis the **Speed** in which you are moving. 
 
-## Setting up Parameters
+![](LocoMotionBlendSpace.png)
 
-Parameters for the BlendSpace are set up in the **Parameters** section of the UI.
+[REGION:note]
+Please see the [](Engine\Animation\AnimHowTo\BlendSpace\) how-to page for a guide on setting up multi-directional locomotion with a Blend Space. 
+[/REGION]
 
-![Blendspace Parameters](bs_params.png)
+## Sample Poses
 
-[INCLUDE:Engine/Animation/Blendspaces#BlendParameters]
+With your Axis Settings defined, the next thing you will want to do is add some animations to sample from to the Blend Space Editor grid. 
 
-Once the appropriate values have been set for each of the Blend Parameters, pressing the ![Apply Parameter Changes](button_apply_changes.png) 
-button will apply the values to the BlendSpace, updating the **Graph**.
+![](DragDropToGrid.png)
 
-![](BlendSpaceSetup.png)
+In the example above, we have dragged an Idle animation from the **Asset Browser** into the grid at the point where Direction/Speed values are 0, or the character is not moving. 
 
-## Working with Sample Poses
+You can **Right-click** on a sample point to expand a roll-out menu with adjustable options for the sample.
 
-[INCLUDE:Engine/Animation/Blendspaces#SampleData]
+![](RightClickSamplePoint.png)
 
-### Adding Samples
+You can change the location of the sample by adjusting either Axis Value, change the Animation or adjust the **Rate Scale** (or playback speed of the animation). 
 
-Samples are added to a BlendSpace by dragging an AnimationSequence from the **Asset Browser** 
-into the **Graph**.
+In addition to using the roll-out menu to move the sample by changing the Axis Value properties, you can drag-and-drop the sample to a new location on the grid. 
 
-[REGION:note]
-When working with AimOffset BlendSpaces, the AnimationSequences must have **Additive Anim Type** 
-set to AAT_RotationOffsetMeshSpace in order to be added as Samples.
-[/REGION]
+![](MovedKey.png)
 
-![](BlendShapeDragDrop.png)(w:720)
+To remove a sample from the grid, select your sample and press the **Delete** key. 
 
-A node is placed at the location in the **Graph**:
+### Previewing Blend Spaces
 
-![Sample node](bs_sample_node.png)
+Once you have a few sample poses placed on the grid, you can hold the **Shift** key and drag the green diamond around to view the blend between poses. 
 
-### Positioning Samples
-
-Samples can be moved to different positions interactively or manually. Moving Samples interactively 
-is very intuitive, but does not provide the preciseness of manually setting the Sample's position.
+[OBJECT:EmbeddedVideo]
+[PARAMLITERAL:width]
+640
+[/PARAMLITERAL]
+[PARAMLITERAL:height]
+360
+[/PARAMLITERAL]
+[PARAMLITERAL:videoid]
+Kh3EubQ1IWM
+[/PARAMLITERAL]
+[/OBJECT]
 
-To interactively move a sample:
+When you release the **Shift** key, the green diamond will remain in its location enabling you to freely move the camera in the viewport to view the pose from different angles. 
 
-Click and dragging the node in the **Graph**:
+![](StickyKey.png)
 
-![Move Sample by Dragging](bs_sample_move_drag.png)(convert:false)
- 
-To move a sample manually:
+### Grid Options
 
-Click the node and then specify the  **X** and **Y** values to place the node at:
+Inside the Blend Space Editor grid, you can also **Show Triangulation** or adjust **Stretch** settings. 
 
-![Move Node Manually](bs_sample_move_manual.png)
+![](GridOptionsImage.png)
 
-### Replacing Samples
+The video below illustrates the effects of both options and you can experiment to reach your desired effect. 
 
-Samples placed in the BlendSpace can be replaced with another Sample in the **Graph** via drag and drop or using the 
-**Samples** section of the UI.
+[OBJECT:EmbeddedVideo]
+[PARAMLITERAL:width]
+640
+[/PARAMLITERAL]
+[PARAMLITERAL:height]
+360
+[/PARAMLITERAL]
+[PARAMLITERAL:videoid]
+rmTo7dcN6k4
+[/PARAMLITERAL]
+[/OBJECT]
 
-To replace by drag and drop:
 
-To replace a sample in the **Graph**, drag a new AnimationSequence from the **Asset Browser** directly on 
-to the existing node.
 
-[REGION:fullwidth]
-![Replace Sample by Dragging](bs_sample_replace_drag.png)(convert:false)
-[/REGION]
 
-[REGION:imagetable]
-| ![Original Sample](bs_sample_replace_before.png)(h:220) |  ![Replaced Sample](bs_sample_replace_after.png)(w:175 h:220) |
-| ------ | ------ |
-| Original Sample | Replaced Sample |
-[/REGION]
+## Blend Space Asset Details
 
-To replace manually:
+While defining the Axis Settings and Sample Poses are two of the more critical elements to setting up your Blend Space, there are other properties in the **Asset Details** panel that you can adjust to further define how your Blend Space is used. 
 
-To replace a sample in the **Samples** section, select a new AnimationSequence in the **Content Browser** 
-and press the ![Use Selected Asset](button_use_selected.png) button for the Sample in the list.
+### Input Interpolation
 
-[REGION:note]
-	In this case, you must select your AnimationSequence from the **Content Browser**. This operation will not read any assets selected in Persona's **Asset Browser**.
-[/REGION]
+This section enables you to define the **Axis to Scale Animation** and is useful if you have input interpolation, as you can specify which axis to drive animation speed (scale). 
+For example, for locomotion animation, speed axis will drive animation speed (thus scale). 
+You can set this to **None** if it is not applicable, **Horizontal** or **Vertical** Axis. 
+This option is only available for normal Blend Space assets and not the Blend Space 1D asset. 
 
-![Select New Animation](bs_sample_replace_manual_select.png)(convert:false)
+### Axis Settings
 
-![Replace Sample Manually](bs_sample_replace_manual.png)(convert:false)
+As described at the top of this page, this enables you to define the Horizontal and (in the case of a normal Blend Space, Vertical) Axis properties that will drive your Blend Space.  
 
-### Removing Samples
+[INCLUDE:#axissettings]
 
-Samples can be deleted directly in the **Graph** or using the **Samples** section of the UI.
+[REGION:note]
+When changing the number of grid divisions or the size of either axis is changed, any existing samples will be snapped to the closest available grid point. 
+Those for which no valid position could be found are marked as invalid and will be drawn in red enabling you to reposition them manually on the grid. 
+[/REGION]
 
-To delete in the Graph:
 
-Select the Sample node in the **Graph** to display the edit dialog and press the ![Delete](button_delete.png) 
-button to remove the Sample.
+### Animation
 
-![Remove Sample in Graph](bs_sample_remove.png)(convert:false)
+This section contains options related to previewing as well as informational properties, such as what Skeleton Asset is associated with this asset. 
 
-To delete in the Samples section:
+|Animation||
+|---|---|
+|**Preview Pose Asset**| The default Skeletal Mesh to use when previewing this asset (only applies when you opening the Blend Space asset).  |
+|**Skeleton**| Pointer to the Skeleton this asset can be played on. |
+|**Parent Asset**| If set, enables you to add content to the mapping table. |
+|**Asset Mapping Table**| When Parent Asset is set, enables you to define the mappings associated with this asset. |
+|**Asset User Data**| Stores an array of any user data associated with this asset. |
 
-Press the ![Remove Sample](button_clear_reference.png) button for the Sample to be removed from the list.
+### Sample Interpolation
 
-![Remove Sample Manually](bs_sample_remove_manual.png)(convert:false)
+This section enables you to smooth out interpolation between sample poses on a per bone basis or per pose basis. 
 
-## Previewing
+|Sample Interpolation||
+|---|---|
+|**Per Bone Blend**| Define target weight interpolation per bone, this will blend in different speed per each bone setting defined. |
+|**Target Weight Interpolation Speed Per Sec**| When target samples are set, this determines how fast you would like to reach the target and can improve blending. |
 
-The result of the blending between poses can be previewed interactively in the **Viewport** panel of **Persona** 
-using the **Graph**. The position of the cursor determines the values of the **X** and **Y** parameters used for 
-performing the blending and the resulting pose is displayed in the **Viewport** panel.
+### Animation Notifies
 
-To enable previewing for the BlendSpace, check the **Enable Preview BlendSpace** option in the **Options** 
-section of the UI.
+This section defines the mode used by the Blend Space to determine which [Animation Notifiy](Engine/Animation/Sequences/Notifies) to fire. 
 
-![Enable Preview Blendspace Option](bs_options_preview.png)
+Notify Trigger Modes include: 
 
-Once previewing is enabled, moving the cursor around in the **Graph** will cause the **Viewport** panel to display 
-the blended pose corresponding to the position of the cursor. The grid block the cursor is within is shaded and 
-values representing the blending between the sample points making up the block are displayed.
+* **All Animations** - All notify events will fire. 
+* **Highest Weighted Animation** - Notify events will only fire from the highest weight animation. 
+* **None** - No notify events will fire from any animations. 
 
-![Graph Grid Block](bs_samples_block.png)
+### Meta Data
 
-Detailed information about the cursor location and the Samples being blended can be viewed by checking the **Enable 
-Tooltip Display** option.
+This section enables you to associate any custom metadata with the animation asset. You can query this by using the `GetMetaData` function. 
 
-![Enable Tooltip Display Option](bs_options_tooltip.png)
+### Thumbnail
 
-![Graph Tooltip](bs_samples_tooltip.png)
+This section enables
+ you to adjust the view perspective inside the asset thumbnail that appears in the **Content Browser**.
 
-An example of a previewing a Blend Space is illustrated below. 
 
-[OBJECT:EmbeddedVideo]
-[PARAMLITERAL:width]
-640
-[/PARAMLITERAL]
-[PARAMLITERAL:height]
-360
-[/PARAMLITERAL]
-[PARAMLITERAL:videoid]
-L435S6VaUMY
-[/PARAMLITERAL]
-[/OBJECT]
\ No newline at end of file
