Copying //UE4/Dev-Documentation to //UE4/Samples-Main (Source: //UE4/Dev-Documentation @ 3372893)

diff --git a/Engine/Documentation/Source/Engine/Animation/Blendspaces/Overview/BlendspaceOverview.INT.udn b/Engine/Documentation/Source/Engine/Animation/Blendspaces/Overview/BlendspaceOverview.INT.udn
new file mode 100644
--- /dev/null
+++ b/Engine/Documentation/Source/Engine/Animation/Blendspaces/Overview/BlendspaceOverview.INT.udn
@@ -0,0 +1,101 @@
+Availability:Docs
+Title: Blend Spaces Overview
+Crumbs: %ROOT%, Engine, Engine/Animation
+Description:Describes how Blend Space assets are used and the different types of Blend Space assets.
+Related: Engine\Animation\AnimHowTo\BlendSpace
+Related: Engine/Animation/Blendspaces/Editor
+Related: Engine/Animation/Blendspaces/UserGuide
+Related: Engine/Animation/StateMachines
+Related: Engine/Animation/AnimBlueprints
+version: 4.15
+type:overview
+topic-image:blendspacetopic.png
+tags:Animation
+tags:Animation Blueprint
+tags:Blend Space
+
+[TOC (start:2 end:2)]
+
+The goal of a **Blend Space** is to reduce the need for creating individual, hard-coded nodes for blending
+animations with an asset that performs blending based on specific properties or conditions instead. By enabling the animator or programmer to specify 
+the inputs, the animations, and how the inputs are used to blend between animations, virtually any 
+type of blending can be performed using the generic Blend Space. 
+
+This is completely contrary to how AnimTrees in Unreal Engine 3 handled the same tasks. 
+Any complex blending required the creation of a new node to handle that blending. 
+Blend Spaces are completely generic enabling the factors that determine the blending to be specified on each individual Blend Space. 
+Each Blend Space has inputs that simply take in values. 
+These values can be calculated during the update cycle of an [Animation Blueprint](Engine\Animation\AnimBlueprints) (via the **EventGraph**), 
+driven by gameplay code, or any other means (as seen in the example chart below). This makes them extremely flexible and puts the power into 
+the hands of the person creating the AnimGraph to blend animations in any way they see fit.
+
+![Input Flow](bs_flow.png)
+
+[REGION:caption]
+Above, the Event Graph, Gameplay Code or other factor drives each of the variables that in turn determines which pose in the Blend Space to use as the final pose. 
+[/REGION]
+
+
+## Blend Space
+
+You can think of a Blend Space like a 2D graph with each input value along one axis and animations plotted 
+on the graph at various points. The blended animation is calculated by blending between the animations 
+on the graph based on the location designated by the current input values.
+As an example, take the typical blending between directional run animations and an idle animation based on 
+the direction and speed of the movement of the player.
+
+[OBJECT:EmbeddedVideo]
+[PARAMLITERAL:width]
+640
+[/PARAMLITERAL]
+[PARAMLITERAL:height]
+360
+[/PARAMLITERAL]
+[PARAMLITERAL:videoid]
+HzgiW2mZomU
+[/PARAMLITERAL]
+[/OBJECT]
+
+In Unreal Engine 3, this was performed using the directional blending node in combination with a speed 
+blending node. Each of these nodes was hard-coded to only do that specific blend. The values used to 
+perform the blending were hidden away in the code, so no modifications could be made to tweak the blend 
+except by a programmer; and if you modified the value in the code, it affected every instance of the node 
+in every AnimTree making use of it. This is far from ideal.
+
+With Blend Spaces, this same blend can be performed using a single Blend Space. 
+The direction (which is indicated in the horizontal yellow box below) and speed (indicated in the vertical yellow box below) of the player would be values passed to the Blend Space as inputs.
+The animations (white diamonds below) would be set up to blend appropriately based on the values of these inputs resulting in the final pose you see (green diamond below) in-game.
+
+![](BlendSpaceDirection.png)(w:720)
+
+Below, our **AnimGraph** implements the Blend Space to which we can adjust our inputs (in this case **Speed**) to drive the blending. 
+
+![](AnimGraphBlendSpace.png)
+
+Just as in the previewer, as we adjust our input values we alter and drive the blending in our Blend Space. 
+
+## Blend Space 1D
+
+Blend Spaces can also be created in a one-dimensional format, known as a **Blend Space 1D**. These can blend between any number of poses or animations but do so based on a single input value. 
+An example use case for a Blend Space 1D would be when you have a character that automatically orientates towards the direction they are moving. 
+If the character cannot sidestep or move in multiple directions, a Blend Space 1D could be used to blend from an Idle to a Walk and ultimately Run based on a single **Speed** value (as shown in the example below). 
+
+[OBJECT:EmbeddedVideo]
+[PARAMLITERAL:width]
+640
+[/PARAMLITERAL]
+[PARAMLITERAL:height]
+360
+[/PARAMLITERAL]
+[PARAMLITERAL:videoid]
+2g_IxGNE0T4
+[/PARAMLITERAL]
+[/OBJECT]
+
+[REGION:note]
+Aside from having only one axis instead of two, a Blend Space 1D is exactly the same as a standard two-axis Blend Space. 
+[/REGION]
+
+
+
+
diff --git a/Engine/Documentation/Source/Engine/Animation/Blendspaces/UserGuide/BlendspaceUserGuide.INT.udn b/Engine/Documentation/Source/Engine/Animation/Blendspaces/UserGuide/BlendspaceUserGuide.INT.udn
--- a/Engine/Documentation/Source/Engine/Animation/Blendspaces/UserGuide/BlendspaceUserGuide.INT.udn
+++ b/Engine/Documentation/Source/Engine/Animation/Blendspaces/UserGuide/BlendspaceUserGuide.INT.udn
@@ -1,43 +1,36 @@
 Availability:Public
 Title: Using Blend Spaces
 Crumbs: %ROOT%, Engine, Engine/Animation, Engine/Animation/Blendspaces
-Description:Guide to using Blend Spaces in Animation Blueprints to input data and sample output animation poses.
-version: 4.9
-
-[VAR:Topic]
-[OBJECT:Topic]
-	[PARAM:image]
-		![%Engine/Animation/Blendspaces/UserGuide:title%](Engine/Animation/Blendspaces/Using.png)
-	[/PARAM]
-	[PARAM:icon]
-		![](%ROOT%/content_icon.png)(convert:false)
-	[/PARAM]
-	[PARAM:title]
-		%Engine/Animation/Blendspaces/UserGuide:title%
-	[/PARAM]
-	[PARAM:description]
-		%Engine/Animation/Blendspaces/UserGuide:description%
-	[/PARAM]
-	[PARAM:path]
-		[RELATIVE:Engine/Animation/Blendspaces/UserGuide]
-	[/PARAM]
-[/OBJECT]
-[/VAR]
-
-[TOC(start:2)]
-
+Description:Describes Blend Spaces are used within Animation Blueprints resulting in an output pose. 
+Related: Engine\Animation\AnimHowTo\BlendSpace
+Related: Engine/Animation/StateMachines
+Related: Engine/Animation/AnimBlueprints
+Related: Engine/Animation/AnimBlueprints/AnimGraph
+version: 4.15
+type:overview
+topic-image:userguidetopic.png
+tags:Animation
+tags:Animation Blueprint
+tags:Blend Space
 
 [EXCERPT:Intro]
-Blend Spaces are used by placing nodes in the **AnimGraphs** of Animation Blueprints. These nodes take in value data that is used to drive which animations the Blend Space uses to create the final blended pose.
+Blend Spaces are used by placing nodes in the **AnimGraphs** of [Animation Blueprints](Engine/Animation/AnimBlueprints) which take in value data to drive the blending of animations and the final animation pose. 
 [/EXCERPT:Intro]
 
+[REGION:note]
+For a full walk-through of setting up a Blend Space for locomotion blending, see the [](Engine\Animation\AnimHowTo\BlendSpace) how-to page. 
+[/REGION]
+
+
+[TOC(start:2)]
+
 ## Adding Blend Space Nodes
 
 Blend Space nodes provide access to the sample parameter data and output pose of the Blend Space.
  
-To place a Blend Space node using drag and drop:
+**To place a Blend Space node using drag-and-drop:**
 
-1. In **Persona**, find the Blend Space you want to sample in the **Asset Browser**.
+1. In the [](Engine/Animation/AnimBlueprints/AnimGraph) of an **Animation Blueprint**, find the Blend Space you want to sample in the **Asset Browser**.
 
 	![](SelectBlendSpace.png)
 
@@ -45,40 +38,36 @@ To place a Blend Space node using drag and drop:
 
 	![](DragDropBlendSpace.png)(w:720)
 
-1. Release the **Left Mouse Button** to place the Blend Space node.
+1. Release the **Left Mouse Button** to place the Blend Space node in the graph. 
 
 	![](BlendSpaceNodeAdded.png)
 
 
 **To place a Blend Space node using the context menu:**
 
-1. **Right-click** in the **AnimGraph** and expand **Animation** to see a list of all compatible Blend Spaces.
+1. **Right-click** in the **AnimGraph** and expand **Blend Spaces** and select the Blend Space you want to sample. 
 
 	![](AnimationContextMenu.png)
 
-1. Click on the Blend Space you want to sample.
-
-	![](BlendSpaceContextSelection.png)
-
-1. The Blend Space node is placed.
+1. The Blend Space node is placed in the graph. 
 
 	![](BlendSpaceNodeAdded.png)
 
 
-## Inputting Data
+## Sampling Parameters
 
-Blend Spaces rely on data being passed into the sample parameters in order to perform the animation blending. Each Blend Space node has input data pins for each sample parameter. Variables or outputs from other nodes can be wired to these pins to pass the necessary data to the Blend Space.
+Blend Spaces rely on data being passed into the sample parameters in order to perform the animation blending. 
+Each Blend Space node has input data pins for each sample parameter that is created during the [](Engine\Animation\Blendspaces\Editor) process.
+Variables or outputs from other nodes can be wired to these pins to pass the necessary data to the Blend Space and drive the blending of your Sample Poses. 
 
 ![Variable Input Data](bs_input_variables.png)
 
-[REGION:note]
-Since Blend Space editing is currently limited to using two inputs, the third input pin is not used.
-[/REGION]
-
-The values passed to the Blend Space are generally variables calculated during the update cycle in the **EventGraph** or through gameplay code. An example of the Direction and Speed variables being calculated in the **EventGraph** of an Animation Blueprint is shown below.
+The values passed to the Blend Space are generally variables calculated during the update cycle in the **EventGraph** or through gameplay code. 
 
-[REGION:fullwidth]
 ![Variable Input Data - Update Cycle](bs_input_calc.png)
+
+[REGION:caption]
+An example of the Direction and Speed variables being calculated in the **EventGraph** of an Animation Blueprint is shown above.
 [/REGION]
 
 ## Output Pose
@@ -90,3 +79,4 @@ The **Pose** output pin of the Blend Space node can be wired to the **Pose** inp
 Or, it can be wired to any other input pin that accepts an animation pose as an input when a more complex animation chain is necessary.
 
 ![Output Pose - Blend](bs_output_blend.png)
+
diff --git a/Engine/Documentation/Source/Engine/Animation/Persona/MeshDetails/MeshDetails.INT.udn b/Engine/Documentation/Source/Engine/Animation/Persona/MeshDetails/MeshDetails.INT.udn
--- a/Engine/Documentation/Source/Engine/Animation/Persona/MeshDetails/MeshDetails.INT.udn
+++ b/Engine/Documentation/Source/Engine/Animation/Persona/MeshDetails/MeshDetails.INT.udn
@@ -1,7 +1,7 @@
 Availability:Public
 Title:Skeletal Mesh Asset Details
 Crumbs: %ROOT%, Engine, Engine/Animation, Engine/Animation/Persona
-Description:User guide for the Mesh Details panel in the Persona animation editor.
+Description:User guide for the Asset Details panel in the Skeletal Mesh Editor.
 Related: Engine/Content/Types/SkeletalMeshes
 Related: Engine/Rendering/Materials
 Related: Engine/Content/Types/StaticMeshes/HowTo/LODCollision
@@ -14,32 +14,37 @@ tags:Animation
 tags:Skeletal Mesh
 tags:FBX Importing
 
+[TOC (start:2)]
+
 ![](MeshDetails.png)(w:940)
 
 [EXCERPT:Overview]
-The **Asset Details** panel inside the **Skeletal Mesh Editor** shows specific properties pertaining to the Skeletal Mesh, such as the Materials applied to the surface, as well as skinning and LOD options. 
-You can also assign and adjust APEX clothing or define and associate a **Physics Asset** for the Skeletal Mesh which is needed for creating ragdoll type motions. 
+The **Asset Details** panel, which is located inside of the **Skeletal Mesh Editor**, shows specific properties pertaining to a Skeletal Mesh (such as the Materials being applied to its surface) as well as skinning and LOD options. 
+You can also assign and adjust APEX clothing, or define and associate a **Physics Asset** for the Skeletal Mesh (which is needed for creating ragdoll type motions). 
 [/EXCERPT:Overview]
 
-This page highlights each section within the Asset Details panel with links to additional documentation, please see each section below: 
-
-[TOC (start:2)]
+This page highlights each section within the Asset Details panel with links to additional documentation.
 
 ## Materials 
 
+All materials applied to meshes modeled in external applications will be exported with the mesh. 
+When importing a mesh into Unreal Engine 4 (UE4), you can automatically import its materials as part of the [FBX Import Process](Engine/Content/FBX/Materials). 
+Materials will appear inside of the Skeletal Mesh Editor under the **Materials** section and can be overridden by any other Materials that exist within the UE4 project. 
+
 ![](MaterialsTable.png)
 
-Any materials applied to meshes modeled in external applications will be exported along with the mesh. 
-When importing that mesh into Unreal Engine 4, you can automatically import those materials as part of the [FBX Import Process.](Engine/Content/FBX/Materials)
-Those materials will appear inside the Skeletal Mesh Editor under this **Materials** section and can be overridden by any other Materials that exist within your project. 
-You can add **Material Slots** which can be used to override a LODs section (not the base LOD). 
+[REGION:note]
+You can add **Material Slots**, which can be used to override an LOD's section (excluding the base LOD).
+[/REGION]
 
 ## LODs
 
+When importing a Skeletal Mesh, you can also import the [Level of Details](Engine/Content/ImportingContent/ImportingSkeletalLODs/) (**LODs**) that were created and associated with the mesh. 
+
 ![](LODSection.png)
 
-When importing a Skeletal Mesh, you can also import any [Level of Details](Engine/Content/ImportingContent/ImportingSkeletalLODs/) (**LODs**) that you may have created and associated with your mesh. 
-You can **Highlight** portions of the mesh or **Isolate** portions to only show certain sections by clicking your desired option's checkbox. You can also assign the Material Elements from the Material section to **Material Slots** or override sections on a per LOD basis.
+You can **Highlight** portions of the mesh, or **Isolate** portions to only show certain sections of the mesh by clicking your desired option's checkbox. 
+You can also assign Material Elements from the Material section to **Material Slots**, or override sections per LOD.
 
 [OBJECT:ComparisonSlider]
 	[PARAM:before]
@@ -50,14 +55,14 @@ You can **Highlight** portions of the mesh or **Isolate** portions to only show
 	[/PARAM]
 [/OBJECT]
 
-Above we toggle between highlighting the helmet of our character to displaying only the helmet. 
+_The comparison slider toggles between highlighting the helmet with its character, versus isolating the helmet from its associated character._
 
-The LODs section also allows you to define when to use your LODs through the **Screen Size** option (how close or far away should the LOD become visible) as well as defining the amount of **LODHysteresis** to avoid "flickering" when on an LOD boundary. 
-Inside the **Viewport**, you can click the **Preview LODs** button (below) to switch between the available LODs (it may also help to enable the **Display Info** for statistical information). 
+The LODs section also allows you to define when to use your LODs with the **Screen Size** option (how near or far the LOD should become visible), as well as defining the amount of **LODHysteresis** (flickering) to avoid when you're on an LOD boundary. 
+Inside of the **Viewport**, you can click the **Preview LODs** button to switch between the available LODs (it may also help to enable **Display Info** for statistical information). 
 
 ![](previewLODs.png)
 
-Below using the LOD Previewer, we can see the visual difference between the LODs. 
+Using the LOD Previewer, you can inspect the visual differences between LODs. Also, you can inspect statistical information about each LOD. For example, the amount of triangles in LOD 2 is almost half the amount of triangles in LOD 0. 
 
 [OBJECT:ComparisonSlider]
 	[PARAM:before]
@@ -68,9 +73,7 @@ Below using the LOD Previewer, we can see the visual difference between the LODs
 	[/PARAM]
 [/OBJECT]
 
-With our stats being displayed we can also see statistical information about each LOD. Above, the amount of triangles in LOD 2 is almost half the amount of triangles in LOD 0. 
-
-You can also turn on/off the ability for sections within an LOD to **Cast Shadows** or not. 
+You can also enable or disable shadow casting. 
 
 [OBJECT:ComparisonSlider]
 	[PARAM:before]
@@ -81,81 +84,103 @@ You can also turn on/off the ability for sections within an LOD to **Cast Shadow
 	[/PARAM]
 [/OBJECT]
 
-You can turn on/off casting shadows on a per element, per LOD basis (some areas you may not need to cast shadows). 
+_You can turn on/off casting shadows on a per element, per LOD basis (some areas might not require shadow casting)._ 
 
 ## Clothing
 
-The **Clothing** section allows you to add an [](Engine/Physics/Apex) clothing file to associate with the mesh. 
-
-Click the **Add APEX clothing file...** button to begin importing an .apb or .apx file which follows the same pipeline as importing an FBX file. 
+The **Clothing** section associates an [](Engine/Physics/Apex) clothing file with the mesh. Click the **Add APEX clothing file...** button to begin importing an .apb or .apx file, 
+which uses the same pipeline as importing an FBX file. 
 
 ![](AddApex.png)
 
-Once you've imported your Apex file, you can adjust the [Physics Properties](Engine/Physics/PhAT/Reference) to achieve your desired results. 
+Once you've imported your Apex file, you can adjust its [Physics Properties](Engine/Physics/PhAT/Reference) to achieve your desired results. 
 
 ![](OwenApexCloth.png)
 
-Above our character has two Apex files associated with it, one for the character's tie and the other for the character's coat. 
+_For example, this character has two Apex files associated with it; one for the character's tie, and the other for the character's coat._
 
 ## Mesh
 
-![](MeshSection.png)
+In the **Mesh** section, you can see what [Skeleton](Engine/Animation/Skeleton) asset the mesh is using and navigate to it with the **Content Browser**. 
+There are also options for setting **Positive Bounds Extension** and **Negative Bounds Extension** values, which can be used to offset the bounds values of the mesh when they're being imported. 
+The other option, **Use Full Precision UVs** can be turned on to use 32 bit UVs (or left off to save memory, thereby using 16 bit UVs). 
 
-In the **Mesh** section, you can see what [Skeleton](Engine/Animation/Skeleton) asset the mesh is using and browse to it in the **Content Browser**. 
-There are also options for setting **Positive Bounds Extension** and **Negative Bounds Extension** values which you can use to offset the bounds values of the mesh when you import it. 
-The other option, **Use Full Precision UVs** can be turned on to use 32 bit UVs (or left off to use 16 bit UVs to save memory). 
+![](MeshSection.png)
 
 ## Physics
 
-![](PhysicsSection.png)
+The **Physics** section allows you to **Enable Per Poly Collision** for the mesh, which uses skinned data for collision information. Per Poly Collision cannot be used for simulation. In most cases, you're better off using a **Physics Asset**. 
+There is also the ability to assign a [Physics Asset](Engine/Physics/PhAT) to the Skeletal Mesh, which is used for processing physics and collision information. This is used for per-bone hit detection, accurate bounding box calculation, and ragdoll physics. 
 
-The **Physics** section allows you to **Enable Per Poly Collision** for the mesh which uses skinned data for collision data. Per Poly Collision cannot be used for simulation, in most cases, you are better off using a **Physics Asset** instead. 
-There is also the ability to assign a [Physics Asset](Engine/Physics/PhAT) to the Skeletal Mesh which is used for processing physics and collision information. This is used for per-bone hit detection, accurate bounding box calculation, and ragdoll physics. 
+![](PhysicsSection.png)
 
 ## Lighting
 
-![](Lighting.png)
+The **Shadow Physics Asset** option in the Lighting section allows you to specify a Physics Asset, with shapes that will be used for shadowing when components have their **Cast Character Capsule Direct Shadow** or **Cast Character Capsule Indirect Shadow** option(s) enabled. 
+Only spheres and sphyl shapes in the Physics Asset are supported. Additionally, it's imporant to remember that capsule shadows can result in higher rendering costs when more shapes are used. 
 
-The **Shadow Physics Asset** option in the Lighting section allows you to specify a Physics Asset whose shapes will be used for shadowing when components have their **Cast Character Capsule Direct Shadow** or **Cast Character Capsule Indirect Shadow** options enabled. 
-Only spheres and sphyl shapes in the Physics Asset are supported and the more shapes used, the higher the cost the capsule shadows will be. 
+![](Lighting.png)
 
 In-game, indirect capsule shadows are a good way to ground characters to the environment: 
 
 ![](IndirectShadows.png)(w:840)
 
-Direct capsule shadows can be used to achieve extremely soft character shadows in an efficient baked lighting environment. 
+Direct capsule shadows can be used to achieve soft character shadows in an efficient baked lighting environment. 
 
 ![](DirectShadows.png)(w:840)
 
-The light's **Source Radius** or **Source Angle** determines how soft the shadows will be. Capsule shadowing is very efficient as it computes shadowing at half resolution with a depth-aware upsample and uses screen tile culling to limit the shadowing work to where it is needed. 
+The light's **Source Radius** or **Source Angle** determines how soft the shadows will be. Capsule shadowing is very efficient because it computes shadowing at half-resolution with a depth-aware up-sample, using screen tile culling to limit the shadowing work (where it's needed). 
 The GPU cost is proportional to the number of capsules and the number of pixels affected by the casted shadow. 
 
 ## Skeletal Mesh
 
+In this section, you can assign a **Post Process Anim Blueprint**, which is an [Animation Blueprint](Engine/Animation/AnimBlueprints) class, to run as a Post Process for the mesh. This Blueprint will be run before Physics, but after the Main Anim Instance for any [Skeletal Mesh Component](Engine/Components/SkeletalMesh/) using the mesh. 
+
 ![](SkeletalMesh.png)
 
-In this section, you can assign a **Post Process Anim Blueprint** which is an [Animation Blueprint](Engine/Animation/AnimBlueprints) class to run as a Post Process for the mesh. This Blueprint will be run before Physics, but after the Main Anim Instance for any Skeletal Mesh component using the mesh. 
-To assign an Animation Blueprint, click the drop-down menu and specify your desired Animation Blueprint. 
+This allows you to set up [](Engine/Animation/NodeReference/SkeletalControls/AnimDynamics/) or other controllers that will always be applied to the associated mesh, whether that mesh is viewed with the Animation Tools, as part of a Sequencer cinematic, or through a Skeletal Mesh Actor placed in a level. 
+This is especially useful for complex characters that run multiple Animation Blueprints; because any Anim Dynamics, controllers, and IK (among other Anim Blueprint features) will be related to the mesh and not have to be duplicated in every Animation Blueprint intended to be used with the mesh. 
 
-You can also assign **Asset User Data** from the expanded menu option and click the **+** sign to add elements to the array. Any additions will be stored along with the Skeletal Mesh asset. 
+To assign a Post Process Anim Blueprint, select the drop-down menu and specify your default Animation Blueprint. 
+
+![](SelectPostProcessAnimBP.png)
+
+The example video illustrates how the Post Process Anim Blueprint is run during gameplay on a Skeletal Mesh. 
+
+[OBJECT:EmbeddedVideo]
+[PARAMLITERAL:width]
+640
+[/PARAMLITERAL]
+[PARAMLITERAL:height]
+360
+[/PARAMLITERAL]
+[PARAMLITERAL:videoid]
+uroFU3M_-C4
+[/PARAMLITERAL]
+[/OBJECT]
+
+If you have a more complex animation setup that uses multiple Animation Blueprints, you may need to include a **Sub-Graph Input** node to the Post Process Anim Blueprint. 
+
+![](SubGraphInput.png)
+
+_This will allow you to get the pose information from a parent Animation Blueprint (if applicable) rather than using the Skeletal Mesh's base pose._
 
 ## Transform
 
-![](TransformSection.png)
+The **Transform** section inside of the **Skeletal Mesh Editor** allows you to offset the mesh from the world origin (when imported) to match its position in the source art program through the **Import Translation** option. 
+**Import Rotation** will rotate the mesh (when imported) to match its position in the source art program, while the **Import Uniform Scale** option will import the uniform scaling of the mesh from the source art program. 
 
-The **Transform** section inside the Skeletal Mesh Editor allows you to offset the mesh from the world origin when imported to match its position in the source art program through the **Import Translation** option. 
-The **Import Rotation** will rotate the mesh when imported to match its position in the source art program while the **Import Uniform Scale** option will import the uniform scaling of the mesh from the source art program. 
+![](TransformSection.png)
 
 ## Thumbnail
 
-![](ThumbnailSection.png)
+**Thumbnail** options allow you to modify how the mesh is framed in the thumbnail that appears in the **Content Browser** or **Editor Toolbar**. 
 
-The **Thumbnail** options allow you to modify how the mesh is framed in the thumbnail that appears in the **Content Browser** or **Editor Toolbar**. 
+![](ThumbnailSection.png)
 
 ## Import Settings
 
-![](ImportSettingsSection.png)
-
 The **Import Settings** section allows you to modify any of the settings used during the initial [FBX Skeletal Mesh Import Process](https://docs.unrealengine.com/latest/INT/Engine/Content/FBX/ImportOptions/#skeletalmeshoptions) for the mesh. 
 
+![](ImportSettingsSection.png)
 
diff --git a/Engine/Documentation/Source/Engine/Animation/Persona/MorphTargetPreviewer/MorphTargetPreviewer.INT.udn b/Engine/Documentation/Source/Engine/Animation/Persona/MorphTargetPreviewer/MorphTargetPreviewer.INT.udn
--- a/Engine/Documentation/Source/Engine/Animation/Persona/MorphTargetPreviewer/MorphTargetPreviewer.INT.udn
+++ b/Engine/Documentation/Source/Engine/Animation/Persona/MorphTargetPreviewer/MorphTargetPreviewer.INT.udn
@@ -1,6 +1,6 @@
 Availability:Public
 Title:Morph Target Previewer
-Crumbs: %ROOT%, Engine, Engine/Animation, Engine/Animation/Persona
+Crumbs:
 Description:User guide for the editing modes available in the Animation Editor.
 Related: Engine/Content/FBX/MorphTargets
 Related: Engine/Content/Types/SkeletalMeshes
@@ -14,13 +14,16 @@ tags:Animation
 tags:Morph Target
 tags:Skeletal Mesh
 tags:FBX Importing
+SkillLevel:Intermediate
+Parent:Engine/Animation/Persona
+Order:
 
-<!--
 [TOC (start:2)]
--->
 
 [EXCERPT:Overview]
 
+The **Morph Target Previewer** previews any **Morph Targets** (sometimes called "morphs" or "blend shapes") that are applied to a **Skeletal Mesh**. Each **Morph Target** is additively blended into the existing **Skeletal Mesh** geometry. Multiple morph targets can be combined to create a complex vertex-driven animation, which is often ideal for such things as facial expressions.
+
 [OBJECT:EmbeddedVideo]
 [PARAMLITERAL:width]
 640
@@ -33,16 +36,14 @@ YoEQYVEEcKo
 [/PARAMLITERAL]
 [/OBJECT]
 
-The **Morph Target Previewer** allows you to preview any morph targets (also known as "morphs" and "blend shapes", depending on what 3D animation package you are using) that are applied to your Skeletal Meshes. Each Morph Target is additively blended into the existing Skeletal Mesh geometry. This allows for multiple morph targets to be combined to create complex vertex-driven animation, which is often ideal for such things as facial expressions.
 [/EXCERPT:Overview]
 
-The **Morph Target Previewer** window is visible by default whenever the [](Engine\Animation\Persona\Modes\Mesh) is open. 
+The **Morph Target Previewer** window is visible by default whenever the [](Engine\Animation\Persona\Modes\Mesh) is open.
 
 [REGION:tip]
-To see an example of a Morph Target applied to a character, refer to the [Animation Content Examples](Resources\ContentExamples\Animation) page under section 1.6.
+To see an example of a Morph Target applied to a character, refer to the [Animation Content Examples](Resources\ContentExamples\Animation) page.
 [/REGION]
 
-
 ## Interface
 
 The **Morph Target Previewer** panel has two main sections: 
@@ -52,15 +53,14 @@ The **Morph Target Previewer** panel has two main sections:
 1. Search Bar
 1. Morph Target List
 
-The Search Bar allows you to filter the list of Morph Targets to quickly find the one you need. To filter, type the first few letters of the desired target and the list will begin filtering.
-
-You can also **Right-click** on a Morph Target which will pop-up a dialog box with additional actions such as **Delete** or **Copy Morph Target Names**. 
+The Search Bar filters the list of **Morph Targets** to quickly find the one you need. Type the first few letters of the desired target and the list will begin filtering. 
+You can also right-click on a **Morph Target**, which will pop-up a dialog box with additional actions, such as **Delete** or **Copy Morph Target Names**. 
 
 ![](RightClickDelete.png)
 
 ## Creating Morph Targets
 
-You can import Morph Targets as a part of a Skeletal Mesh, as well as independently of a given mesh. This is supported via the FBX file format. 
+You can import **Morph Targets** as a part of a **Skeletal Mesh**, as well as independently (from a given mesh). This is supported via the FBX file format. 
 
 [REGION:note]
 For more information on the setup procedure, as well as how to import Morph Targets into Unreal, please see [](Engine/Content\FBX\MorphTargets).
@@ -76,18 +76,35 @@ Once your Morph Targets are in place, you will need to set up your [Animation Bl
 | --- | --- |
 | [REGION:tablesection] Input Pins [/REGION] ||
 | Execution | An execution wire to trigger the effect of the node. |
-| **Target** | The target Skeletal Mesh. For most circumstances, this will point to "self". |
-| **Morph Target Name** | The name of the Morph Target being edited. |
-| **Value** | A float value (between 0 and 1) used to set the value of the Morph Target being edited. |
+| Target | The target **Skeletal Mesh**. For most circumstances, this will point to "self". |
+| Morph Target Name | The name of the **Morph Target** being edited. |
+| Value | A float value (between 0.0 and 1.0) used to set the value of the **Morph Target** being edited. |
 | [REGION:tablesection] Output Pins [/REGION] ||
 | Execution | Passes execution on to the next node. |
 
+## Morph Target Debug View Mode
+
+You can enable this view mode to easily see what vertices are affected by each Morph Target. 
+
+[OBJECT:EmbeddedVideo]
+	[PARAMLITERAL:width]
+	640
+	[/PARAMLITERAL]
+	[PARAMLITERAL:height]
+	360
+	[/PARAMLITERAL]
+	[PARAMLITERAL:videoid]
+	yOV3mmR-sX4
+	[/PARAMLITERAL]
+[/OBJECT]
+
+1. In the viewport window, click on **Show** > **Mesh Overlay** >,  and then **Selected MorphTarget Vertices**.
+1. Now select any **Morph Target** from the **Morph Target Preview** panel to see the debug view.
 
 ## Optimization
 
-For target platforms that support Shader Model 5, you can enable Morph Targets being calculated on the GPU. This means that the CPU does not have to perform the calculations.
+For target platforms that support Shader Model 5, you can enable the GPU calculation of morph targets. This means that the CPU does not have to perform the calculations if your game is CPU bound, and you have GPU processing to spare. This feature can be activated in the **Project Settings**, as follows:
 
-To enable this setting:
 1. On the file menu, go to **Edit** > **Project Settings**.
 1. Open the **Rendering** section of **Project Settings**.
 1. In the **Optimizations** category, find the **Use GPU for computing morph targets** checkbox and enable it. 
@@ -96,5 +113,6 @@ To enable this setting:
 [![](ProjectSettingsMorphGPU.png)(w:600)](ProjectSettingsMorphGPU.png)
 [/REGION]
 
-
-
+[REGION:caption]
+Click for full image.
+[/REGION]
\ No newline at end of file
diff --git a/Engine/Documentation/Source/Engine/Animation/Persona/VirtualBones/VirtualBones.INT.udn b/Engine/Documentation/Source/Engine/Animation/Persona/VirtualBones/VirtualBones.INT.udn
new file mode 100644
--- /dev/null
+++ b/Engine/Documentation/Source/Engine/Animation/Persona/VirtualBones/VirtualBones.INT.udn
@@ -0,0 +1,75 @@
+Availability:Docs
+Title:Virtual Bones
+Crumbs:
+Description:Overview of Virtual Bones that enables you to add IK or Anim constraint joints inside of the Editor.
+Related: Engine/Content/FBX/Animations
+Related: Engine/Animation/Persona/Modes
+Related: Engine/Content/Types/SkeletalMeshes
+version: 4.15
+parent: Engine/Animation
+type:reference
+topic-image: 
+reference-image:WalkAnim.gif
+tags: Animations, Skeletal Mesh, Skeleton, Bones
+order:
+
+
+[REGION:fullwidth raw]
+![Virtual Bones](VirtualBones.png)
+[/REGION]
+
+Inside the [Animation Tools](Engine/Animation/Persona/Modes), namely the Skeleton or Animation Editors, you can add **Virtual Bones** to the Skeleton of your mesh. 
+These serve to help improve iteration time by enabling you to change your target joint hierarchy for IK or aim inside of the Editor. Before this feature was added, you would 
+have to do this outside of the engine to add them to your asset in your modeling application and then reimport all the animations in order to fix the animation data with the 
+new joints included. 
+
+Virtual Bones, by their very nature, are not skinnable, but, instead, constrained between two existing bones on the Skeleton and they automatically have data generated for them 
+for each animation on the Skeleton. For example, you can add a joint that is a child of a hand, but is constrained to the palm joint. Unlike sockets, this joint can be used 
+in Animation Blueprints as a target, such as an IK target or look at target, or you could even modify them in the Animation Blueprint for later use. 
+
+For additonal information and suggestions for pratical usage of Virtual Bones, like how they have been used to make it easier to retarget or change reference frames for controllers
+and also used for orientation and slope warping. Check out the Animation Techniques used in Paragon training stream for more information. 
+
+[OBJECT:EmbeddedVideo]
+	[PARAMLITERAL:width]
+	640
+	[/PARAMLITERAL]
+	[PARAMLITERAL:height]
+	360
+	[/PARAMLITERAL]
+	[PARAMLITERAL:videoid]
+	1UOY-FMm-xo
+	[/PARAMLITERAL]
+[/OBJECT]
+
+## Adding Virtual Bones to your Skeleton
+
+To add Virtual Bones to your Skeletal Mesh:
+
+1. In the **Content Browser**, open any Skeletal Mesh or Animation Asset. 
+ 
+	![](CB_Skeleton.png)(w:350) 
+1. In the **Skeletal Editor** or **Animation Editor**, select any Bone from the joint hierarchy from the **Skeleton Tree** window. Now, right-click and hover over **Add Virtual Bone** from the context menu. 
+ Select a bone from the **Pick Target Bone** window that opens. 
+   
+	[REGION:raw]
+	![](SelectVirtualBone.png)(w:850)
+	[/REGION]
+
+
+Now that you've added your Virtual Bone(s), you will see it listed in the Skeleton hierarchy. These are indicated by the prefix "VB" in their name and their lavender color in the Skeleton Tree window. 
+
+[REGION:lightbox]
+[![](VirtualBoneView.png)(w:550)](VirtualBoneView.png)
+[/REGION]
+
+[REGION:tip]
+If you do not see your bones in the Viewport of the Skeleton or Animation Editors, go to **Show** > **Bones** > **All Hierarchy** so that your Skeleton is always visible. There are also other options here
+to fit your preference if you do not want to always see all the bones in the viewport.
+[/REGION]
+
+In the Animation Editor, you will see your Virtual Bone(s) constrained between the selected joints and that it updates with your animation. 
+ 
+[REGION:asyncgif]
+![Virtual Bone](WalkAnim.gif)(w:350)
+[/REGION]
diff --git a/Engine/Documentation/Source/Engine/Animation/RetargetingDifferentSkeletons/RetargetingDifferentSkeletons.INT.udn b/Engine/Documentation/Source/Engine/Animation/RetargetingDifferentSkeletons/RetargetingDifferentSkeletons.INT.udn
--- a/Engine/Documentation/Source/Engine/Animation/RetargetingDifferentSkeletons/RetargetingDifferentSkeletons.INT.udn
+++ b/Engine/Documentation/Source/Engine/Animation/RetargetingDifferentSkeletons/RetargetingDifferentSkeletons.INT.udn
@@ -1,4 +1,4 @@
-Availability:Docs
+Availability:Public
 Title: Animation Retargeting (Different Skeletons)
 Crumbs: %ROOT%, Engine, Engine/Animation
 Description:Retargeting animations for use between different skeletons. 
diff --git a/Engine/Documentation/Source/Engine/Audio/Dialogue/2/Dialogue_2.INT.udn b/Engine/Documentation/Source/Engine/Audio/Dialogue/2/Dialogue_2.INT.udn
--- a/Engine/Documentation/Source/Engine/Audio/Dialogue/2/Dialogue_2.INT.udn
+++ b/Engine/Documentation/Source/Engine/Audio/Dialogue/2/Dialogue_2.INT.udn
@@ -16,7 +16,7 @@ related:Engine/Audio/SoundCues/NodeReference
 [NAV]
 
 
-For this example, we are just going to implement a greeting from Quest NPC to Adam and Zoe. For more information about complex dialogue, see the [Sound Cues and Dialogue section](#soundcuesanddialogue).
+For this example, we are just going to implement a greeting from Quest NPC to Adam and Zoe. For more information about complex dialogue, see the [Sound Cues and Dialogue](Engine/Audio/Dialogue/2) section.
 Each line of dialogue needs a **Dialogue Wave** asset associated with it.
 
 1. Go back to the **Content Browser**, and create a new **Dialogue Wave** asset. 
diff --git a/Engine/Documentation/Source/Engine/Basics/DerivedDataCache/DerivedDataCache.INT.udn b/Engine/Documentation/Source/Engine/Basics/DerivedDataCache/DerivedDataCache.INT.udn
--- a/Engine/Documentation/Source/Engine/Basics/DerivedDataCache/DerivedDataCache.INT.udn
+++ b/Engine/Documentation/Source/Engine/Basics/DerivedDataCache/DerivedDataCache.INT.udn
@@ -24,22 +24,61 @@ not do any harm, but it is a waste of time as it generally takes longer to trans
 in the DDC than it would to simply generate it from scratch locally.
 [/REGION]
 
-To set up a shared DDC, override the paths for the [DerivedDataBackendGraph] by declaring this 
+To set up a shared DDC, override the paths for the `[InstalledDerivedDataBackendGraph]` or `[DerivedDataBackendGraph]` (depending on whether you use the binary install from the Epic Game Launcher or compile from source) by declaring this 
 section in your game's **DefaultEngine.ini**. This section is originally declared in the **BaseEngine.ini** 
 where the paths are set to Epic's internal DDC share. Re-declare the section and change the paths to point 
 to a share on your network (e.g. `Path=\\mystudio.net\DDC`):
 
-	[InstalledDerivedDataBackendGraph]
-	MinimumDaysToKeepFile=7
-	Root=(Type=KeyLength, Length=120, Inner=AsyncPut)
-	AsyncPut=(Type=AsyncPut, Inner=Hierarchy)
-	Hierarchy=(Type=Hierarchical, Inner=Boot, Inner=Pak, Inner=EnginePak, Inner=Local, Inner=Shared)
-	Boot=(Type=Boot, Filename=\%GAMEDIR\%DerivedDataCache/Boot.ddc, MaxCacheSize=256)
-	Local=(Type=FileSystem, ReadOnly=false, Clean=false, Flush=false, PurgeTransient=true, DeleteUnused=true, UnusedFileAge=17, FoldersToClean=-1, Path=../../../Engine/DerivedDataCache)
-	Shared=(Type=FileSystem, ReadOnly=false, Clean=false, Flush=false, DeleteUnused=true, UnusedFileAge=19, FoldersToClean=-1, Path=\\mystudio.net\DDC, EnvPathOverride=UE-SharedDataCachePath)
-	AltShared=(Type=FileSystem, ReadOnly=true, Clean=false, Flush=false, DeleteUnused=true, UnusedFileAge=19, FoldersToClean=-1, Path=\\mystudio.net\DDC2, EnvPathOverride=UE-SharedDataCachePath2)
-	Pak=(Type=ReadPak, Filename=\%GAMEDIR\%DerivedDataCache/DDC.ddp)
-	EnginePak=(Type=ReadPak, Filename=../../../Engine/DerivedDataCache/DDC.ddp)
+%Globals:BinarySourceSelectVar%
+
+[OBJECT:ToggleButtonContent]
+	[PARAMLITERAL:category]
+	BinarySourceSelect
+	[/PARAMLITERAL]
+	[PARAMLITERAL:id]
+	binaryinstall
+	[/PARAMLITERAL]
+	[PARAMLITERAL:active]
+	active_button_content
+	[/PARAMLITERAL]
+	[PARAM:content]
+			[InstalledDerivedDataBackendGraph]
+			MinimumDaysToKeepFile=7
+			Root=(Type=KeyLength, Length=120, Inner=AsyncPut)
+			AsyncPut=(Type=AsyncPut, Inner=Hierarchy)
+			Hierarchy=(Type=Hierarchical, Inner=Boot, Inner=Pak, Inner=EnginePak, Inner=Local, Inner=Shared)
+			Boot=(Type=Boot, Filename=\%GAMEDIR\%DerivedDataCache/Boot.ddc, MaxCacheSize=256)
+			Local=(Type=FileSystem, ReadOnly=false, Clean=false, Flush=false, PurgeTransient=true, DeleteUnused=true, UnusedFileAge=17, FoldersToClean=-1, Path=../../../Engine/DerivedDataCache)
+			Shared=(Type=FileSystem, ReadOnly=false, Clean=false, Flush=false, DeleteUnused=true, UnusedFileAge=19, FoldersToClean=-1, Path=\\mystudio.net\DDC, EnvPathOverride=UE-SharedDataCachePath)
+			AltShared=(Type=FileSystem, ReadOnly=true, Clean=false, Flush=false, DeleteUnused=true, UnusedFileAge=19, FoldersToClean=-1, Path=\\mystudio.net\DDC2, EnvPathOverride=UE-SharedDataCachePath2)
+			Pak=(Type=ReadPak, Filename=\%GAMEDIR\%DerivedDataCache/DDC.ddp)
+			EnginePak=(Type=ReadPak, Filename=../../../Engine/DerivedDataCache/DDC.ddp)
+	[/PARAM]
+[/OBJECT]
+[OBJECT:ToggleButtonContent]
+	[PARAMLITERAL:category]
+	BinarySourceSelect
+	[/PARAMLITERAL]
+	[PARAMLITERAL:id]
+	sourceinstall
+	[/PARAMLITERAL]
+	[PARAMLITERAL:active]
+	[/PARAMLITERAL]
+	[PARAM:content]
+			[DerivedDataBackendGraph]
+			MinimumDaysToKeepFile=7
+			Root=(Type=KeyLength, Length=120, Inner=AsyncPut)
+			AsyncPut=(Type=AsyncPut, Inner=Hierarchy)
+			Hierarchy=(Type=Hierarchical, Inner=Boot, Inner=Pak, Inner=EnginePak, Inner=Local, Inner=Shared)
+			Boot=(Type=Boot, Filename=\%GAMEDIR\%DerivedDataCache/Boot.ddc, MaxCacheSize=256)
+			Local=(Type=FileSystem, ReadOnly=false, Clean=false, Flush=false, PurgeTransient=true, DeleteUnused=true, UnusedFileAge=17, FoldersToClean=-1, Path=../../../Engine/DerivedDataCache)
+			Shared=(Type=FileSystem, ReadOnly=false, Clean=false, Flush=false, DeleteUnused=true, UnusedFileAge=19, FoldersToClean=-1, Path=\\mystudio.net\DDC, EnvPathOverride=UE-SharedDataCachePath)
+			AltShared=(Type=FileSystem, ReadOnly=true, Clean=false, Flush=false, DeleteUnused=true, UnusedFileAge=19, FoldersToClean=-1, Path=\\mystudio.net\DDC2, EnvPathOverride=UE-SharedDataCachePath2)
+			Pak=(Type=ReadPak, Filename=\%GAMEDIR\%DerivedDataCache/DDC.ddp)
+			EnginePak=(Type=ReadPak, Filename=../../../Engine/DerivedDataCache/DDC.ddp)
+	[/PARAM]
+[/OBJECT]
+
 
 
 Satellite studios working out of a single shared code base can set the **UE-SharedDataCachePath** environment 
diff --git a/Engine/Documentation/Source/Engine/Basics/Projects/Packaging/PackagingProjects.INT.udn b/Engine/Documentation/Source/Engine/Basics/Projects/Packaging/PackagingProjects.INT.udn
--- a/Engine/Documentation/Source/Engine/Basics/Projects/Packaging/PackagingProjects.INT.udn
+++ b/Engine/Documentation/Source/Engine/Basics/Projects/Packaging/PackagingProjects.INT.udn
@@ -1,8 +1,9 @@
 Availability:Public
 Title:Packaging Projects
-Crumbs: %ROOT%, Engine, Engine/Basics/Projects
+Crumbs:
 Description:Packaging Unreal game projects for distribution.
 version: 4.9
+parent:Engine/Basics/Projects
 
 [TOC(start:2)]
 
@@ -78,7 +79,7 @@ The number and type of files created will vary depending on the target platform.
 	| Configuration | EXE Name                             | EXE Location                                                    |
 	| ------------- | ------------------------------------ | --------------------------------------------------------------- |
 	| Development   | [ProjectName].exe                    | [PackageDirectory]\WindowsNoEditor |
-	| Shipping      | [ProjectName]-Win32-Shipping.exe.exe | [PackageDirectory]\WindowsNoEditor |
+	| Shipping      | [ProjectName]-Win32-Shipping.exe     | [PackageDirectory]\WindowsNoEditor |
 
 1. **Double-click** the executable to run the game.
 [/EXCERPT:RunCooked]
diff --git a/Engine/Documentation/Source/Engine/Blueprints/Anatomy/Anatomy.INT.udn b/Engine/Documentation/Source/Engine/Blueprints/Anatomy/Anatomy.INT.udn
--- a/Engine/Documentation/Source/Engine/Blueprints/Anatomy/Anatomy.INT.udn
+++ b/Engine/Documentation/Source/Engine/Blueprints/Anatomy/Anatomy.INT.udn
@@ -5,6 +5,7 @@ Description: The User Guide is the go-to source to learn the different parts of
 version: 4.9
 parent:Engine/Blueprints
 order:3
+tags:Blueprints
 
 [VAR:Topic]
 [OBJECT:Topic]
diff --git a/Engine/Documentation/Source/Engine/Blueprints/BP_HowTo/ActorReference/ActorReference.INT.udn b/Engine/Documentation/Source/Engine/Blueprints/BP_HowTo/ActorReference/ActorReference.INT.udn
--- a/Engine/Documentation/Source/Engine/Blueprints/BP_HowTo/ActorReference/ActorReference.INT.udn
+++ b/Engine/Documentation/Source/Engine/Blueprints/BP_HowTo/ActorReference/ActorReference.INT.udn
@@ -9,6 +9,7 @@ type:how-to
 parent:Engine/Blueprints/BP_HowTo
 order:1
 tags:Data Types
+tags:Blueprints
 
 [INCLUDE:Gameplay\HowTo\ReferenceAssets#referenceoverview]
 
diff --git a/Engine/Documentation/Source/Engine/Blueprints/BP_HowTo/PlacingNodes/PlacingNodes.INT.udn b/Engine/Documentation/Source/Engine/Blueprints/BP_HowTo/PlacingNodes/PlacingNodes.INT.udn
--- a/Engine/Documentation/Source/Engine/Blueprints/BP_HowTo/PlacingNodes/PlacingNodes.INT.udn
+++ b/Engine/Documentation/Source/Engine/Blueprints/BP_HowTo/PlacingNodes/PlacingNodes.INT.udn
@@ -10,7 +10,7 @@ skilllevel:Beginner
 parent:Engine/Blueprints/BP_HowTo
 type:how-to
 order:2
-
+tags:Blueprints
 
 On this page, you will learn how you can place nodes on your graphs in Blueprints. 
 
diff --git a/Engine/Documentation/Source/Engine/Blueprints/Editor/Modes/Components/ComponentsPanel.INT.udn b/Engine/Documentation/Source/Engine/Blueprints/Editor/Modes/Components/ComponentsPanel.INT.udn
--- a/Engine/Documentation/Source/Engine/Blueprints/Editor/Modes/Components/ComponentsPanel.INT.udn
+++ b/Engine/Documentation/Source/Engine/Blueprints/Editor/Modes/Components/ComponentsPanel.INT.udn
@@ -1,7 +1,8 @@
 Availability:Public
 Title: Components Mode
 Crumbs:%ROOT%, Engine/Blueprints, Engine/Blueprints/Editor
-Description: Allows you to add Components to your Blueprint, manage the details of or arrange existing Components. 
+Description: Allows you to add Components to your Blueprint, manage the details of or arrange existing Components.
+tags:Blueprints
 
 [INCLUDE:Globals#bpmodesremoved]
 
diff --git a/Engine/Documentation/Source/Engine/Blueprints/Editor/Modes/Defaults/BlueprintDefaults.INT.udn b/Engine/Documentation/Source/Engine/Blueprints/Editor/Modes/Defaults/BlueprintDefaults.INT.udn
--- a/Engine/Documentation/Source/Engine/Blueprints/Editor/Modes/Defaults/BlueprintDefaults.INT.udn
+++ b/Engine/Documentation/Source/Engine/Blueprints/Editor/Modes/Defaults/BlueprintDefaults.INT.udn
@@ -1,7 +1,8 @@
 Availability:Public
 Title: Defaults Mode
 Crumbs:%ROOT%, Engine/Blueprints, Engine/Blueprints/Editor
-Description:The Defaults Mode allows you to edit all the default values of the properties inside your Blueprint. 
+Description:The Defaults Mode allows you to edit all the default values of the properties inside your Blueprint.
+tags:Blueprints
 
 [INCLUDE:Globals#bpmodesremoved]
 
diff --git a/Engine/Documentation/Source/Engine/Blueprints/Editor/Modes/GraphPanel/GraphPanel.INT.udn b/Engine/Documentation/Source/Engine/Blueprints/Editor/Modes/GraphPanel/GraphPanel.INT.udn
--- a/Engine/Documentation/Source/Engine/Blueprints/Editor/Modes/GraphPanel/GraphPanel.INT.udn
+++ b/Engine/Documentation/Source/Engine/Blueprints/Editor/Modes/GraphPanel/GraphPanel.INT.udn
@@ -1,7 +1,8 @@
 Availability:Public
 Title: Graph Mode
 Crumbs:%ROOT%, Engine/Blueprints, Engine/Blueprints/Editor
-Description: Graph Mode allows you visually layout the network of nodes within a Blueprint, providing its functionality. 
+Description: Graph Mode allows you visually layout the network of nodes within a Blueprint, providing its functionality.
+tags:Blueprints
 
 
 [VAR:Topic]
diff --git a/Engine/Documentation/Source/Engine/Blueprints/HowTo/BPHT_6/BPHT_6.INT.udn b/Engine/Documentation/Source/Engine/Blueprints/HowTo/BPHT_6/BPHT_6.INT.udn
--- a/Engine/Documentation/Source/Engine/Blueprints/HowTo/BPHT_6/BPHT_6.INT.udn
+++ b/Engine/Documentation/Source/Engine/Blueprints/HowTo/BPHT_6/BPHT_6.INT.udn
@@ -3,7 +3,7 @@ Title:6. Communicating with the HUD
 Crumbs: %ROOT%, Engine, Engine/Blueprints, Engine/Blueprints/HowTo
 Description:In this step we will update our HUD based on values stored in our Character Blueprint. 
 version: 4.9
-
+tags:Blueprints
 
 [VAR:Steps]
 [OBJECT:Navigation]
diff --git a/Engine/Documentation/Source/Engine/Blueprints/TechnicalGuide/NativizingBlueprints/Advanced/AdvancedNativeBlueprints.INT.udn b/Engine/Documentation/Source/Engine/Blueprints/TechnicalGuide/NativizingBlueprints/Advanced/AdvancedNativeBlueprints.INT.udn
--- a/Engine/Documentation/Source/Engine/Blueprints/TechnicalGuide/NativizingBlueprints/Advanced/AdvancedNativeBlueprints.INT.udn
+++ b/Engine/Documentation/Source/Engine/Blueprints/TechnicalGuide/NativizingBlueprints/Advanced/AdvancedNativeBlueprints.INT.udn
@@ -1,53 +1,94 @@
 Availability:Public
 Title:Advanced Blueprint Nativization
-Crumbs: %ROOT%, Programming, Engine/Blueprints/TechnicalGuide
-Description:Technical guide for gameplay programmers exposing gameplay elements to Blueprints.
+Crumbs:
+Description:Technical guide for programmers exposing gameplay elements to Blueprints.
 Related: Engine/Blueprints/TechnicalGuide/NativizingBlueprints
 Parent: Engine/Blueprints/TechnicalGuide
-Version:4.12
+Version:4.15
 type:reference
+tags:Blueprints
+SkillLevel:Advanced
+prereq:Engine/Blueprints/TechnicalGuide/NativizingBlueprints
 
 [TOC(start:2)]
 
-
 ## Building with Unreal Automation Tool
-To nativize Blueprints while building with Unreal Automation too, add the `-nativizeAssets` switch to your BuildCookRun command line.
 
-## C++ Cooperation
+To nativize Blueprints while building with Unreal Automation Tool from the commandline, add the `-nativizeAssets` switch to your `BuildCookRun` command.
+
+##Header Declaration
+
+Nativized Blueprints are stored in a new, separate module, which means that headers from the gameplay module are now used outside of their original module. 
+You should check that all of your headers have declared all of the classes, and all of the structs, being used; either as an included header or as a forward declaration.
+
+###Fixing Missing Headers
 
-Nativized blueprints are stored in a new, separate module, which means that the headers from gameplay module are now used outside their module. 
-You should check that all headers have declared all classes and all structs that are used, either as a included header, or as a forward declaration.
+Many projects have an "aggregate" header: either `[GameName]/[GameName].h`, or `[GameName]/Public/[GameName].h`. 
+* If you add the path of this header to `FilesToIncludeInModuleHeader`, under `[BlueprintNativizationSettings]` in `DefaultEditor.ini`, you can fix many 
+  issues related to missing headers.
 
-Many projects have an "aggregate" header `[GameName]/[GameName].h` or `[GameName]/Public/[GameName].h`. If you add the path to this header to 
-`FilesToIncludeInModuleHeader` under `[BlueprintNativizationSettings]` in `DefaultEditor.ini`, it can fix many issues related to missing headers.
+##Fostering Code Cooperation
 
-There are some rules/restrictions that help make cooperation between manually written code and the nativized code more reliable and efficient. 
+There are some restrictions that help make cooperation between manually written code and nativized code, reliable and efficient. With this in mind, 
+the following implementations aren't supported: 
 
-The following implementations are not supported: 
-- private/protected variables with `BlueprintVisible` or `Editable` tags
-- private/protected functions that are marked as `BlueprintPure` or `BlueprintCallable`
-- `noexport` structures/classes
-- bitfield variables with `BlueprintVisible` or `Editable` tags
-- private bitfield variables
-- `UBlueprintGeneratedClass`, `UUserDefinedStruct`, `UUserDefinedEnum` variables, or any type based on them
+* Private and Protected Functions Marked as `BlueprintPure` or `BlueprintCallable`
+* `noexport` Structures and Classes
+* Bitfield Variables with `BlueprintVisible` or `Editable` Tags
+* Private Bitfield Variables
+* `UBlueprintGeneratedClass`, `UUserDefinedStruct`, and `UUserDefinedEnum` Variables (or any type based on these variables)
 
 ## Technical Details
 
-Nativizing Blueprints should work with Visual Studio 2013 and Visual Studio 2015, however 2015 is strongly recommended as it handles larger projects better.
-If there is a strange compiler/linker error like LNK1000, LNK1248, then try to reduce the amount of nativized Blueprints by excluding types or specific Blueprints.
+Nativizing Blueprints should work with Visual Studio 2013 and Visual Studio 2015; however, Visual Studio 2015 is strongly recommended because it was developed 
+to handle large projects. If you encounter a compiler or linker error (such as LNK1000 or LNK1248), try reducing the amount of nativized Blueprints from your project,
+by excluding Blueprint types, or specific Blueprints, in the `DefaultEngine.ini` configuration settings file. For more information about excluding Blueprints 
+from nativization, refer to our [](Engine/Blueprints/TechnicalGuide/NativizingBlueprints) documentation.
 
-The `BlueprintCompilerCppBackend` module (found here: `Engine\Source\Developer\BlueprintCompilerCppBackend\`) converts the Blueprint syntax tree into C++. 
-The `BlueprintNativeCodeGen` module (found here: `Engine\Source\Developer\BlueprintNativeCodeGen\`) does a lot of magic while cooking and packaging. 
+If you want to learn how the Blueprint Nativization process works "under-the-hood", read through the following modules:
 
+* The `BlueprintCompilerCppBackend` module (found in `Engine\Source\Developer\BlueprintCompilerCppBackend\`) converts the Blueprint syntax tree into C++. 
+* The `BlueprintNativeCodeGen` module (found in `Engine\Source\Developer\BlueprintNativeCodeGen\`) does a lot of "magic" while cooking and packaging. 
 
 [PUBLISH:Licensee]
+
 ## Debugging Converted Code on Windows
 
-If you are working with a project that already includes some C++ code, you can debug using the existing project solution.
-However, if your project originally contained only Blueprints, you need to execute `GenerateProjectFiles.bat` so that your project will have a solution.  
-Do not use this solution (or the new project under the main UE4 solution) to compile your project or to build an Editor version of your project. 
+If you are working with a project that already includes some C++ code, you can use the existing project solution to debug the code. However, if your project originally contained only Blueprints, you'll need to run the `GenerateProjectFiles` batch file so that your project will have a Visual Studio solution.
+
+[REGION:warning]
+Do not use this solution (or a new project under the main UE4 solution) to compile your project or to build an Editor version of your project. 
+[/REGION]
+
+###Steps
+
+The following steps assume that you've already cooked, built, and staged your project. 
+
+1.	Navigate to the NativizedAssets plugin folder, which should be located in your game's Intermediate folder (for example, `[GameDirectory]\Intermediate\WindowsNoEditor\NativizedAssets`).
+
+	![](NativizedAssetsFolder_Loc.png)(w:863)
+
+1.	Move the NativizedAssets plugin folder into your game's Plugins folder (for example, `[GameDirectory]\Plugins`).
+
+	[REGION:lightbox]
+		[![](NativizedAssetsMove_Loc.png)(w:863)](NativizedAssetsMove_Loc.png)
+	[/REGION]
+	
+	[REGION:caption]
+		Click for full image.
+	[/REGION]
+
+1.	Now, compile UE4 into a Development (or Development-Client, Development-Server) executable configuration, passing `-basedir="[Path to Cooked Game]\Binaries\Win64"` as a command line parameter.
+
+	[REGION:note]
+	Passing `-basedir=""` as a command line parameter points UE4 to your game's cooked build.
+	[/REGION]
+
+At this point, you should be able to debug and modify the Nativized code at your discretion.
+
+[/PUBLISH:Licensee]
 
-To build a debug .exe file:
+<!--Commented out per SME; however, it's being preserved for future review.
 
 1. Run the following lines. They can be found in your output log when your nativized build is packaged, so it's easiest to copy them out to verify that the switches 
 are correct for your environment. Just change the `Development` switch to `Debug`. 
@@ -60,4 +101,4 @@ are correct for your environment. Just change the `Development` switch to `Debug
 
 3. Run debug in Visual Studio, but do not recompile the project from Visual Studio.
 
-[/PUBLISH:Licensee]
\ No newline at end of file
+-->
\ No newline at end of file
diff --git a/Engine/Documentation/Source/Engine/Blueprints/TechnicalGuide/NativizingBlueprints/NativizingBlueprints.INT.udn b/Engine/Documentation/Source/Engine/Blueprints/TechnicalGuide/NativizingBlueprints/NativizingBlueprints.INT.udn
--- a/Engine/Documentation/Source/Engine/Blueprints/TechnicalGuide/NativizingBlueprints/NativizingBlueprints.INT.udn
+++ b/Engine/Documentation/Source/Engine/Blueprints/TechnicalGuide/NativizingBlueprints/NativizingBlueprints.INT.udn
@@ -1,31 +1,79 @@
 Availability:Public
 Title:Nativizing Blueprints
-Crumbs: %ROOT%, Programming, Engine/Blueprints/TechnicalGuide
-Description:Technical guide for gameplay programmers exposing gameplay elements to Blueprints.
-Version: 4.12
+Crumbs:
+Description:Technical guide for programmers exposing gameplay elements to Blueprints.
+Version: 4.15
 Parent: Engine/Blueprints/TechnicalGuide
 type:overview
+tags:Blueprints
+Skilllevel:Intermediate
+Related: Engine/Blueprints/TechnicalGuide/NativizingBlueprints/Advanced
+social-image:social_bp_nativization.png
 
-[REGION:warning]
-This is an experimental beta feature. All Unreal Engine 4 sample games and some medium-size demos were successfully nativized, but we're still working on using the feature with larger projects.
+[TOC(start:1 end:4)]
+
+When teams script gameplay with Blueprints, they're creating new UClasses without having to write or compile native C++ (CPP) code. As a result, Blueprints work well for game teams 
+that don't have the technical expertise of native CPP programmers. The reason that non-programmers can work strictly with Blueprints, is because Blueprint nodes run in 
+a virtual machine (VM), enabling them to call native CPP functions. Unfortunately, relying on a VM that translates Blueprints into native CPP code comes at a cost; and as you can imagine, 
+translating Blueprint nodes into native CPP functions can potentially slow your game's performance per frame. 
+
+When the cost of running Blueprints in a VM isn't noticeable, the VM overhead can be an acceptable cost to bear; however, there are times when you'll want your game to run as fast as 
+possible (especially, when you're shipping to devices that have limited hardware capabilities). This is where Unreal Engine's Blueprint Nativization tool can help. 
+
+[REGION:note]
+Here at Epic, we've successfully used Blueprint Nativization while developing our newest VR title, [Robo Recall](https://www.epicgames.com/roborecall/en-US/home).
 [/REGION]
 
+##Reducing VM Overhead
+
+The main goal of Blueprint Nativization is to reduce VM overhead in the runtime version of your project by generating native CPP code from your project's Blueprints. In general, all 
+of the Blueprint Classes are replaced by the newly generated native CPP classes that will then be compiled with the project during the [cooking process](Engine/Deployment/Cooking). 
+Although you'll be able to read the translated CPP code from a text editor, you'll notice that the code isn't formatted to be reusable (or generated to be "reader friendly").
 
-The new Blueprint Nativization tool enables you to generate native C++ code from your project's Blueprints. Cooking your project out for testing outside the editor or for release will
-compile the project with the generated code. All of the Blueprint Classes are then replaced by the newly-generated native classes. The main goal of Blueprint Nativization is to 
-reduce the virtual machine (VM) overhead in the runtime version of the project. As a result, the generated code is not necessarily designed to be "friendly" or reusable. 
+[OBJECT:ComparisonSlider]
+	[PARAM:before]
+	![Blueprint](Blueprint.png)(w:1200)
+	[/PARAM]
+	[PARAM:after]
+	![Nativized Blueprint](NativizedBlueprint.png)(w:1200)
+	[/PARAM]
+[/OBJECT]
 
-This page will get you started with nativizing your project's Blueprints, and for Blueprint-only projects, this should be a mostly transparent process. If you have C++ in your project as well,
-or want to debug the nativized code, see  [](Engine/Blueprints/TechnicalGuide/NativizingBlueprints/Advanced) for more information.
+_Here, we compare a Player Character Blueprint with its nativized code._
 
+This page will get you started with nativizing your project's Blueprints, and for Blueprint-only projects, this should be a straightforward task. To get started, you'll need to know 
+where to find the user interface (UI) that you'll need to disable and enable Blueprint Nativization.
 
-## Enabling Blueprint Nativization
+##Enabling Blueprint Nativization
 
-To enable Blueprint Nativization, use the **Edit Menu** to open **Project Settings**. Under **Packaging**, enable the **Nativize Blueprint Assets** checkbox.
+When you first create a project, Blueprint Nativization is disabled by default. To locate the Blueprint Nativization settings, open the **Project Settings** menu from Unreal Editor, and select 
+**Packaging** from the left side bar. You should be able to see the **Blueprints** dropdown menu from inside of the **Project - Packaging** menu.
+
+[REGION:lightbox]
+	[![](BPNativization_DisabledUI.png)(w:600)](BPNativization_DisabledUI.png)
+[/REGION]
+
+[REGION:caption]
+	Click for full image.
+[/REGION]
 
-![](BlueprintNativizeCheckbox.png)
+At any time during development, if you choose to disable Blueprint Nativization, select **Disable** from the **Blueprint Nativization Method** dropdown list. 
 
-Then, when you package your project, the following Blueprint assets will be nativized: 
+![](BPNativization_DisabledUI1.png)(w:900)
+
+As you can see from the dropdown list, there are two additional options that are available to you, namely:
+
+* **Inclusive** Nativization
+* **Exclusive** Nativization
+
+##Inclusive Nativization
+
+When you select **Inclusive** nativization, you're electing to nativize all of the supported Blueprint classes in your project. 
+
+![](BPNativization_InclusiveUI.png)(w:900)
+
+After selecting the **Inclusive** nativization method, it's important to note that during the [packaging process](Engine/Basics/Projects/Packaging), 
+Unreal Automation Tool (UAT) will nativize the following supported Blueprint assets:
 
 * Blueprint Classes
 * Blueprint Function Libraries
@@ -35,16 +83,199 @@ Then, when you package your project, the following Blueprint assets will be nati
 * User Defined Structures
 * User Defined Enums
 
-Currently, _Level Blueprints_ will not be converted. The nativized code is stored in `[ProjectDirectory]\Intermediate\[Platform That Was Cooked For]\NativizedAssets\Source\NativizedAssets\`.
+[REGION:note]
+At this time, Level Blueprints and Blueprint Macros aren't supported by the Blueprint Nativization tool, and won't be converted into native CPP code. 
+[/REGION]
+
+###Excluding Blueprints from Conversion
 
-## Excluding Blueprints from Conversion
+When the Blueprint Nativization tool converts Blueprint nodes to native CPP code, the resulting executable is compiled with additional machine instructions, 
+dependencies, and embedded metadata to support the conversion. 
+
+[REGION:warning]
+So, it's important to be aware that when you opt into the **Inclusive** Blueprint Nativization Method, 
+your executable size will grow to a size that may be unacceptable for the platform that you're shipping for.
+[/REGION]
+
+You can mitigate the bloating of your executable's size by specifying which Blueprints you want to exclude from the conversion process. To exclude Blueprints from conversion, 
+open your project's `DefaultEditor.ini` file, and create a `[BlueprintNativizationSettings]` heading with the following configuration properties:
+
+* To exclude a certain type of Blueprint asset, set the `ExcludedBlueprintTypes` configuration property.  
+  For example, `+ExcludedBlueprintTypes=/Script/UMGEditor.WidgetBlueprint`, won't convert any Widget Blueprint assets.
+* To exclude a specific Blueprint asset, set the `ExcludedAssets` configuration property.  
+  For example, `+ExcludedAssets=/Engine/Tutorial/SubEditors/TutorialAssets/Character/TutorialCharacter`, won't convert the **Tutorial Character** Blueprint asset.
+
+![](DefaultEngineINI.png)(w:900)
+
+Note that Blueprint assets are nativized during the packaging process only if they're cooked as part of that process. So, if you haven't specified any Blueprint assets or maps for cooking, Unreal Engine 
+will cook everything that's located inside of your **Content** directory. If, after selecting the **Inclusive** method, you notice that unused Blueprints are being nativized, you should probably opt into 
+the **Exclusive** Blueprint Nativization Method.
+
+##Exclusive Nativization
+
+For larger projects being shipped to devices with limited hardware resources, we recommend that you select the **Exclusive** Blueprint Nativization Method.
+
+![](BPNativization_ExclusiveUI.png)(w:900)
+
+While the **Inclusive** method is available for smaller projects that aren't in danger of bloating the executable beyond the target hardware's capacity to host the application, there are several advantages when choosing to 
+use the **Exclusive** Blueprint Nativization method. Some of the advantages include:
+
+* Limiting the final size of the executable.
+* Actively preventing unused Blueprint assets from being converted.
+* Having the ability to explicitly select which Blueprint assets to convert.
+
+In general, it's recommended that you opt into the Exclusive Blueprint Nativization Method when you have Blueprints that are computationally expensive; including (but, not limited to):
+
+* Blueprints with lots of nodes.
+* Blueprints that require optimized loops.
+* Blueprints that perform complex mathematical operations. 
+
+###Selecting Blueprints to Nativize
+
+[REGION:note]
+If you don't have a project containing Blueprint assets, go ahead and create a new project using one of the available project templates. For illustrative purposes, 
+we're using the **Puzzle** Project Template.
+[/REGION]
 
-You can also exclude certain _Blueprints_ from conversion. In your project's `DefaultEditor.ini` file, create a `[BlueprintNativizationSettings]` heading.
+Perform the following steps to select Blueprint assets that you want to nativize, using the **Exclusive** Blueprint Nativization Method. 
 
-* To exclude a certain type of asset, set `ExcludedBlueprintTypes`. For example, `+ExcludedBlueprintTypes=/Script/UMGEditor.WidgetBlueprint` will not convert any _Widget Blueprints_.
-* To exclude an individual asset, set `ExcludedAssets`. For example, `+ExcludedAssets=/Engine/Tutorial/SubEditors/TutorialAssets/Character/TutorialCharacter` will not convert the **Tutorial Character**.
+1.	First, open the **Edit** menu and select **Project Settings...**.
 
-Assets are nativized if and only if they are cooked. When the cooker has no assets/maps specified to be cooked, it cooks 
-everything from your project's **Content** directory. If you see some unused Blueprints are being nativized, you can also use the list of maps to package to omit assets. 
+	[REGION:lightbox]
+		[![](BPNativization_ExclusiveStep_1.png)(w:900)](BPNativization_ExclusiveStep_1.png)
+	[/REGION]
+	
+	[REGION:caption]
+		Click for full image.
+	[/REGION]
+
+1.	With the **Project Settings** menu open, go ahead and click **Packaging**, opening the **Project Settings - Packaging** menu.
+
+	[REGION:lightbox]
+		[![](BPNativization_ExclusiveStep_2.png)(w:900)](BPNativization_ExclusiveStep_2.png)
+	[/REGION]
+	
+	[REGION:caption]
+		Click for full image.
+	[/REGION]
+
+1.	From the **Blueprints** submenu, click on the **Blueprint Nativization Menu** dropdown list.
+
+	[REGION:lightbox]
+		[![](BPNativization_ExclusiveStep_3.png)(w:900)](BPNativization_ExclusiveStep_3.png)
+	[/REGION]
+	
+	[REGION:caption]
+		Click for full image.
+	[/REGION]
+
+1.	Select the **Exclusive** option to opt into the Exclusive Blueprint Nativization Method.
+
+	[REGION:lightbox]
+		[![](BPNativization_ExclusiveStep_4.png)(w:900)](BPNativization_ExclusiveStep_4.png)
+	[/REGION]
+	
+	[REGION:caption]
+		Click for full image.
+	[/REGION]
+
+1.	Go ahead and select the expansion arrow at the bottom of the submenu, revealing an empty array. 
+
+	[REGION:lightbox]
+		[![](BPNativization_ExclusiveStep_5.png)(w:900)](BPNativization_ExclusiveStep_5.png)
+	[/REGION]
+	
+	[REGION:caption]
+		Click for full image.
+	[/REGION]
+
+	This array will contain the list of Blueprint assets that you'll explicitly select for nativization.
+	![](BPNativization_ExclusiveStep_5Note.png)(w:900)
+
+1.	After exiting the **Project Properties** menu, navigate to the folder containing your Blueprint assets. 
+
+	[REGION:lightbox]
+		[![](BPNativization_ExclusiveStep_6.png)(w:900)](BPNativization_ExclusiveStep_6.png)
+	[/REGION]
+	
+	[REGION:caption]
+		Click for full image.
+	[/REGION]
+
+1.	Now, open the Blueprint asset(s) that you'd like to nativize.
+
+	[REGION:lightbox]
+		[![](BPNativization_ExclusiveStep_7.png)(w:900)](BPNativization_ExclusiveStep_7.png)
+	[/REGION]
+	
+	[REGION:caption]
+		Click for full image.
+	[/REGION]
+
+1.	After the **Blueprint Editor** opens, click the **Class Settings** button.
+
+	[REGION:lightbox]
+		[![](BPNativization_ExclusiveStep_8.png)(w:900)](BPNativization_ExclusiveStep_8.png)
+	[/REGION]
+	
+	[REGION:caption]
+		Click for full image.
+	[/REGION]
+
+1.	To mark this Blueprint asset for nativization, go ahead and click the **Nativize** checkbox located in the **Packaging** submenu in the **Details** panel.
+
+	[REGION:lightbox]
+		[![](BPNativization_ExclusiveStep_9.png)(w:900)](BPNativization_ExclusiveStep_9.png)
+	[/REGION]
+	
+	[REGION:caption]
+		Click for full image.
+	[/REGION]
+
+1.	Make sure to click the **Save** button before closing the **Blueprint Editor** for this asset.
+
+	[REGION:lightbox]
+		[![](BPNativization_ExclusiveStep_10.png)(w:900)](BPNativization_ExclusiveStep_10.png)
+	[/REGION]
+	
+	[REGION:caption]
+		Click for full image.
+	[/REGION]
+
+1.	To verify the Blueprint asset that you've selected for Exclusive nativization, open the **Project Settings** menu to view the array containing the list of Blueprint asset(s) that you've marked for nativization.
+
+	[REGION:lightbox]
+		[![](BPNativization_ExclusiveStep_11.png)(w:900)](BPNativization_ExclusiveStep_11.png)
+	[/REGION]
+	
+	[REGION:caption]
+		Click for full image.
+	[/REGION]
+
+At this point, you're ready to Cook and Package your project using the **Exclusive** Blueprint Nativization Method.
+
+[REGION:note]
+Note that all supported dependencies required to build the Blueprint(s) that you've selected for Nativization, will also be flagged for Nativization. 
+[/REGION]
+
+##Locating Nativized Code
+
+You can find your nativized code in `[ProjectDirectory]\Intermediate\\[Platform That Was Cooked For]\NativizedAssets\Source\NativizedAssets\`
+
+![](NativizedAssetsLocation.png)(w:900)
+
+##Setting the Nativization Flag
+
+If you navigate back to the **Project Settings > Project - Packaging** menu, you'll notice that there is a checkbox next to the **Warn if Packaged Without Nativization Flag** label. 
+
+[REGION:lightbox]
+	[![](NativizationFlag.png)(w:900)](NativizationFlag.png)
+[/REGION]
+
+[REGION:caption]
+	Click for full image.
+[/REGION]
 
-![](AdditionalMaps.png)
\ No newline at end of file
+This flag is intended for Advanced users who build/cook their projects from the command line (such as a server-side build engineers). Basically, if you enable this flag while also enabling nativization in 
+the **Project Settings** menu, a warning will be emitted if you don't use the nativization flag from the command line at build/cook time. To learn how to use the nativization switch from the command line, please 
+read the [](Engine/Blueprints/TechnicalGuide/NativizingBlueprints/Advanced) documentation.
\ No newline at end of file
diff --git a/Engine/Documentation/Source/Engine/Blueprints/UserGuide/BlueprintCommsUsage/BPComHowTo/BPComHowTo.INT.udn b/Engine/Documentation/Source/Engine/Blueprints/UserGuide/BlueprintCommsUsage/BPComHowTo/BPComHowTo.INT.udn
--- a/Engine/Documentation/Source/Engine/Blueprints/UserGuide/BlueprintCommsUsage/BPComHowTo/BPComHowTo.INT.udn
+++ b/Engine/Documentation/Source/Engine/Blueprints/UserGuide/BlueprintCommsUsage/BPComHowTo/BPComHowTo.INT.udn
@@ -1,7 +1,8 @@
 Availability:Docs
 Title:How To Use Blueprint Communications
 Crumbs: %ROOT%, Engine, Engine/Blueprints, Engine/Blueprints/UserGuide, Engine/Blueprints/UserGuide/BlueprintCommsUsage
-Description:Examples of how to setup and use different methods of Blueprint Communication. 
+Description:Examples of how to setup and use different methods of Blueprint Communication.
+tags:Blueprints
 
 [TOC(start:2 end:4)]
 
diff --git a/Engine/Documentation/Source/Engine/Blueprints/UserGuide/KismetNodeReference.INT.udn b/Engine/Documentation/Source/Engine/Blueprints/UserGuide/KismetNodeReference.INT.udn
--- a/Engine/Documentation/Source/Engine/Blueprints/UserGuide/KismetNodeReference.INT.udn
+++ b/Engine/Documentation/Source/Engine/Blueprints/UserGuide/KismetNodeReference.INT.udn
@@ -6,6 +6,7 @@ version: 4.9
 parent:Engine/Blueprints
 order:3
 type:reference
+tags:Blueprints
 
 [VAR:Topic]
 [OBJECT:Topic]
@@ -56,6 +57,7 @@ type:reference
 %Engine/Blueprints/UserGuide/OnlineNodes:topiccompact%
 %Engine/Blueprints/UserGuide/MathNode:topiccompact%
 %Engine/Blueprints/UserGuide/RandomStreams:topiccompact%
+%Engine/Blueprints/UserGuide/PatchingNodes:topiccompact%
 
 
 [/REGION]
diff --git a/Engine/Documentation/Source/Engine/Blueprints/UserGuide/PatchingNodes/MobilePatchUtilities.INT.udn b/Engine/Documentation/Source/Engine/Blueprints/UserGuide/PatchingNodes/MobilePatchUtilities.INT.udn
--- a/Engine/Documentation/Source/Engine/Blueprints/UserGuide/PatchingNodes/MobilePatchUtilities.INT.udn
+++ b/Engine/Documentation/Source/Engine/Blueprints/UserGuide/PatchingNodes/MobilePatchUtilities.INT.udn
@@ -1,6 +1,6 @@
-Availability:Docs
+Availability:Public
 Title:Mobile Patch Utility Nodes
-Crumbs: %ROOT%, Engine, Engine/Blueprints, Engine/Blueprints/UserGuide
+Crumbs:
 Description:Describes the different kinds of Mobile Patch Utility Blueprint nodes.
 Related: Engine/Blueprints
 Related: Platforms/Mobile
@@ -8,6 +8,9 @@ Related: Platforms/Mobile/MobilePackagingWizard
 version: 4.14
 Parent:Engine/Blueprints
 type:overview
+tags:Blueprints
+skilllevel:Intermediate
+topic-image:Patching_Topic.png
 
 [TOC (start:2 end:2)]
 
@@ -53,7 +56,7 @@ If the required disk space is less than the disk free space, then a custom event
 [![](Patching_3.png)(convert:false)](Patching_3.png)
 [/REGION]
 
-Although there is a lot going on in this graph, the only Mobile Patch Utility function is **Start Install**, which actually attempts to download and install the remote content. Just like **Request Content**, **Start Install** has
+Although there is a lot going on in this graph, the only Mobile Patch Utility function is **Start Install**, which attempts to download and install the remote content. Just like **Request Content**, **Start Install** has
 event parameters for success and failure. If the install succeeds, the **Mount Content** custom event is used to trigger the final stage of the patching process. Another custom event, **Display Download State** is fired off here
 with a timer. Using the timer and the **Update Download State** function, the download state can be shown to the user at regular intervals. More on that logic is in the next Blueprint graph.
 
@@ -65,9 +68,9 @@ with a timer. Using the timer and the **Update Download State** function, the do
 ![](Patching_4.png)
 [/REGION]
 
-While the **Display Download State** custom event is connected to a **Format Text** node here, this could also be the stage of the patching process where you fill a progress bar on the UI, or other methods of
+While the **Display Download State** custom event is connected to a **Format Text** node here, this could also be the stage of the patching process where you fill a progress bar on the UI (User Interface), or other methods of
 displaying progress to the user. Here, the outputs of **Get Download Size**, **Get Download Speed**, **Get Download Status Text**, **Get Install Progress**, and **Get Total Downloaded Size** nodes are 
-put together in a simple text block to present to the user. This script is called at regular intervals with the timer, but does not call back to any of the other example graphs.
+combined into a simple text block to present to the user. This script is called at regular intervals with the timer, but does not call back to any of the other example graphs.
 
 | Node | Description | 
 | --- | --- |
@@ -95,11 +98,11 @@ does not have success and failure event callbacks.
 ## Other Mobile Patching Functions
 
 The above graphs are a simple example of how to set up patching on a mobile device, but there are other, more complex, behaviors you could set up as well. For example, you could have your game warn users when they don't have
-an active WIFI connection, or provide different content for iOS and Android devices. For these use cases and others, the below functions would be useful additions to your Blueprints.
+an active WiFi connection, or provide different content for iOS and Android devices. For these use cases and others, the below functions would be useful additions to your Blueprints.
 
 | Node | Description | 
 | --- | --- |
-| Get Active Device Profile Name | Gets the name of currently selected device profile name. |
+| Get Active Device Profile Name | Gets the name of the currently selected device profile name. |
 | Get Installed Content Size | Gets the installed content size in megabytes  |
 | Get Supported Platform Names | Gets the list of supported platform names on this device. Example: Android_ETC2, Android_ASTC |
 | Has Active WiFi Connection | Returns whether or not a WiFi connection is currently available. |
diff --git a/Engine/Documentation/Source/Engine/Components/Camera/CameraComponents.INT.udn b/Engine/Documentation/Source/Engine/Components/Camera/CameraComponents.INT.udn
--- a/Engine/Documentation/Source/Engine/Components/Camera/CameraComponents.INT.udn
+++ b/Engine/Documentation/Source/Engine/Components/Camera/CameraComponents.INT.udn
@@ -8,15 +8,15 @@ Related: Gameplay/HowTo/UsingCameras
 Related: Engine/Blueprints/UserGuide/Components
 version: 4.9
 
-The **CameraComponent** (which adds a camera perspective to view from) and **SpringArmComponent** (which extends its children at a fixed distance then retracts when a collision occurs), when used together provide functionality for a 3rd person perspective that can dynamically adjust to your game world. 
+The **CameraComponent** (which adds a camera perspective to view from) and **SpringArmComponent** (which extends its children at a fixed distance then retracts when a collision occurs), when used together, provide functionality for a 3rd person perspective that can dynamically adjust to your game world. 
 
 ## Camera Component
 
-The CameraComponent allows you to add a camera perspective as a sub-object to an Actor. The CameraComponent will provide information about the camera's properties if the **ViewTarget** is a **CameraActor** or an Actor that contains a CameraComponent and has its **bFindCameraComponentWhenViewTarget** option set to true. 
+The CameraComponent adds a camera perspective as a sub-object to an **Actor**. The CameraComponent will provide information about the camera's properties if the **ViewTarget** is a **CameraActor** or an Actor that contains a CameraComponent and has its **Find Camera Component When ViewTarget** option set to true. 
 
-For example, say you wanted to [switch between multiple cameras](Gameplay/HowTo/UsingCameras/Blueprints) in your level during gameplay. By using **SetViewTargetWithBlend** and a CameraActor, you can switch between each of your cameras and use the properties that you define within the CameraActor for each camera (including field of view, angle, or any Post Processing effects, etc.) 
+For example, you could use CameraComponents to [switch between multiple cameras](Gameplay/HowTo/UsingCameras/Blueprints) in your level during gameplay. By using **Set View Target With Blend** and a CameraActor, you can switch between each of your cameras and use the properties that you define within the CameraActor for each camera (including field of view, angle, Post Processing effects, etc.).
 
-A related property that can be set for any **Pawn** is **bTakeCameraControlWhenPossessed**, where the Pawn will automatically become the ViewTarget upon possession by a PlayerController. So for example if you had multiple Characters (which are a form of Pawn) that you wanted to switch between, and each of them had their own CameraComponents assigned to them providing the camera perspective to view from, you could use set bTakeCameraControlWhenPossessed to true for each of them and whenever you switch between them, the CameraComponent for each will be used (where each could have its own properties and settings defined). 
+A related property that can be set for any **Pawn** is **Take Camera Control When Possessed**, which causesu the Pawn to become the ViewTarget automatically upon possession by a PlayerController. So, for example, if you had multiple Characters (which are a form of Pawn) that you wanted to switch between, and each of them had their own CameraComponent assigned to them, providing the camera perspective to view from, you could set Take Camera Control When Possessed to true for each of them and whenever you switch between them, the CameraComponent for that pawn would be used. 
 
 ## Spring Arm Component
 
@@ -24,12 +24,6 @@ The SpringArmComponent tries to maintain its children at a fixed distance from t
 
 ![](camera_springArm.png)
 
-_Above a CameraComponent (Camera) and SpringArmComponent (red line) are used to create a 3rd person perspective for a character._
+_Above, a CameraComponent (Camera) and SpringArmComponent (red line) are used to create a 3rd person perspective for a character._
 
 You can adjust several camera related properties on the SpringArmComponent such as the **TargetArmLength** which is the natural length of the spring arm when there are no collisions, **Probe Size** or how big should the probe sphere be when checking for collisions, as well as things like **CameraLag** useful if you want the camera to lag slightly behind the object it is attached to. 
-
-
-
-
-
-
diff --git a/Engine/Documentation/Source/Engine/Components/Haptic/HapticComponents.INT.udn b/Engine/Documentation/Source/Engine/Components/Haptic/HapticComponents.INT.udn
new file mode 100644
--- /dev/null
+++ b/Engine/Documentation/Source/Engine/Components/Haptic/HapticComponents.INT.udn
@@ -0,0 +1,16 @@
+Availability:Documents
+Title:Haptic Components
+Crumbs:%ROOT%, Engine, Engine/Components
+Description:Description of the Force Feedback Component is provided.
+Related: Gameplay/ForceFeedback
+version: 4.15
+
+The **Force Feedback Component** gives any **Actor** in the world the ability to act as a source of [Force Feedback](Gameplay/ForceFeedback).
+
+## Force Feedback Component
+
+[EXCERPT:WhatIsForceFeedbackComponent]
+The Force Feedback Component expands on the standard force feedback implementation by adding the concept of attenuation based on relative location in the world. Like sound or light, the intensity of the force experienced by the player will change with the player's distance from the source, according to a data-defined attenuation setting. The addition of attenuation makes it easy to create effects like explosions that feel more powerful at close range, trains that shake the ground with greater strength when the player is standing right next to the tracks, or even a character's "sixth sense" when passing close by to an object of interest.
+[/EXCERPT:WhatIsForceFeedbackComponent]
+
+For usage information, including detail on Force Feedback assets, please see the [Force Feedback page](Gameplay/ForceFeedback).
diff --git a/Engine/Documentation/Source/Engine/Components/Rendering/CableComponent/CableComponent.INT.udn b/Engine/Documentation/Source/Engine/Components/Rendering/CableComponent/CableComponent.INT.udn
new file mode 100644
--- /dev/null
+++ b/Engine/Documentation/Source/Engine/Components/Rendering/CableComponent/CableComponent.INT.udn
@@ -0,0 +1,281 @@
+Availability:Public
+Title:Cable Component
+Crumbs:%ROOT%, Engine
+Description: A breakdown of the Cable Component.
+version: 4.15
+parent:Engine/Editor
+
+![](CC_InGame.png)
+
+[TOC(start:1 end:2)]
+
+Having the ability to cheaply add cables, ropes or chains that can sway back and forth as if being gently blown by the wind can add a nice bit of life to your Unreal Engine 4 (UE4) projects. In the following document, we will take a look at how you can use the Cable component plugin to create, setup and control how the cables look, react and even collided with objects in your levels.
+
+
+## Simulation & Rendering
+
+To perform the actual cable simulation, a technique very well known to game development called **Verlet Integration** is used. The idea is to represent the cable as a series of particles, with **distance constraints** between them. The particles at the ends are **fixed** and move with whatever they are attached to. The ones in the middle are **free** and fall under gravity. Each step, you update the velocity and position of each particle, and then move them to satisfy the constraints. The **stiffness** of the cable is controlled by the number of times we iterate to enforce the constraints, each step.
+
+[REGION:lightbox]
+[![](CC_ParticleRope.png)(w:600)](CC_ParticleRope.png)
+[/REGION]
+
+[REGION:caption]
+Click for full image.
+[/REGION]
+
+Now that there is a nice chain of connected particles bouncing around, they now need to be rendered. To render the cable, a new class called **FCableSceneProxy** was created to represent the render of the cable. Then at each from the new particles positions from the simulation (which is done on the main thread inside TickComponent) are passed to this proxy via the **SendRenderDynamicData_Concurrent** function. Next, on the render thread, the update is locked, and then index and vertex buffers are updated to make a **tube** mesh. For each vertex on this tube mesh, we need to calculate a position, a Texture UV, and three Tangent Basis vectors. When doing this, **X** will point along the cable, **Z** will point straight out from the cable (the normal) and **Y** is perpendicular to **X** and **Z**. These properties have been exposed to, the components giving you the ability to control the number of sides, radius if the tube and how many times the UV's can title along the cable.
+
+[REGION:lightbox]
+[![](CC_RenderingGeo.png)(w:600)](CC_RenderingGeo.png)
+[/REGION]
+
+[REGION:caption]
+Click for full image.
+[/REGION]
+
+
+## Enabling the Plug-in
+
+By default, the Cable Component plugin should be enabled. However, if it is not, you can enable it by first going to the main toolbar and selecting **Edit** > **Plugins**. Then from the plugins list, go to **Rendering** and make sure the **Cable Component** has a checkmark in the box next to **Enabled**. 
+
+[REGION:lightbox]
+[![](CC_PluginEnable.png)(w:505)](CC_PluginEnable.png)
+[/REGION]
+
+[REGION:caption]
+Click for full image.
+[/REGION]
+
+## Using the Cable Component
+
+
+You can Cable Components to your projects levels in two different ways. In the following section, we will go over the two different ways that you can add cables to your projects levels.
+
+###  Using Cable Components from the Modes Pannel 
+
+To use the Cable Component from the Modes Panel, you will need to do the following:
+
+1. First, make sure that the **Modes** panel is visible and then in the **Search Classes** box, input the word **Cable**.
+
+    ![](CC_Find_In_Modes.png)
+
+1. To add the Cable Actor to the world, click on the Cable Actor in the Modes panel and then drag it into the level. 
+
+    [REGION:lightbox]
+    [![](CC_Add_CC_From_Modes.png)(w:700)](CC_Add_CC_From_Modes.png)
+    [/REGION]
+
+    [REGION:caption]
+    Click for full image.
+    [/REGION]
+
+1. You can now position, rotate and scale the Cable Actor so that it fits the needs of your level.
+
+    [OBJECT:EmbeddedVideo]
+        [PARAMLITERAL:width]
+        640
+        [/PARAMLITERAL]
+        [PARAMLITERAL:height]
+        360
+        [/PARAMLITERAL]
+        [PARAMLITERAL:videoid]
+        8sAwYlKV_ao
+        [/PARAMLITERAL]
+    [/OBJECT]
+
+###  Using Cable Component in Blueprints
+
+To use the Cable Component in your Blueprints, you will need to do the following:
+
+1. First, create a new Blueprint called **BP_Cable** and make sure that it uses **Actor** as the parent class.
+
+    ![](CC_MakeNewBP.png)
+
+1. Next, from the **Components** section of the BP_Cable Blueprint, click on the **Add Component** button and then locate the **Cable** component in the list. When located, click on it to add it to the Components list.
+
+    ![](CC_Add_CC.png)
+
+1. With the Cable Component now added, select the **Cable** component in the components list so that you can adjust its properties via the **Details** panel. For now, leave everything as the default and make sure to Compile and Save your Blueprint.
+
+    ![](CC_Cable_Options.png)
+
+    [region:note]
+    To make either side of the cable fall un-check the **Attach Start** or **Attach End** option the Details panel of the Cable Component. This can also be toggled while the game is running to go along with effects.
+    [/region]
+
+1. Locate your Cable Blueprint in the Content Browser and then drag it into the level. Once placed you can use the move and rotate tools to position it to your needs.
+
+    [REGION:lightbox]
+    [![](CC_Add_BP_CC.png)(w:700)](CC_Add_BP_CC.png)
+    [/REGION]
+
+    [REGION:caption]
+    Click for full image.
+    [/REGION]
+
+## Attaching Objects to the Cable ends
+
+You can also attach objects to either end of the cable so that the objects will swing along with the cable. To accomplish this in your UE4 projects all you need to do the following:
+
+1. You will first need to add a Cable Actor and a Static Mesh to the level. 
+
+    [REGION:lightbox]
+    [![](CC_Add_CC_And_SM.png)(w:500)](CC_Add_CC_And_SM.png)
+    [/REGION]
+
+    [REGION:caption]
+    Click for full image.
+    [/REGION]
+
+    [region:note]
+    Make sure that the Static Mesh that you are attaching to the end has its **Mobility** set to **Moveable**.
+    [/region]
+
+1. Locate the Static Mesh you want to attach to the end of the Cable Actor in the World Outliner and then drag it on top of the Cable Actor. After doing this, the following input window will be displayed.
+
+    [REGION:lightbox]
+    [![](CC_Attach_To_End.png)(w:500)](CC_Attach_To_End.png)
+    [/REGION]
+
+    [REGION:caption]
+    Click for full image.
+    [/REGION]
+
+1. Select the **Cable End** option, and in the viewport, you should see your Static Mesh snap to the end of the Cable Actor.
+
+    [OBJECT:ComparisonSlider]
+        [PARAM:before]
+        ![Before](CC_Attach_Before.png)
+        [/PARAM]
+        [PARAM:after]
+        ![After](CC_Attach_After.png)
+        [/PARAM]
+    [/OBJECT]
+
+1. Now select the Cable Actor in the level. Then in the **Details** panel under the **Cable** section, uncheck the **Attach End** box. 
+
+    [REGION:lightbox]
+    [![](CC_Remove_ATE.png)(w:300)](CC_Remove_ATE.png)
+    [/REGION]
+
+    [REGION:caption]
+    Click for full image.
+    [/REGION]
+
+    [region:note]
+    Note that you do not have to only use the **Attach Start** and **Attach End** option to attach cables to Actors. You can also specify a Socket that can be used as the attachment point.
+    [/region]
+
+1. After doing this, the cable should now swing freely in the viewport.
+
+    [OBJECT:EmbeddedVideo]
+        [PARAMLITERAL:width]
+        640
+        [/PARAMLITERAL]
+        [PARAMLITERAL:height]
+        360
+        [/PARAMLITERAL]
+        [PARAMLITERAL:videoid]
+        3auyPSMCGcE
+        [/PARAMLITERAL]
+    [/OBJECT]
+
+    [region:note]
+    Note that you can dynamically toggle the **Attach Start** and **Attach End** bools at run time, which can make for some interesting effects.
+    [/region]
+
+## Collision & Stiffness
+
+The Cable Component does have the option to have the cable collide with the world and also control how stiff the cable is. To enable this functionality, you will need to do the following:
+
+
+[region:warning]
+Enabling Collision and Stiffness will greatly increase the cost of the Cable Actor. Only enable these features if the cable has to collide with something in the world or if the cable needs some stiffness to help make the effect look better. If none of that is needed, it is best to keep these options disabled to save performance.
+[/region]
+
+1. First, in the Details of the Cable Component under the Cable section, press the small white triangle to expose the advanced rollout option.
+
+    ![](CC_Advanced_Options.png)
+
+1. Check both the **Enable Stiffness** and **Enable Collision** options to enable these features.
+
+    [REGION:lightbox]
+    [![](CC_Added_CC.png)(w:700)](CC_Added_CC.png)
+    [/REGION]
+
+    [REGION:caption]
+    Click for full image.
+    [/REGION]
+
+1. Now when you move the Cable Actor around or if an object collides with the Cable Actor the cable should collide with the object that it comes in contact with.
+
+    [OBJECT:EmbeddedVideo]
+        [PARAMLITERAL:width]
+        640
+        [/PARAMLITERAL]
+        [PARAMLITERAL:height]
+        360
+        [/PARAMLITERAL]
+        [PARAMLITERAL:videoid]
+        e18Kr8v-nNI
+        [/PARAMLITERAL]
+    [/OBJECT]
+
+## Property Breakdown
+
+The following table breaks down of what each of the properties on the Cable component does:
+
+#### Cable
+
+![](CC_Cable_Properties.png)
+
+|-------------|----------------|
+|**Property Name**| **Description**|
+|Attach Start| Should we fix the start to something, or leave it free. If false, component transform is just used for initial location of the start of cable.|
+|Attach End|Should we fix the end to something (using AttachEndTo and EndLocation), or leave it free. If false, AttachEndTo and EndLocation are just used for initial location of end of cable.|
+|Attach End To| Actor or Component that the defines the end position of the cable.|
+|Component Property| Name of component property to attach the cable to.|
+|Attach End To Socket Name| Socket name on the AttachEndTo component to attach to.|
+|End Location| End location of cable, relative to AttachEndTo (or AttachEndToSocketName) if specified, otherwise relative to cable component.|
+|Cable Length| Rest length of the cable.|
+|Num Segments| How many segments the cable has.|
+|Solver Iterations| The number of solver iterations controls how 'stiff' the cable is.|
+|Substep Time| Controls the simulation substep time for the cable.|
+|Enable Stiffness| Add stiffness constraints to cable.|
+|Enable Collision| EXPERIMENTAL. Perform sweeps for each cable particle, each substep, to avoid collisions with the world. Uses the Collision Preset on the component to determine what is collided with.This greatly increases the cost of the cable simulation.|
+
+### Cable Forces
+
+![](CC_CableForces_Properties.png)
+
+|-------------|----------------|
+|**Property Name**| **Description**|
+|Cable Forces|Force vector (world space) applied to all particles in cable.|
+|Cable Gravity Scale|Scaling applied to world gravity affecting this cable.|
+
+### Cable Rendering
+
+![](CC_CableRendering_Options.png)
+
+|-------------|----------------|
+|**Property Name**| **Description**|
+|Cable Width|How wide the cable geometry is.|
+|Num Sides|Number of sides of the cable geometry.|
+|Tile Material|How many times to repeat the material along the length of the cable.|
+
+
+## Reference 
+
+The following section contains links to various reference documents where you can read more about the papers used to power the cable components. 
+
+* To read more about **Verlet Integration** check out the following paper from Thomas Jakobsen.
+
+    * [http://graphics.cs.cmu.edu/nsp/course/15-869/2006/papers/jakobsen.htm](http://graphics.cs.cmu.edu/nsp/course/15-869/2006/papers/jakobsen.htm)
+
+
+
+
+
+
+
diff --git a/Engine/Documentation/Source/Engine/Components/Rendering/RenderingComponents.INT.udn b/Engine/Documentation/Source/Engine/Components/Rendering/RenderingComponents.INT.udn
--- a/Engine/Documentation/Source/Engine/Components/Rendering/RenderingComponents.INT.udn
+++ b/Engine/Documentation/Source/Engine/Components/Rendering/RenderingComponents.INT.udn
@@ -45,13 +45,10 @@ The **MaterialBillboardComponent** is a 2D Material that will be rendered always
 
 ## Cable Component
 
-A **CableComponent** allows you attach two Components together while rendering a cable between them. On the cable itself, you can assign a Material and define parameters that affect the way the cable is displayed. 
+A **CableComponent** allows you attach two Components together while rendering a cable between them. On the cable itself, you can assign a Material and define parameters that affect the way the cable is displayed. To read more about how to use the cable component, check out the link below.
 
-An example use case for this type of Component is if you were in need of a rope or something that a player character could swing on. By attaching one point to a fixed location and attaching the other to the player character through script, the character could swing on the rope. Below you have a first person character that can shoot a cable out from their gun and attach it to the impact location, creating a "grappling gun" of sorts. 
+* [**Cable Componet User Guide**](Engine\Components\Rendering\CableComponent\)
 
-![](cable.png)(w:440)
-
-Of course you could use this for simple visual elements as well such as representing wires that connect objects together. 
 
 ## Custom Mesh Component
 
