Copying //UE4/Dev-Documentation to //UE4/Samples-Main (Source: //UE4/Dev-Documentation @ 3372893)

diff --git a/Engine/Documentation/Source/Engine/Components/Utility/UtilityComponents.INT.udn b/Engine/Documentation/Source/Engine/Components/Utility/UtilityComponents.INT.udn
--- a/Engine/Documentation/Source/Engine/Components/Utility/UtilityComponents.INT.udn
+++ b/Engine/Documentation/Source/Engine/Components/Utility/UtilityComponents.INT.udn
@@ -1,26 +1,63 @@
 Availability:Public
 Title:Utility Components
-Crumbs:%ROOT%, Engine, Engine/Components
+Crumbs:
 Description:Description of the Utility based Components and how you can use them. 
 Related: Engine/Blueprints/UserGuide/Components
+parent:Engine/Components
+order:
+tags:Components
+type:overview
 SkillLevel: Beginner
-Version: 4.9
+Version: 4.14
 
+[TOC (start:1 end:4)]
 
 ## Application Lifecycle Component
 
-The **ApplicationLifecycleComponent** handles receiving notifications from the OS about application state (activated, suspended, termination, etc). 
+The **ApplicationLifecycleComponent** handles receiving notifications from the OS (Operating System) about the application state (activated, suspended, termination, etc). 
+
+## Platform Events Component
+
+The **PlatformEventComponent** provides the ability to poll and and to receive events based on platform-specific OS information and events. Currently, this component supports convertible laptops, providing information about whether or not the platform is a convertible laptop, and what state it is in (laptop mode or tablet mode). In order to know when the platform changes state without constantly polling these functions, there are also delegates that can be bound to notify your code (or blueprints) when these state changes happen.
 
 ## Child Actor Component
 
-A **ChildActorComponent** is a Component that spawns an Actor when registered, and destroys it when unregistered. 
+A **ChildActorComponent** is a Component that spawns an Actor when registered, and destroys it when unregistered.
+
+Child Actor Components specify the Actor that will be spawned by providing a C++ class or a Blueprint asset. After selecting a class or asset, the **Child Actor Template** field will appear. This field can be expanded to edit any exposed, editable field of the specified child Actor or its components, including public variables being declared in Blueprints.
+
+![](ChildActorTemplate.png)
 
 ## Scene Component
 
-The **SceneComponent** has a transform and supports attachment, but has no rendering or collision capabilities. This is useful as a "dummy" Component in the hierarchy to offset others. 
+The **SceneComponent** has a transform and supports attachment, but has no rendering or collision capabilities. This is useful as a "dummy" Component in the hierarchy to offset others, or to interact with its child components as if they were all one, such as when moving or rotating them, making them all visible or invisible, or detaching them. It is important not to overuse Scene Components as organizational tools, since each one will add to the time and memory it takes to spawn the Actors that own them.
+
+### Arrow Component
+
+An **ArrowComponent** is a Scene Component that features simple arrows drawn with lines to indicate the direction that it is facing. Arrow Components are especially useful for debugging things like the position and orientation of sockets. This is especially useful for things like clothing, hand-held or attached objects that may exist in your game, or gun muzzles where projectiles will be spawned. By default, these components are only visible in the Editor. If you want to see them in-game, be sure to uncheck the `Hidden in Game` checkbox in the Arrow component's **Details** panel.
+
+![](HiddenInGame_Checkbox.png)
+
+In the following image, an Arrow Component has been attached to the Unreal Engine 4 mannequin's index_03_r bone, revealing that the bone is located at the final knuckle of the right-hand index finger, and is pointed back toward the second knuckle. If you want to attach a laser beam component to show where your character's fingertip is pointing, you should either add a socket that is slightly offset and rotated 180 degrees and then attach to that, or understand that when you attach your laser component, that it needs to be offset from the bone.
+
+![](ArrowComponent.png)  
+
+## Spline Component
 
+The **Spline Component** provides users with a way to draw splines within a level. This is particularly useful for designating paths for moving Actors, such as trains, background characters, or cameras. When placed in a level, an Actor with a Spline Component will show a point at the component's location with a line connecting it to a second point a short distance away. The line can be right-clicked to add a new point at the location of the click, while the point can be clicked to remove it (unless it is the only point), duplicate it, or change the curve leading from it to the next point in the spline. Using duplication will create a new point that sits at the same location as the current point, but is considered to be "after" it in the line. Spline points can be freely manipulated within the level using standard movement and rotation tools. There is also a field labeled `Duration` located in the component's **Details** panel, which acts as a range value for querying the spline.
 
+![](Spline_Duration.png)
 
+There are three primary ways to query the spline: 
+* By point index, where the spline can be asked for information like the location and rotation of any of the user-placed points along its length. 
+* By distance, where the spline can be asked for information such as the rotation forward vector of the spline after traveling a set distance along it.
+* With "time", where the spline's length is replaced by a user-entered `Duration` value and can be queried for distance, location, and rotation information about the spline at a given "time" value. This third method is especially useful in the case of fixed objects like cameras moving along a track defined by a spline, since the designer or programmer can simply ask for the camera's location and rotation based on how long the camera has been moving along the spline. 
+It is also easy to interpolate or blend along a spline by setting the duration to 1.0 and then providing an interpolation "alpha" value.
 
+In the following image, the spline component could be used to move a "hint" marker, showing the player where to go.
 
+![](SplineComponent.png)  
 
+[REGION:note]
+Splines created by a Spline Component will never fork into multiple paths or contain loops or orphaned points. Splines are always single paths that contain each spline point exactly once.
+[/REGION]
\ No newline at end of file
diff --git a/Engine/Documentation/Source/Engine/Editor/VR/ActivateVRMode/ActivateVRMode.INT.udn b/Engine/Documentation/Source/Engine/Editor/VR/ActivateVRMode/ActivateVRMode.INT.udn
--- a/Engine/Documentation/Source/Engine/Editor/VR/ActivateVRMode/ActivateVRMode.INT.udn
+++ b/Engine/Documentation/Source/Engine/Editor/VR/ActivateVRMode/ActivateVRMode.INT.udn
@@ -4,57 +4,26 @@ Crumbs:
 Availability: public
 parent: Engine/Editor/VR
 type:how-to
-order: 2
+order: 1
 tags: VR
 tags: Getting Started
-version: 4.13
+version: 4.15
 skilllevel: intermediate
+topic-image:activate_topic.png
 
-[VAR:Topic]
-[OBJECT:Topic]
-	[PARAM:image]
-		![%Engine/Editor/VR/ActivateVRMode:title%](Engine/Editor/VR/ActivateVRMode/activate_topic.png)
-	[/PARAM]
-	[PARAM:title]
-		%Engine/Editor/VR/ActivateVRMode:title%
-	[/PARAM]
-	[PARAM:path]
-		[RELATIVE:Engine/Editor/VR/ActivateVRMode]
-	[/PARAM]
-	[PARAM:description]
-		%Engine/Editor/VR/ActivateVRMode:description%
-	[/PARAM]
-[/OBJECT]
-[/VAR]
-
-[VAR:TopicCompact]
-[OBJECT:TopicCompact]
-	[PARAM:image]
-		![%Engine/Editor/VR/ActivateVRMode:title%](Engine/Editor/VR/ActivateVRMode/activate_topic.png)
-	[/PARAM]
-	[PARAM:title]
-		%Engine/Editor/VR/ActivateVRMode:title%
-	[/PARAM]
-	[PARAM:path]
-		[RELATIVE:Engine/Editor/VR/ActivateVRMode]
-	[/PARAM]
-	[PARAM:description]
-		%Engine/Editor/VR/ActivateVRMode:description%
-	[/PARAM]
-[/OBJECT]
-[/VAR]
+[TOC(start:1 end:4)]
 
-As of Unreal Engine 4.12, VR Editing Mode is included with the version of the engine available through the Epic Games Launcher. You can also get the most up-to-date build from [GitHub](Engine/Editor/VR/Setup).
+As of Unreal Engine 4.12, VR (Virtual Reality) Editing Mode is included with the version of the engine available through the Epic Games Launcher. You can also get the most up-to-date build from [GitHub](Engine/Editor/VR/Setup).
 
 [EXCERPT:main]
 
-1. After the editor loads, click **File > New** and open up our new "VR-Basic" level.  This environment is a good playground for you to learn the controls.
+1. After the Editor loads, click **File > New** and open up our new "VR-Basic" level.  This environment is a good playground for you to learn the controls.
 
 	![](vr_map_template.png)
 
-1. In the Editor Preferences (**Edit > Editor Preferences**), select the **Experimental** category and choose **Enable VR Editing**.
+1. In the **Editor Preferences** menu (**Edit > Editor Preferences**), select the **Experimental** category and choose **Enable VR Editing**.
 
-1. To get into VR, click the **VR** button on the editor's main tool bar to activate VR Mode. Assuming your tracking space is properly calibrated, you should see a Tutorial window floating in front of you with some help.
+1. To get into VR, click the **VR** button on the Editor's main toolbar to activate VR Mode. Assuming your tracking space is properly calibrated, you should see a Tutorial window floating in front of you with some help.
 
 	![](vrmode_toolbar.png)
 	
@@ -63,13 +32,13 @@ As of Unreal Engine 4.12, VR Editing Mode is included with the version of the en
 	[/REGION]
 1. To use the tutorial, position the Controller over the section you want to read then squeeze the Controller trigger to open that section. 
 
-At any time you can press the Home button to return to this menu and then select another help section.
+At any time you can press the **Home** button to return to the main tutorial menu and then select another help section. 
 
 [/EXCERPT:main]
 
 ## Automatic VR Editing Mode
 
-As of 4.13, you can also enter VR Editing Mode automatically.  As long as the editor is in the foreground and VR Editing mode is enabled, when you wear the HMD, you will automatically enter VR editing mode. When you remove the HMD, you will leave it. 
+As of 4.13, you can also enter VR Editing Mode automatically.  As long as the editor is in the foreground and VR Editing mode is enabled, when you wear the HMD, you will automatically enter VR Editing Mode. When you remove the HMD, you will leave it. 
 
 [OBJECT:EmbeddedVideo]
 [PARAMLITERAL:width]
@@ -83,14 +52,33 @@ cZJmh881FCo
 [/PARAMLITERAL]
 [/OBJECT]
 
-**To enable/disable VR Editing Mode Auto-Entry:**
+###Enabling or Disabling VR Editing Mode Auto-Entry
+
 1. In the **Edit** menu, click on **Editor Preferences** to view the Unreal Editor options.
 	
 	![](EditorPref.png)
-2. Under the **General** heading, select the **Experimental** section and locate the **VR** category which contains the VR Editor options.
+1. Under the **General** heading, select the **Experimental** section and locate the **VR** category, which contains the VR Editor options.
 
 	![](ExperimentalVR.png)
-3. Toggle the **Enable VR Mode Auto-Entry** option to disable or enable the auto-entry functionality.
+1. Toggle the **Enable VR Mode Auto-Entry** option to disable or enable the auto-entry functionality.
 
 	![](AutoEntryMenu_Highlight.png)
 
+
+## VR Editor Tutorial Window
+
+As of 4.14, when you enter VR Editing Mode you can enable or disable the VR Editor Tutorial window from showing each time you start a new session.
+
+![](VRTutorialWindow.png)(w:500)
+
+###Enabling or Disabling VR In-Editor Tutorials Window
+
+1. In the **Edit** menu, click on **Editor Preferences** to view the Unreal Editor options. 
+	
+	![](EditorPref.png)
+1. Under the **General** heading, select the **Experimental** section and locate the **VR** category, which contains the VR Editor options. 
+
+	![](ExperimentalVR.png)
+1.	Toggle the **Always Show VR Tutorial at Start** option to disable or enable the VR Tutorial window. 
+
+	![](InEditorTutorial.png)
diff --git a/Engine/Documentation/Source/Engine/Editor/VR/ActorInteraction/ActorInteraction.INT.udn b/Engine/Documentation/Source/Engine/Editor/VR/ActorInteraction/ActorInteraction.INT.udn
--- a/Engine/Documentation/Source/Engine/Editor/VR/ActorInteraction/ActorInteraction.INT.udn
+++ b/Engine/Documentation/Source/Engine/Editor/VR/ActorInteraction/ActorInteraction.INT.udn
@@ -7,42 +7,12 @@ type:how-to
 order: 4
 tags: VR
 tags: Level Editing
-version: 4.11
+version: 4.15
 skilllevel: intermediate
+topic-image:actorinteraction.png
+prereq:Engine/Editor/VR/ActivateVRMode
 
-[VAR:Topic]
-[OBJECT:Topic]
-	[PARAM:image]
-		![%Engine/Editor/VR/ActorInteraction:title%](Engine/Editor/VR/ActorInteraction/actorinteraction.png)
-	[/PARAM]
-	[PARAM:title]
-		%Engine/Editor/VR/ActorInteraction:title%
-	[/PARAM]
-	[PARAM:path]
-		[RELATIVE:Engine/Editor/VR/ActorInteraction]
-	[/PARAM]
-	[PARAM:description]
-		%Engine/Editor/VR/ActorInteraction:description%
-	[/PARAM]
-[/OBJECT]
-[/VAR]
-
-[VAR:TopicCompact]
-[OBJECT:TopicCompact]
-	[PARAM:image]
-		![%Engine/Editor/VR/ActorInteraction:title%](Engine/Editor/VR/ActorInteraction/actorinteraction.png)
-	[/PARAM]
-	[PARAM:title]
-		%Engine/Editor/VR/ActorInteraction:title%
-	[/PARAM]
-	[PARAM:path]
-		[RELATIVE:Engine/Editor/VR/ActorInteraction]
-	[/PARAM]
-	[PARAM:description]
-		%Engine/Editor/VR/ActorInteraction:description%
-	[/PARAM]
-[/OBJECT]
-[/VAR]
+[TOC(start:1 end:4)]
 
 [EXCERPT:main]
 
diff --git a/Engine/Documentation/Source/Engine/Editor/VR/Controls/VREditorControls.INT.udn b/Engine/Documentation/Source/Engine/Editor/VR/Controls/VREditorControls.INT.udn
--- a/Engine/Documentation/Source/Engine/Editor/VR/Controls/VREditorControls.INT.udn
+++ b/Engine/Documentation/Source/Engine/Editor/VR/Controls/VREditorControls.INT.udn
@@ -1,51 +1,21 @@
+Availability: Public
 Title: VR Editor Controls
 Description: Learn about the various control schemes and interactions used when creating worlds in VR mode.
 Crumbs:
-Availability: public
 parent: Engine/Editor/VR
-type: reference
-order: 1
+type: overview
+order:
 tags: VR
-version: 4.13
+version: 4.15
 skilllevel: intermediate
+topic-image:vreditor_controls.png
+prereq:Engine/Editor/VR/ActivateVRMode
 
-[VAR:Topic]
-[OBJECT:Topic]
-	[PARAM:image]
-		![%Engine/Editor/VR/Controls:title%](Engine/Editor/VR/Controls/vreditor_controls.png)(fill:#919293)
-	[/PARAM]
-	[PARAM:title]
-		%Engine/Editor/VR/Controls:title%
-	[/PARAM]
-	[PARAM:path]
-		[RELATIVE:Engine/Editor/VR/Controls]
-	[/PARAM]
-	[PARAM:description]
-		%Engine/Editor/VR/Controls:description%
-	[/PARAM]
-[/OBJECT]
-[/VAR]
-
-[VAR:TopicCompact]
-[OBJECT:TopicCompact]
-	[PARAM:image]
-		![%Engine/Editor/VR/Controls:title%](Engine/Editor/VR/Controls/vreditor_controls.png)(fill:#919293)
-	[/PARAM]
-	[PARAM:title]
-		%Engine/Editor/VR/Controls:title%
-	[/PARAM]
-	[PARAM:path]
-		[RELATIVE:Engine/Editor/VR/Controls]
-	[/PARAM]
-	[PARAM:description]
-		%Engine/Editor/VR/Controls:description%
-	[/PARAM]
-[/OBJECT]
-[/VAR]
-
-There are various controls for navigating the scene, selecting and manipulating Actors, and working with editor panels and windows.
-
-## The Controller
+[TOC(start:1 end:4)]
+
+There are various controls in virtual reality (VR) for navigating the scene, selecting and manipulating Actors, and working with editor panels and windows.
+
+## Controller Interface
 
 [REGION:raw]
 ![](controller.png)(convert:false)
@@ -56,9 +26,9 @@ There are various controls for navigating the scene, selecting and manipulating
 3. Trigger
 4. Side Buttons
 
-## Laser Pointer
+### Laser Pointer
 
-Each Controller emits a beam called the laser pointer which is used to determine which Actor you will interact with.  The color of the beam indicates the interaction mode the controller is in.
+Each Controller emits a beam called the "laser pointer" which is used to determine with which Actor you will interact.  The color of the beam indicates the interaction mode the controller is currently in.
 
 | Color | Appearance | Meaning |
 | ----- | ---------- | ------- |
@@ -66,7 +36,48 @@ Each Controller emits a beam called the laser pointer which is used to determine
 | Green |![](selectionmode.png)(w:192) | The controller is in selection mode, meaning it currently has an Actor selected.  |
 | Yellow |![](movementmode.png)(w:192) | The controller is in movement mode, meaning it is being used to navigate the world. |
 
-## Navigation Controls
+##UI in VR
+
+The UI in VR Editor mode enables you to work within the Editor without having to leave VR. Use the controls to bring up the Menu system that gives you access to a variety of settings and Editor windows that enables you the 
+full functionality of the desktop Editor in VR. Use the sections below to learn how to call and interact with these systems.
+
+
+### Menu Controls
+
+If you want to use the menus available in VR Editing Mode, refer to the following table, where you'll learn about the **Radial** and the **Quick Select** menus. 
+
+[REGION:raw]
+| Control | Description |
+| ------- | ----------- |
+| ![Modifier Button](radial_menu.png)(convert:false) | Lightly press the Track Pad with your thumb to open the **Radial Menu**. Move your thumb around the Track Pad to highlight the options. Press the Track Pad fully to select the highlighted option. |
+| ![Aim One Controller at the Other Controller](quick_select_menu.png)(convert:false) | Aim one controller at the other controller to access the **Quick Select** Menu. |
+[/REGION]
+
+### Window Controls
+
+If you want to control the various windows while in VR Editing Mode, refer to the following table, where you'll learn how to show and hide windows, and how to scale windows up or down.
+
+[REGION:raw]
+| Control | Description |
+| ------- | ----------- |
+| ![Tap Track Pad](show_hide_windows.png)(convert:false) | Tap the Track Pad to show or hide any open windows. |
+| ![Aim Controller at Window Grip + Squeeze Trigger + Track Pad Towards or Away](scale_window.png)(convert:false) | Aim the controller at the Grip of a window and squeeze the Trigger to grab the window. Move your thumb up or down on the Track Pad to scale the window up or down. |
+[/REGION]
+
+##Editing in VR
+
+When editing in VR Editor mode you'll need to be able to traverse the world and manipulate the Actors in your scene in an intuitive way. Use the sections below to learn how to interact with the Actors by moving them around, scaling, and rotating them, and 
+how to effectively navigate the world all while in VR. 
+
+
+<!--Add introductory text summarizing this section.-->
+
+### Navigation Controls
+
+If you want to navigate your way through a VR environment, refer to the following table, 
+where you'll learn how to push and pull the world with your hands, teleport to locations within 
+the world, rotate the world around you, and scale the world up or down.
+
 
 [REGION:raw]
 | Control | Description |
@@ -75,18 +86,24 @@ Each Controller emits a beam called the laser pointer which is used to determine
 | ![Press Side Buttons + Aim Controller + Squeeze Trigger](teleport_world.png)(convert:false) | Press the Side Buttons, aim the controller, and squeeze the Trigger to teleport from your current location to the location the controller is aiming at. |
 | ![Press Side Buttons + Move Controllers Around Each Other](rotate_world.png)(convert:false) | Press the Side Buttons on both controllers and move them around each other to rotate the world around you, as if you are grabbing the world with your hands and then rotating it. |
 | ![Press Side Buttons + Move Controllers Towards or Away From Each Other](scale_world.png)(convert:false) | Press the Side Buttons on both controllers and move them toward or away from each other to scales the world up or down. |
+| ![Press Side Buttons + Tap the Tack Pad](Reset_World_Scale.png)(convert:false) | Press the Side Buttons to grab the world and tap the Track Pad to instantly reset the world's scale back to your default Player's Size while editing. |
+
 [/REGION]
 
-## Selection Controls
+### Selection Controls
+
+If you want to select Actors in a VR environment, refer to the following table, where you'll learn how to select an Actor and how to select multiple Actors at a time. 
 
 [REGION:raw]
 | Control |  Description |
 | ------- | ----------- |
 | ![Aim Controller at Actor + Squeeze Trigger](select_actor.png)(convert:false) | Aim the controller at an Actor in the world and squeeze the Trigger to select the Actor. |
-| [INCLUDE:#multiselect] | Aim one controller at an Actor in the world and squeeze the Trigger halfway to select the Actor and enter multi-select mode. While holding the Trigger halfway, aim the other controller at other Actors and squeeze its Trigger to add or remove them from the selection. |
+| [REGION:lightbox][![Aim Controller at Actor + Squeeze Trigger Halfway + Aim Controller at another Actor + Squeeze Trigger](select_actor_multiple.png)(h:108 convert:false)](select_actor_multiple.png)[/REGION] | Aim one controller at an Actor in the world and squeeze the Trigger halfway to select the Actor and enter multi-select mode. While holding the Trigger halfway, aim the other controller at other Actors and squeeze its Trigger to add or remove them from the selection. |
 [/REGION]
 
-## Transform Controls
+### Transform Controls
+
+If you want to transform your Actors using the Translate, Rotate, and Scale gizmos, refer to the following table, where you'll learn how to move, rotate, scale, and slide Actors while in VR.
 
 [REGION:raw]
 | Control | Description |
@@ -97,25 +114,15 @@ Each Controller emits a beam called the laser pointer which is used to determine
 | ![Aim Controller at Actor + Squeeze Trigger + Track Pad Towards or Away](slide_actor.png)(convert:false) | Aim the controller at an Actor in the world and squeeze the Trigger to select the Actor. Move your thumb down or up on the Track Pad to move the selected Actor towards or away from you. |
 [/REGION]
 
-## Menu Controls
+##Painting in VR
 
-[REGION:raw]
-| Control | Description |
-| ------- | ----------- |
-| ![Modifier Button](radial_menu.png)(convert:false) | Lightly press the Track Pad with your thumb to open the Radial Menu. Move your thumb around the Track Pad to highlight the options. Press the Track Pad fully to select the highlighted option. |
-| ![Aim One Controller at the Other Controller](quick_select_menu.png)(convert:false) | Aim one controller at the other controller to access the Quick Select Menu. |
-[/REGION]
+When painting in VR Editor mode, you'll have similar options to the desktop Editor to choose from, such as painting for Landscape mode, with the Foliage tool, and Mesh Painting. Use the sections below to learn how to paint
+with these different modes using their available tools and options. 
 
-## Window Controls
 
-[REGION:raw]
-| Control | Description |
-| ------- | ----------- |
-| ![Tap Track Pad](show_hide_windows.png)(convert:false) | Tap the Track Pad to show or hide any open windows. |
-| ![Aim Controller at Window Grip + Squeeze Trigger + Track Pad Towards or Away](scale_window.png)(convert:false) | Aim the controller at the Grip of a window and squeeze the Trigger to grab the window. Move your thumb up or down on the Track Pad to scale the window up or down. |
-[/REGION]
+### Mesh Painting Controls
 
-## Mesh Painting Controls
+If you want to paint on a mesh using Vertex Painting, refer to the following table, where you'll learn how to paint and erase painted areas. 
 
 [REGION:raw]
 | Control | Description |
@@ -124,9 +131,9 @@ Each Controller emits a beam called the laser pointer which is used to determine
 | ![Modifier + Aim Controller at Actor + Squeeze Trigger + Move Controller](ModifierMeshPaint.png)(convert:false) | To erase instead of paint, hold the modifier button. Aim the controller at a Static Mesh in the world and squeeze the Trigger to begin erasing. Move the controller to move the eraser along the Mesh.  |
 [/REGION]
 
-## Foliage Painting Controls
+### Foliage Painting Controls
 
-For both painting and erasing foliage, pressure sensitivity on the Trigger is supported. 
+If you want to use the Foliage Painting controls, refer to the following table, where you'll learn about painting and erasing foliage. It is also worth noting that pressure sensitivity on the Trigger is supported to control how strong an area is painted or erased.
 
 [REGION:raw]
 | Control | Description |
@@ -135,13 +142,13 @@ For both painting and erasing foliage, pressure sensitivity on the Trigger is su
 | ![Modifier + Aim Controller at Actor + Squeeze Trigger + Move Controller](ModifierMeshPaint.png)(convert:false) | To erase foliage, hold the modifier button. Aim the controller at a location in the world and squeeze the Trigger to begin erasing. Move the controller to move the eraser around the world.  |
 [/REGION]
 
-<!--
-[EXCERPT:MultiSelect]
-[PUBLISH:Rocket]
-[REGION:lightbox][![Aim Controller at Actor + Squeeze Trigger Halfway + Aim Controller at another Actor + Squeeze Trigger](select_actor_multiple.png)(h:108 convert:false)](select_actor_multiple.png)[/REGION]
-[/PUBLISH:Rocket]
-[PUBLISH:Epic]
-![Aim Controller at Actor + Squeeze Trigger Halfway + Aim Controller at another Actor + Squeeze Trigger](select_actor_multiple.png)(h:108 convert:false)
-[/PUBLISH:Epic]
-[/EXCERPT:MultiSelect]
--->
\ No newline at end of file
+### Landscape Painting Controls
+
+If you want to learn about Landscape tools, refer to the following table, where you'll learn about raising and lowering the landscape. It is also worth noting that pressure sensitivity on the Trigger is supported to control how strong or soft an area is raised or lowered. 
+
+[REGION:raw]
+| Control | Description |
+| ------- | ----------- |
+| ![Aim Controller at Actor + Squeeze Trigger + Move Controller](MeshPaint_2_NoModifier.png)(convert:false) | Select a landscape brush in the **Sculpt** panel. Aim the controller at a location in the world and squeeze the Trigger to begin raising the landscape. Move the controller to move the brush around the world.|
+| ![Modifier + Aim Controller at Actor + Squeeze Trigger + Move Controller](ModifierMeshPaint.png)(convert:false) | To lower the landscape, hold the modifier button and make sure you have a brush selected in the **Sculpt** panel. Aim the controller at a location in the world and press the modifier button in a while squeezing the Trigger to begin lowering the landscape. Move the controller to lower different areas around the world.  |
+[/REGION]
diff --git a/Engine/Documentation/Source/Engine/Editor/VR/EditorUI/EditorUIVR.INT.udn b/Engine/Documentation/Source/Engine/Editor/VR/EditorUI/EditorUIVR.INT.udn
--- a/Engine/Documentation/Source/Engine/Editor/VR/EditorUI/EditorUIVR.INT.udn
+++ b/Engine/Documentation/Source/Engine/Editor/VR/EditorUI/EditorUIVR.INT.udn
@@ -6,46 +6,19 @@ parent: Engine/Editor/VR
 type: reference
 order: 5
 tags: VR
-version: 4.13
+version: 4.15
 skilllevel: intermediate
+topic-image:editorui_topic.png
+prereq:Engine/Editor/VR/ActivateVRMode
 
-[VAR:Topic]
-[OBJECT:Topic]
-	[PARAM:image]
-		![%Engine/Editor/VR/EditorUI:title%](Engine/Editor/VR/EditorUI/editorui_topic.png)
-	[/PARAM]
-	[PARAM:title]
-		%Engine/Editor/VR/EditorUI:title%
-	[/PARAM]
-	[PARAM:path]
-		[RELATIVE:Engine/Editor/VR/EditorUI]
-	[/PARAM]
-	[PARAM:description]
-		%Engine/Editor/VR/EditorUI:description%
-	[/PARAM]
-[/OBJECT]
-[/VAR]
-
-[VAR:TopicCompact]
-[OBJECT:TopicCompact]
-	[PARAM:image]
-		![%Engine/Editor/VR/EditorUI:title%](Engine/Editor/VR/EditorUI/editorui_topic.png)
-	[/PARAM]
-	[PARAM:title]
-		%Engine/Editor/VR/EditorUI:title%
-	[/PARAM]
-	[PARAM:path]
-		[RELATIVE:Engine/Editor/VR/EditorUI]
-	[/PARAM]
-	[PARAM:description]
-		%Engine/Editor/VR/EditorUI:description%
-	[/PARAM]
-[/OBJECT]
-[/VAR]
+
+[TOC(start:1 end:4)]
 
 [EXCERPT:main]
 
-The VR Editor enables you to summon and dock a number of different panels and windows including:
+The VR Editor enables you to summon and dock a number of different panels and windows.
+
+##Panels and Windows
 
 [REGION:imagetable]
 | ![](window_content.png)(w:225) | ![](window_details.png)(w:225) | ![](window_modes.png)(w:225) | ![](window_outliner.png)(w:225) |
@@ -53,6 +26,8 @@ The VR Editor enables you to summon and dock a number of different panels and wi
 | Content Browser | Details Panel | Modes Panel | Outliner Panel | 
 |![](window_tutorials.png)(w:225) | ![](window_asset.png)(w:225) | ![](colorPicker.png)(w:225)| ![](worldSettings.png)(w:225)|
 |Tutorials | Asset Panel | Color Picker | World Settings |
+|![](window_Landscape.png)(w:225) | ![](window_Foliage.png)(w:225) | | |
+|Landscape | Foliage|   ||||
 [/REGION]
 
 ## Opening a Window in VR Mode
diff --git a/Engine/Documentation/Source/Engine/Editor/VR/GDC2017/VREditorGDC2017.INT.udn b/Engine/Documentation/Source/Engine/Editor/VR/GDC2017/VREditorGDC2017.INT.udn
new file mode 100644
--- /dev/null
+++ b/Engine/Documentation/Source/Engine/Editor/VR/GDC2017/VREditorGDC2017.INT.udn
@@ -0,0 +1,247 @@
+Availability:Public
+Title: VR and Mesh Editor GDC 2017 Features Preview
+Crumbs:
+Description:Previewing some of the VR and Mesh Editor features being presented at GDC 2017.
+Parent:Engine/Editor/VR
+type: reference
+order: 1000
+topic-image:gdp_topic.png
+reference-image:snapping.gif
+reference-image:vr_contentBrowser.gif
+reference-image:vr_contextOption.gif
+reference-image:vr_radialMenu.gif
+prereq:Engine/Editor/VR/ActivateVRMode
+prereq:Engine/Editor/VR/Setup
+
+[TOC(start:1 end:4)]
+
+
+We are so excited about the latest additions to the VR Editor and Mesh Editor, we couldn't wait to get it into your hands! So, we have made the GDC 2017 Features Branch available in an experimental build through GitHub! This builds contains all of the VR Editor and Mesh Editor features demonstrated at the **State of Unreal** keynote at the 2017 Game Developers Conference.
+
+![image alt text](image_0.gif)(h:320) ![](poly_extrude.png)(h:320) ![](poly_extrude_subd.png)(h:320)
+
+[REGION:note]
+The VR Editor and Mesh Editor are experimental features. For this particular release, features were tested exclusively on Windows, so there may be complications compiling for Mac or Linux.
+[/REGION]
+
+## Major Features
+
+* Mesh Editor Mode enables manipulating Static Meshes within Unreal Editor, including using subdivision surfaces!
+* The Sequencer non-linear editing tool is now available in VR Mode
+* Redesigned Radial Menu with context sensitive actions
+* Smart Snapping allows for quickly and accurately placing objects relative to other assets in the level
+* Ability to save Static Mesh changes made during Simulate Mode
+
+## Downloading the Dev-Geometry GDC 2017 Features Branch
+
+To download the Dev-Geometry GDC 2017 Features Branch:
+
+1. [Associated your GitHub account with your EpicID](https://www.unrealengine.com/ue4-on-github). 
+
+1. Download the build from our GitHub repository: [Dev-Geometry](https://github.com/EpicGames/UnrealEngine/tree/dev-geometry)
+
+1. [Compile](Programming/Development/BuildingUnrealEngine) and [run Unreal Editor](GettingStarted/RunningUnrealEngine).
+
+
+## Mesh Editor
+
+The **Mesh Editor** mode is a new geometry editing toolset designed to enable designers to quickly create and modify Static Mesh geometry in the Level Editor viewport. In addition to a low-poly toolset, Mesh Editor mode adds the ability to work with subdivision surfaces, enabling artists and designers to create smooth, organic surfaces directly in Unreal Editor.
+
+![](poly_extrude.png) ![](poly_extrude_subd.png)
+[REGION:note]
+The Mesh Editor mode is an extremely early access feature.
+[/REGION]
+
+To edit a mesh using the Mesh Editor mode:
+
+1. In the **Modes** panel, select the **Mesh Editor** mode.
+
+1. Choose an editing mode depending on what elements - Vertices, Edges, Polygons, All - of the mesh you want to edit.
+
+1. Select a mesh in the viewport. The elements corresponding to the mode highlight when the cursor is over them.
+
+1. Select an element to edit it.
+
+
+### Selection Modes
+
+#### Vertices
+
+The **Vertices** selection mode limits editing to vertices of the mesh.
+
+![](vert_mode.png)(w:400)
+
+### Vertex Tools
+
+[REGION:imagetable]
+| ![](vert_move.png)(w:295) | ![](vert_extend.png)(w:295) |
+| --- | --- |
+| Move | Extend |
+[/REGION]
+
+#### Edges
+
+The **Edges** selection mode limits editing to edges of the mesh.
+
+![](edge_mode.png)(w:400)
+
+### Edge Tools
+
+[REGION:imagetable]
+| ![](edge_move.png)(w:235) | ![](edge_split.png)(w:235) | ![](edge_split_drag.png)(w:235) | ![](edge_loop.png)(w:235) | ![](edge_extend.png)(w:235) |
+| --- | --- | --- | --- | --- |
+| Move | Split | Split and Drag Vert | Add Loop | Extend |
+[/REGION]
+
+#### Polygons
+
+The **Polygon** selection mode limits editing to polygon faces of the mesh.
+
+![](poly_mode.png)(w:400)
+
+### Polygon Tools
+
+[REGION:imagetable]
+| ![](poly_move.png)(w:295) | ![](poly_extrude.png)(w:295) | ![](poly_inset.png)(w:295) | ![](poly_bevel.png)(w:295)
+| --- | --- | --- | --- |
+| Move | Extrude | Inset | Bevel |
+[/REGION]
+
+#### Any
+
+The **Any** selection mode changes to match the selected element - Vertex, Edge, Polygon, etc. This is essentially a free editing mode that enables you to quickly jump between the other modes just by selecting a different type of element.
+
+[REGION:imagetable]
+| ![](vert_mode.png)(w:295) | ![](edge_mode.png)(w:295) | ![](poly_mode.png)(w:295)
+| --- | --- | --- |
+| Vertices | Edges | Polygons |
+[/REGION]
+
+### Source vs Instance
+
+Mesh editing can be done on a per-instance basis, which only affects the mesh of the placed Actor in the level; or editing can be applied directly to the source mesh asset, propagating changes you make to all instances of that mesh asset. By default, editing is applied to the source mesh asset. To enable per-instance editing, press the **Per Instance** button in the Mesh Editor mode panel.
+
+![](per_instance.png)
+
+### Subdivision Surfaces
+
+Add or Remove Subdivisions from any of the four submodes
+
+#### Subdivision Levels
+
+[REGION:imagetable]
+| ![](mesh_subd_0.png)(w:295) | ![](mesh_subd_1.png)(w:295) | ![](mesh_subd_2.png)(w:295) | ![](mesh_subd_3.png)(w:295) |
+| --- | --- | --- | --- |
+| Original | SubD Level 1 | SubD Level 2 |  SubD Level 3 |
+[/REGION]
+
+#### Editing Tools with Subdivision Surfaces
+
+The editing tools always act on the elements of the low-poly mesh, and the subdivision is applied on top of those edits.
+
+[REGION:imagetable]
+| ![](poly_inset.png)(w:295) | ![](poly_extrude.png)(w:295) | ![](poly_inset_subd.png)(w:295) | ![](poly_extrude_subd.png)(w:295) |
+| --- | --- | --- | --- |
+| Inset | Extrude | Inset w/ SubD |  Extrude w/ SubD |
+[/REGION]
+
+#### Corner Sharpness
+
+Corner sharpness enables you to control how closely the subdivision surface follows the low-poly mesh at specific locations, which provides the ability to create hard edges on subdivision surfaces.
+
+[REGION:imagetable]
+| Vertex Sharpness |||
+| --- | --- | --- |
+| ![](vert_sharp_0.png)(w:295) | ![](vert_sharp_1.png)(w:295) | ![](vert_sharp_2.png)(w:295)
+| Soft | Sharper | Sharpest |
+[/REGION]
+
+
+[REGION:imagetable]
+| Edge Sharpness |||
+| --- | --- | --- |
+| ![](edge_sharp_0.png)(w:295) | ![](edge_sharp_1.png)(w:295) | ![](edge_sharp_2.png)(w:295)
+| Soft | Sharper | Sharpest |
+[/REGION]
+
+## VR Editor Controls
+
+VR Mode is enabled by default in this build, and you can enter with the hotkey (**Alt + ~**) or via the button on the Level Editor Toolbar. In addition, auto-entry will take you into VR Mode when the Oculus or Vive Headset is on.
+
+[REGION:note]Auto-entry can be enabled or disabled in the **Editor Preferences** under **Experimental > VR Editor**.[/REGION]
+
+Controllers are now Asynchronous on both Vive and Oculus. The left controller is used for access to menus and options, while the right controller is used to select and manipulate items.
+
+![](controllers.png)(w:640)
+
+Teleportation no longer grays out the world and automatically resets the World Scale back to 1:1.
+
+Half-trigger press options have been removed. Interacting with the world is now done via full trigger presses.
+
+**OCULUS ONLY**: A "hard reset / Panic Button" has been added. Clicking in the left thumbstick will perform the following actions at once:
+
+* Deselect all selected items
+* Close all open windows / Asset Editors
+* Cancel Simulate 
+* Place the user in Actor / Selection Mode
+
+Options to Undo/Redo have been assigned to the right thumbstick/touchpad.
+
+[REGION:simpletable]
+| Command | Action |
+| --- | --- |
+| **Undo** | Press left on the thumbstick/touchpad |
+| **Redo** | Press right on the thumbstick/touchpad |
+[/REGION]
+
+Options to quickly place an asset currently selected in the Content Browser have been assigned to the right thumbstick / touchpad. This functionality is also present in Simulate mode.
+
+With an asset selected in the Content Browser:
+
+[REGION:simpletable]
+| Platform | Action |
+| --- | --- |
+| **Oculus** | Press and hold Up on the right thumbstick. Click in the world with the right trigger. |
+| **Vive** | Click and hold Up on the right touchpad. Click in the world with the right trigger. |
+[/REGION]
+
+Transformation Gizmos have been updated for consistency with Desktop Editor.
+
+Rotation Gizmo now has an indicator for the asset's original position.
+
+## Editor Windows
+
+Windows are no longer docked to the Controller when opened. Instead, they open slightly offset from the left controller. The Close button is now nested beside the interaction bar underneath each window
+
+When opening a window, if you keep the trigger depressed, you will see a preview of where the window will be placed.
+
+[REGION:asyncgif]
+![](vr_contentBrowser.gif)(w:322)
+[/REGION]
+
+[REGION:note]
+The **Toggle All UI** option has been removed.
+[/REGION]
+
+## Radial Menu
+
+The new Radial Menu provides a redesigned interface that enables the user to quickly reach all functions of VR Mode, including the Modes Panels, Asset Editors, Simulate and Play in VR, consolidating the functionality of the old Radial Menu and the Quick Menu into one intuitive user interface.
+	
+[REGION:asyncgif]
+![](vr_radialMenu.gif)(w:388)
+[/REGION]
+
+**To open the Radial Menu:**
+
+[REGION:simpletable]
+| Platform | Action |
+| --- | --- |
+| **Oculus** | Interact with the left thumbstick |
+| **Vive** | Interact with the left touchpad |
+[/REGION]
+
+A context-sensitive option has been added to the Radial Main Menu. This menu option offers quick access to additional features available within different tools.
+
+[REGION:asyncgif]
+![](vr_contextOption.gif)(w:344)
+[/REGION]
diff --git a/Engine/Documentation/Source/Engine/Editor/VR/QuickSelect/QuickSelectMenu.INT.udn b/Engine/Documentation/Source/Engine/Editor/VR/QuickSelect/QuickSelectMenu.INT.udn
--- a/Engine/Documentation/Source/Engine/Editor/VR/QuickSelect/QuickSelectMenu.INT.udn
+++ b/Engine/Documentation/Source/Engine/Editor/VR/QuickSelect/QuickSelectMenu.INT.udn
@@ -6,74 +6,80 @@ parent: Engine/Editor/VR
 type: reference
 order: 2
 tags: VR
-version: 4.13
+version: 4.15
 skilllevel: intermediate
+topic-image:quickselect_topic.png
+prereq:Engine/Editor/VR/ActivateVRMode
 
-[VAR:Topic]
-[OBJECT:Topic]
-	[PARAM:image]
-		![%Engine/Editor/VR/QuickSelect:title%](Engine/Editor/VR/QuickSelect/quickselect_topic.png)
-	[/PARAM]
-	[PARAM:title]
-		%Engine/Editor/VR/QuickSelect:title%
-	[/PARAM]
-	[PARAM:path]
-		[RELATIVE:Engine/Editor/VR/QuickSelect]
-	[/PARAM]
-	[PARAM:description]
-		%Engine/Editor/VR/QuickSelect:description%
-	[/PARAM]
-[/OBJECT]
-[/VAR]
-
-[VAR:TopicCompact]
-[OBJECT:TopicCompact]
-	[PARAM:image]
-		![%Engine/Editor/VR/QuickSelect:title%](Engine/Editor/VR/QuickSelect/quickselect_topic.png)
-	[/PARAM]
-	[PARAM:title]
-		%Engine/Editor/VR/QuickSelect:title%
-	[/PARAM]
-	[PARAM:path]
-		[RELATIVE:Engine/Editor/VR/QuickSelect]
-	[/PARAM]
-	[PARAM:description]
-		%Engine/Editor/VR/QuickSelect:description%
-	[/PARAM]
-[/OBJECT]
-[/VAR]
+
+[TOC(start:1 end:4)]
 
 The **Quick Select** menu enables you to quickly adjust grid, rotation and snap settings as well as summon the Content Browser, Details panel, and more.
 
-![](QuickMenuVR.png)(w:350)
+![](QuickMenuVR.png)(w:450)
 
 ## Using the Quick Select Menu
 
-1. Aim one of the controllers at the other controller to display the Quick Select menu.
+1. Aim one of the controllers at the other controller to display the **Quick Select** menu.
 1. Aim the free controller at the options on the menu to highlight them.
 1. Squeeze the Trigger to select the highlighted option.
 
 ## Quick Select Menu Options
 
+As you can see in the image, the **Quick Select** menu is grouped into four menu options, namely Grid Snap, Interface, Tools, and Gizmo. 
+
+### Grid Snap
+
+In the **Grid Snap** options you can adjust the snap settings for the **Tranlate**, **Rotation**, and **Scale** menu options. By clicking on the value below the selection you can cyle through its snap values.
+
+![](GridSnapOptions.png)(w:350)
+
+| Option | Description |
+| ------ | ----------- |
+| Translate | This toggles grid snapping on or off when moving Actors. The snapping value is modified by aiming at the numerical box below, squeezing the trigger, and moving the controller. |
+| Rotate | This toggles grid snapping on or off when rotating Actors. The snapping value is modified by aiming at the numerical box below, squeezing the trigger, and moving the controller. |
+| Scale | This toggles grid snapping on or off when scaling Actors. The snapping value is modified by aiming at the numerical box below, squeezing the trigger, and moving the controller. |
+
+### Interface
+
+In the **Interface** options you have quick access to the various editor windows and panels that are available in the Editor. 
+
+![](InterfaceOptions.png)(w:350)
+
+| Option | Description |
+| ------ | ----------- |
+| Details | Opens the **Details** panel for viewing and modifying properties of the selected Actor(s). |
+| Content Browser | Opens the **Content Browser**. |
+| Modes | Opens the **Modes** panel, providing access to Actor placement tools, foliage editing, and mesh painting. |
+| World Outliner | Opens the **World Outliner**, which displays a hierarchical list of all Actors in the world, enabling quick selection, visibility toggling, organization, and actions that an Actor needs to perform. |
+| Asset Editor | Opens the Asset window, which acts as a container for any asset editors (Blueprint Editor, Material Editor, etc.) that may be open. |
+| Tutorial | Opens the **VR Mode Tutorial** window, which provides access to a collection of instructional walkthroughs for common features that are found in the VR Editor. |
+| World | Opens the **World Settings** panel, providing access to properties of the current Level. |
+| Toggle All UI | This will open the available UI Interace windows available under the Interface section of this Quick Select Menu. |
+
+### Tools
+
+In the **Tools** options you can quickly take a screenshot, use a flashlight for your controller, or jump into Play in Editor (PIE) mode to try out your game.
+
+![](ToolsOptions.png)(w:350) 
+
 | Option | Description |
 | ------ | ----------- |
-| [REGION:tablesection]Grid Snap[/REGION] ||
-| Move | Toggles grid snapping on or off when moving Actors. The snapping value is modified by aiming at the numerical box below, squeezing the trigger, and moving the controller. |
-| Rotate | Toggles grid snapping on or off when rotating Actors. The snapping value is modified by aiming at the numerical box below, squeezing the trigger, and moving the controller. |
-| Scale | Toggles grid snapping on or off when scaling Actors. The snapping value is modified by aiming at the numerical box below, squeezing the trigger, and moving the controller. |
-| [REGION:tablesection]Interface[/REGION] ||
-| Content | Opens the Content Browser. |
-| Details | Opens the Details panel for viewing and modifying properties of the selected Actor(s). |
-| Modes | Opens the Modes panel, providing access to Actor placement tools, foliage editing, and mesh painting. |
-| Outliner | Opens the World Outliner which displays a hierarchical list of all Actors in the world enabling quick selection, visibility toggling, organization, and Actor-specific actions to be performed. |
-| Asset | Opens the Asset window, which acts as a container for any asset editors - Blueprint Editor, Material Editor, etc. - that may be open. |
-| Tutorial | Opens the VR Mode Tutorial, which provides access to a collection of instructional walkthroughs for common features of the VR Editor. |
-| Light | Adds a flashlight to, or removes the flashlight from, the controller that was used to click the button.  |
 | Screenshot | Takes a screenshot of the current scene. |
-| World | Opens the World Settings panel, providing access to properties of the current Level. |
-| [REGION:tablesection]Gizmo[/REGION] ||
-| Coordinate System | Sets the coordinate system [World Space/Local Space] to use when transforming Actors. |
-| Transform Selector | Cycles the Universal Gizmo through the available transform modes: Translate, Rotate, Scale. Only the manipulators for the selected mode are displayed. |
-| [REGION:tablesection]Play[/REGION] ||
+| Flashlight | Adds a flashlight to, or removes the flashlight from, the controller that was used to click the button.  |
 | Play | Starts your project from within the VR Editor so that you can test your gameplay quickly. You will begin gameplay from the player start location. To instantly go back to the VR Editor from your test session, press both triggers and a grip button on each motion controller.  |
 
+### Gizmo
+
+In the **Gizmo** options you can quickly cycle through and change the bahavior of your Translate, Rotation, and Scale gizmos.
+
+![](GizmoOptions.png)(w:350) 
+
+| Option | Description |
+| ------ | ----------- |
+| Coordinate System | Sets the coordinate system to either World Space or Local Space that can be used when transforming Actors. |
+| Transform Selector | This will cycle through the Universal Gizmo through the available transform modes (Translate, Rotate, and Scale). Only the manipulators for the selected mode are displayed.  |
+| Cycle Gizmo | When an Actor is selected in the viewport, this will cycle through the Translate, Rotate, and Scale gizmo. |
+
+
+
diff --git a/Engine/Documentation/Source/Engine/Editor/VR/RadialMenu/RadialMenu.INT.udn b/Engine/Documentation/Source/Engine/Editor/VR/RadialMenu/RadialMenu.INT.udn
--- a/Engine/Documentation/Source/Engine/Editor/VR/RadialMenu/RadialMenu.INT.udn
+++ b/Engine/Documentation/Source/Engine/Editor/VR/RadialMenu/RadialMenu.INT.udn
@@ -6,46 +6,17 @@ parent: Engine/Editor/VR
 type: reference
 order: 3
 tags: VR
-version: 4.11
+version: 4.15
 skilllevel: intermediate
+topic-image:radialmenu.png
+prereq:Engine/Editor/VR/ActivateVRMode
 
-[VAR:Topic]
-[OBJECT:Topic]
-	[PARAM:image]
-		![%Engine/Editor/VR/RadialMenu:title%](Engine/Editor/VR/RadialMenu/radialmenu.png)(fill:#919293)
-	[/PARAM]
-	[PARAM:title]
-		%Engine/Editor/VR/RadialMenu:title%
-	[/PARAM]
-	[PARAM:path]
-		[RELATIVE:Engine/Editor/VR/RadialMenu]
-	[/PARAM]
-	[PARAM:description]
-		%Engine/Editor/VR/RadialMenu:description%
-	[/PARAM]
-[/OBJECT]
-[/VAR]
 
-[VAR:TopicCompact]
-[OBJECT:TopicCompact]
-	[PARAM:image]
-		![%Engine/Editor/VR/RadialMenu:title%](Engine/Editor/VR/RadialMenu/radialmenu.png)(fill:#919293)
-	[/PARAM]
-	[PARAM:title]
-		%Engine/Editor/VR/RadialMenu:title%
-	[/PARAM]
-	[PARAM:path]
-		[RELATIVE:Engine/Editor/VR/RadialMenu]
-	[/PARAM]
-	[PARAM:description]
-		%Engine/Editor/VR/RadialMenu:description%
-	[/PARAM]
-[/OBJECT]
-[/VAR]
+[TOC(start:1 end:4)]
 
 The **Radial Menu** provides access to common actions like undo, redo, paste, etc.
 
-![%Engine/Editor/VR/RadialMenu:title%](Engine/Editor/VR/RadialMenu/radialmenu_world.png)(w:640)
+![%Engine/Editor/VR/RadialMenu:title%](Engine/Editor/VR/RadialMenu/radialmenu_world.png)(w:450)
 
 ## Using the Radial Menu
 
@@ -55,13 +26,19 @@ The **Radial Menu** provides access to common actions like undo, redo, paste, et
 
 ## Radial Menu Options
 
+These are the available options when you open up the Radial Menu in the VR Editor. 
+
 | Option | Description |
 | ------ | ----------- |
-| Universal Gizmo | Toggles the Universal Gizmo on or off. |
-| Delete | Deletes the selected Actor(s). |
-| Duplicate | Creates duplicates of the selected Actor(s). |
-| Copy | Copies the selected Actor(s) to the clipboard. |
-| Paste | Pastes the contents of the clipboard into the world. |
-| Undo | Reverts the last action performed. |
-| Redo | Performs the last undone action. |
-| Snap To Ground | Moves the selected Actor(s) down until they collide with other geometry. |
\ No newline at end of file
+| **Snap Actors to Ground** | Moves the selected Actor(s) down until they collide with other geometry. |
+| **Undo** | Reverts the last action performed. |
+| **Paste** | Pastes the contents of the clipboard into the world. |
+| **Duplicate** | Creates duplicates of the selected Actor(s). |
+| **Copy** | Copies the selected Actor(s) to the clipboard. |
+| **Redo** | Performs the last undone action. |
+| **Delete** | Deletes the selected Actor(s). |
+| **Cut** | Copies the selected Actor(s) to the clipboard and removes it from the scene. |
+
+
+
+
diff --git a/Engine/Documentation/Source/Engine/Editor/VR/Setup/VREditorEditorSetup.INT.udn b/Engine/Documentation/Source/Engine/Editor/VR/Setup/VREditorEditorSetup.INT.udn
--- a/Engine/Documentation/Source/Engine/Editor/VR/Setup/VREditorEditorSetup.INT.udn
+++ b/Engine/Documentation/Source/Engine/Editor/VR/Setup/VREditorEditorSetup.INT.udn
@@ -4,45 +4,12 @@ Crumbs:
 Availability: public
 parent: Engine/Editor/VR
 type:how-to
-order: 1
+order: 2
 tags: VR
 tags: GitHub
-version: 4.11
+version: 4.15
 skilllevel: intermediate
-
-[VAR:Topic]
-[OBJECT:Topic]
-	[PARAM:image]
-		![%Engine/Editor/VR/Setup:title%](Engine/Editor/VR/Setup/vreditor_setup.png)
-	[/PARAM]
-	[PARAM:title]
-		%Engine/Editor/VR/Setup:title%
-	[/PARAM]
-	[PARAM:path]
-		[RELATIVE:Engine/Editor/VR/Setup]
-	[/PARAM]
-	[PARAM:description]
-		%Engine/Editor/VR/Setup:description%
-	[/PARAM]
-[/OBJECT]
-[/VAR]
-
-[VAR:TopicCompact]
-[OBJECT:TopicCompact]
-	[PARAM:image]
-		![%Engine/Editor/VR/Setup:title%](Engine/Editor/VR/Setup/vreditor_setup.png)
-	[/PARAM]
-	[PARAM:title]
-		%Engine/Editor/VR/Setup:title%
-	[/PARAM]
-	[PARAM:path]
-		[RELATIVE:Engine/Editor/VR/Setup]
-	[/PARAM]
-	[PARAM:description]
-		%Engine/Editor/VR/Setup:description%
-	[/PARAM]
-[/OBJECT]
-[/VAR]
+topic-image:vreditor_setup.png
 
 [EXCERPT:main]
 
diff --git a/Engine/Documentation/Source/Engine/Editor/VR/UniversalGizmo/UniversalGizmo.INT.udn b/Engine/Documentation/Source/Engine/Editor/VR/UniversalGizmo/UniversalGizmo.INT.udn
--- a/Engine/Documentation/Source/Engine/Editor/VR/UniversalGizmo/UniversalGizmo.INT.udn
+++ b/Engine/Documentation/Source/Engine/Editor/VR/UniversalGizmo/UniversalGizmo.INT.udn
@@ -6,46 +6,19 @@ parent: Engine/Editor/VR
 type: reference
 order: 4
 tags: VR
-version: 4.13
+version: 4.15
 skilllevel: intermediate
+topic-image:universalgizmo_topic.png
+prereq:Engine/Editor/VR/ActivateVRMode
 
-[VAR:Topic]
-[OBJECT:Topic]
-	[PARAM:image]
-		![%Engine/Editor/VR/UniversalGizmo:title%](Engine/Editor/VR/UniversalGizmo/universalgizmo_topic.png)
-	[/PARAM]
-	[PARAM:title]
-		%Engine/Editor/VR/UniversalGizmo:title%
-	[/PARAM]
-	[PARAM:path]
-		[RELATIVE:Engine/Editor/VR/UniversalGizmo]
-	[/PARAM]
-	[PARAM:description]
-		%Engine/Editor/VR/UniversalGizmo:description%
-	[/PARAM]
-[/OBJECT]
-[/VAR]
-
-[VAR:TopicCompact]
-[OBJECT:TopicCompact]
-	[PARAM:image]
-		![%Engine/Editor/VR/UniversalGizmo:title%](Engine/Editor/VR/UniversalGizmo/universalgizmo_topic.png)
-	[/PARAM]
-	[PARAM:title]
-		%Engine/Editor/VR/UniversalGizmo:title%
-	[/PARAM]
-	[PARAM:path]
-		[RELATIVE:Engine/Editor/VR/UniversalGizmo]
-	[/PARAM]
-	[PARAM:description]
-		%Engine/Editor/VR/UniversalGizmo:description%
-	[/PARAM]
-[/OBJECT]
-[/VAR]
+
+[TOC(start:1 end:4)]
 
 The **Universal Gizmo** is a transform widget designed for moving, rotating, and scaling Actors when working in VR mode. There are both local and world space versions available, which can be swapped between using the [](Engine/Editor/VR/QuickSelect). 
 The elements of the Universal Gizmo are colored red, green, and blue representing the X, Y, and Z axes, respectively. Much like in the standard editor, there are gizmos which can be used specifically for translating, rotating, and scaling your Actors. Local and world space versions of these are accessible from the [](Engine/Editor/VR/QuickSelect). 
 
+##Gizmos
+
 | Gizmo | World | Local |
 | ---- |-------| -----|
 | **Universal** | ![](GlobalUniversal.png)(w:350)[INCLUDE:#worlduniversal]|![](LocalUniversal.png)(w:350) [INCLUDE:#localuniversal]  |
diff --git a/Engine/Documentation/Source/Engine/Editor/VR/VREditor.INT.udn b/Engine/Documentation/Source/Engine/Editor/VR/VREditor.INT.udn
--- a/Engine/Documentation/Source/Engine/Editor/VR/VREditor.INT.udn
+++ b/Engine/Documentation/Source/Engine/Editor/VR/VREditor.INT.udn
@@ -6,7 +6,7 @@ parent:Engine/Editor
 type:overview
 order:10
 tags:VR
-version:4.12
+version:4.15
 skilllevel:intermediate
 social-image:vreditor_social.png
 
@@ -79,15 +79,18 @@ In order to use Unreal Editor in VR mode, your development environment must meet
 | **Video Drivers** | Make sure you have the very latest video drivers installed.  Some HMD features require latest drivers. |
 [/EXCERPT:main]
 
-## Starting Out
+## Getting Started
 
-[DIR(end:"1" output:"topic" type:"how-to" tags:"Getting Started" parent:"Engine/Editor/VR")]
+%Engine/Editor/VR/ActivateVRMode:topic%
+%Engine/Editor/VR/Setup:topic%
+%Engine/Editor/VR/Controls:topic%
 
 ## Guides
 
-[DIR(end:"1" output:"topic" type:"how-to" tags:"Level Editing, GitHub" parent:"Engine/Editor/VR")]
+%Engine/Editor/VR/WorldInteraction:topic%
+%Engine/Editor/VR/ActorInteraction:topic%
 
-## Reference
+## References
 
 [DIR(end:"1" output:"topic" type:"reference" parent:"Engine/Editor/VR")]
 
diff --git a/Engine/Documentation/Source/Engine/Editor/VR/WorldInteraction/WorldInteraction.INT.udn b/Engine/Documentation/Source/Engine/Editor/VR/WorldInteraction/WorldInteraction.INT.udn
--- a/Engine/Documentation/Source/Engine/Editor/VR/WorldInteraction/WorldInteraction.INT.udn
+++ b/Engine/Documentation/Source/Engine/Editor/VR/WorldInteraction/WorldInteraction.INT.udn
@@ -7,42 +7,13 @@ type:how-to
 order: 3
 tags: VR
 tags: Level Editing
-version: 4.11
+version: 4.15
 skilllevel: intermediate
+topic-image:worldinteraction.png
+prereq:Engine/Editor/VR/ActivateVRMode
 
-[VAR:Topic]
-[OBJECT:Topic]
-	[PARAM:image]
-		![%Engine/Editor/VR/WorldInteraction:title%](Engine/Editor/VR/WorldInteraction/worldinteraction.png)
-	[/PARAM]
-	[PARAM:title]
-		%Engine/Editor/VR/WorldInteraction:title%
-	[/PARAM]
-	[PARAM:path]
-		[RELATIVE:Engine/Editor/VR/WorldInteraction]
-	[/PARAM]
-	[PARAM:description]
-		%Engine/Editor/VR/WorldInteraction:description%
-	[/PARAM]
-[/OBJECT]
-[/VAR]
-
-[VAR:TopicCompact]
-[OBJECT:TopicCompact]
-	[PARAM:image]
-		![%Engine/Editor/VR/WorldInteraction:title%](Engine/Editor/VR/WorldInteraction/worldinteraction.png)
-	[/PARAM]
-	[PARAM:title]
-		%Engine/Editor/VR/WorldInteraction:title%
-	[/PARAM]
-	[PARAM:path]
-		[RELATIVE:Engine/Editor/VR/WorldInteraction]
-	[/PARAM]
-	[PARAM:description]
-		%Engine/Editor/VR/WorldInteraction:description%
-	[/PARAM]
-[/OBJECT]
-[/VAR]
+
+[TOC(start:1 end:4)]
 
 [EXCERPT:main]
 
diff --git a/Engine/Documentation/Source/Engine/Foliage/Foliage.INT.udn b/Engine/Documentation/Source/Engine/Foliage/Foliage.INT.udn
--- a/Engine/Documentation/Source/Engine/Foliage/Foliage.INT.udn
+++ b/Engine/Documentation/Source/Engine/Foliage/Foliage.INT.udn
@@ -177,6 +177,10 @@ The Paint tool is Foliage Edit Mode's default tool. You can use it to paint inst
 
 ![](Foliage_Paint_Tool_Gimzo.png)
 
+[REGION:tip]
+Foliage painting also works in the VR Editor. For the controls for Foliage painting in VR, see the [](Engine/Editor/VR/Controls).
+[/REGION]
+
 **To add foliage to the brush area:**  
 * Hold down the **Left Mouse** button and **click** where you want to add foliage.
 All currently selected meshes in the mesh list will be added according to their current parameters and settings.
@@ -186,11 +190,6 @@ When painting meshes, the engine performs a number of line traces inside the sph
 **To erase foliage from the brush area:**  
 * Hold down **Shift** and **Left Mouse** button and **click** where you want to remove foliage.
 
-
-[REGION:tip]
-Foliage painting also works in the VR Editor. For the controls for foliage painting in VR, see the [](Engine/Editor/VR/Controls).
-[/REGION]
-
 ### Paint Settings
 
 Each Static mesh in the Mesh List has a number of properties that control how instances are placed, scaled or rotated when the selected Static mesh is painted. 
diff --git a/Engine/Documentation/Source/Engine/Landscape/Creation/Creation.INT.udn b/Engine/Documentation/Source/Engine/Landscape/Creation/Creation.INT.udn
--- a/Engine/Documentation/Source/Engine/Landscape/Creation/Creation.INT.udn
+++ b/Engine/Documentation/Source/Engine/Landscape/Creation/Creation.INT.udn
@@ -52,8 +52,8 @@ The **Landscape** tool will allow you to create massive terrain based worlds tha
 Because of the flexibility of the Landscape system, you can create your Landscape using either of the following methods. 
 
 * You can create a new Landscape heightmap from scratch using the Landscape tools in-engine.
-* You can import a Landscape heightmap created previously in Unreal Editor or created through external tools. To read more about how to do this, please check out
-	the [Creating Custom Heightmaps and Layers](Engine/Landscape/Custom) document.
+* You can import a Landscape heightmap created previously in Unreal Editor or created through external tools. To read more about how to do this, please check out the [](Engine/Landscape/Custom) document.
+* To create your own format for importing Landscapes, please see our [](Engine/Landscape/Creation/Plugin) document.
 
 [REGION:note]
 For examples of the Landscape tool in action, see the [](Resources\ContentExamples\Landscapes) Content Example.
@@ -171,11 +171,3 @@ Once you have pressed the **Create** Button, you should have something that look
 ![Landscape_NewTerrain.png](Landscape_NewTerrain.png)(w:600)
 
 Now with your new Landscape now created, you can now start to [Sculpt](Engine/Landscape/Editing/SculptMode) or [Paint](Engine/Landscape/Editing/PaintMode) the Landscape to your liking. 
-
-
-
-
-
-
-
-
diff --git a/Engine/Documentation/Source/Engine/Landscape/Creation/Plugin/Plugin.INT.udn b/Engine/Documentation/Source/Engine/Landscape/Creation/Plugin/Plugin.INT.udn
new file mode 100644
--- /dev/null
+++ b/Engine/Documentation/Source/Engine/Landscape/Creation/Plugin/Plugin.INT.udn
@@ -0,0 +1,28 @@
+Availability: Public
+Title:Creating Custom Landscape Importers
+Crumbs: %ROOT%, Engine, Engine/Landscape, Engine/Landscape/Creation
+Description:Guide to creating new Landscape import formats via plugin.
+version: 4.15
+parent:Engine/Landscape
+
+Creating your own heightmap and weightmap formats for importing landscape data is accomplished by writing a plugin. The plugin will add your new format to the engine, as well as importing the data from your files.
+
+### Writing A Custom Importer
+
+* In order to create new importers, your plugin should create instances of objects implementing `ILandscapeHeightmapFileFormat` and `ILandscapeWeightmapFileFormat`, and register those objects with `ILandscapeEditorModulemodule::RegisterHeightmapFileFormat` and `ILandscapeEditorModulemodule::RegisterWeightmapFileFormat`, respectively.
+
+* Implementing the `ILandscapeHeightmapFileFormat` interface in a plugin requires overriding the following functions:
+1. **const FLandscapeFileTypeInfo& GetInfo() const**: Returns type information indicating which file types this class handles, and whether or not exporting is supported.
+1. **FLandscapeHeightmapInfo Validate(const TCHAR* HeightmapFilename) const** - Validates the named file, or rejects it and returns an error code and message.
+1. **FLandscapeHeightmapImportData Import(const TCHAR* HeightmapFilename, FLandscapeFileResolution ExpectedResolution) const** - Actually imports the file.
+1. **void Export(const TCHAR* HeightmapFilename, TArrayView<const uint16> Data, FLandscapeFileResolution DataResolution, FVector Scale) const** - Exports the file, if this format supports exporting (see the return value from `GetInfo`). This is the only function that doesn't need to be overridden in order to compile. However, if it is called without being overridden, it will call `check`.
+1. **(Destructor)** - Classes that implement this interface should use a virtual destructor, as they will be deleted via a pointer to the interface class.
+
+* Implementing the `ILandscapeHeightmapFileFormat` interface is nearly identical, with only minor differences in some return types:
+1. **const FLandscapeFileTypeInfo& GetInfo() const** - Returns type information indicating which file types this class handles, and whether or not exporting is supported.
+1. **FLandscapeWeightmapInfo Validate(const TCHAR* WeightmapFilename) const** - Validates the named file, or rejects it and returns an error code and message.
+1. **FLandscapeWeightmapImportData Import(const TCHAR* WeightmapFilename, FLandscapeFileResolution ExpectedResolution) const** - Actually imports the file.
+1. **void Export(const TCHAR* WeightmapFilename, TArrayView<const uint16> Data, FLandscapeFileResolution DataResolution, FVector Scale) const** - Exports the file, if this format supports exporting (see the return value from `GetInfo`). This is the only function that doesn't need to be overridden in order to compile. However, if it is called without being overridden, it will call `check`.
+1. **(Destructor)** - Classes that implement this interface should use a virtual destructor, as they will be deleted via a pointer to the interface class.
+
+* For further information and examples, you can see the interfaces in `LandscapeFileFormatInterfaces.h`, the .PNG implementations in `LandscapeFileFormatPng.cpp` and `LandscapeFileFormatPng.h` , and the .RAW implementations in `LandscapeFileFormatRaw.cpp` and `LandscapeFileFormatRaw.h`. All of this code is in the LandscapeEditor module in the engine.
diff --git a/Engine/Documentation/Source/Engine/Landscape/Custom/CreatingCustomHeightmapsAndLayers.INT.udn b/Engine/Documentation/Source/Engine/Landscape/Custom/CreatingCustomHeightmapsAndLayers.INT.udn
--- a/Engine/Documentation/Source/Engine/Landscape/Custom/CreatingCustomHeightmapsAndLayers.INT.udn
+++ b/Engine/Documentation/Source/Engine/Landscape/Custom/CreatingCustomHeightmapsAndLayers.INT.udn
@@ -1,46 +1,35 @@
 Availability: Public
 Title:Creating and Using Custom Heightmaps and Layers
-Crumbs: %ROOT%, Engine, Engine/Landscape
+Crumbs:
 Description:Creating and using custom Heightmaps and layers for your landscapes.
 version: 4.14
 parent:Engine/Landscape/Creation
-
+order:
+type:overview
+SkillLevel:Intermediate
+related:Engine/Landscape/Editing
+related:Engine/Landscape/Editing/PaintMode
+related:Engine/Landscape/Creation
 
 [TOC(start:2 end:3)]
 
-![](Landscape_Example_Image.png)
-
 Sometimes your Landscape will require that you use external programs to create both the height map and layers that are needed. 
-UE4 does accommodate this style of workflow and everything that you need to be aware of when using a workflow like this is in this guide. 
-If this is your first time using the Landscape tools, you might want to check out the [Landscape Overview](Engine/Landscape/Editing) first as what is going to be covered here may be too advanced for new users. 
+Unreal Engine 4 (UE4) does accommodate this style of workflow and everything that you need to be aware of when using a workflow like this, is in this guide. 
 
+![](Landscape_Example_Image.png)
+
+[REGION:note]
+If this is your first time using the Landscape tools, you might want to check out the [Landscape Overview](Engine/Landscape/Editing) first, as what is going to be covered here may be too advanced for new users. 
+[/REGION]
 
 ## Layers
 
 Landscape Layers are special textures that help to break up the look and feel of the Landscape. 
 
 ### Layer Formats
-* Landscape layers can be imported from external programs via implementation of the `ILandscapeHeightmapFileFormat` and `ILandscapeWeightmapFileFormat` interfaces. The editor's existing support for image-based importing has been converted to use this interface and is still fully supported. Images for the built-in formats are still required to be grayscale, 8 bits per pixel, single channel files in either .PNG or .RAW format. If you are creating layers in Photoshop, you should use the following settings when creating a new document:
-
-	![](Photoshop_Layer_Example.png)
+Landscape layers can be imported from external programs via implementation of the `ILandscapeHeightmapFileFormat` and `ILandscapeWeightmapFileFormat` interfaces. The Editor's existing support for image-based importing has been converted to use this interface and is still fully supported. Images for the built-in formats are still required to be grayscale, 8 bits per pixel, single channel files in either .PNG or .RAW format. If you are creating layers in Photoshop, you should use the following settings when creating a new document:
 
-* In order to create new importers, your plugin should create instances of objects implementing `ILandscapeHeightmapFileFormat` and `ILandscapeWeightmapFileFormat`, and register those objects with `ILandscapeEditorModulemodule::RegisterHeightmapFileFormat` and `ILandscapeEditorModulemodule::RegisterWeightmapFileFormat`, respectively.
-
-* Implementing the `ILandscapeHeightmapFileFormat` interface in a plugin requires overriding the following functions:
-1. **const FLandscapeFileTypeInfo& GetInfo() const**: Returns type information indicating which file types this class handles, and whether or not exporting is supported.
-1. **FLandscapeHeightmapInfo Validate(const TCHAR* HeightmapFilename) const** - Validates the named file, or rejects it and returns an error code and message.
-1. **FLandscapeHeightmapImportData Import(const TCHAR* HeightmapFilename, FLandscapeFileResolution ExpectedResolution) const** - Actually imports the file.
-1. **void Export(const TCHAR* HeightmapFilename, TArrayView<const uint16> Data, FLandscapeFileResolution DataResolution, FVector Scale) const** - Exports the file, if this format supports exporting (see the return value from `GetInfo`). This is the only function that doesn't need to be overridden in order to compile. However, if it is called without being overridden, it will call `check`.
-1. **(Destructor)** - Classes that implement this interface should use a virtual destructor, as they will be deleted via a pointer to the interface class.
-
-* Implementing the `ILandscapeHeightmapFileFormat` interface is nearly identical, with only minor differences in some return types:
-1. **const FLandscapeFileTypeInfo& GetInfo() const** - Returns type information indicating which file types this class handles, and whether or not exporting is supported.
-1. **FLandscapeWeightmapInfo Validate(const TCHAR* WeightmapFilename) const** - Validates the named file, or rejects it and returns an error code and message.
-1. **FLandscapeWeightmapImportData Import(const TCHAR* WeightmapFilename, FLandscapeFileResolution ExpectedResolution) const** - Actually imports the file.
-1. **void Export(const TCHAR* WeightmapFilename, TArrayView<const uint16> Data, FLandscapeFileResolution DataResolution, FVector Scale) const** - Exports the file, if this format supports exporting (see the return value from `GetInfo`). This is the only function that doesn't need to be overridden in order to compile. However, if it is called without being overridden, it will call `check`.
-1. **(Destructor)** - Classes that implement this interface should use a virtual destructor, as they will be deleted via a pointer to the interface class.
-
-* For further information and examples, you can see the interfaces in `LandscapeFileFormatInterfaces.h`, the .PNG implementations in `LandscapeFileFormatPng.cpp1` and `LandscapeFileFormatPng.h` , and the .RAW implementations in `LandscapeFileFormatRaw.cpp` and `LandscapeFileFormatRaw.h`. All of this code is in the LandscapeEditor module in the engine.
+![](Photoshop_Layer_Example.png)
 
 <!--
 ### Landscape Size vs. Layer Size
@@ -70,35 +59,37 @@ If your layers are not output at the correct size you will receive the following
 Importing Layers that were made in an external application is a simple process but you first need to make sure that you have a few things set up in order to get everything to work smoothly. 
 
 1. You first need to make sure that you have a Landscape to work with. If you do not have one now, please create one. 
-1. Once your Landscape has been created, you need to make a new material for it. For this example, we will be making a very simple Material but this Material could easily be expanded upon if needed. The Material that you need to setup should look something like this.
+
+1. Once your Landscape has been created, you need to make a new material for it. For this example, we will be making a very simple Material but this Material could easily be expanded upon if needed. The Material that you need to set up should look something like this.
 
 	![](Landscape_Simple_Shader.png)
+
 1. Once the material is created, apply it to the Landscape Actor so that you have something that looks like this. 
 
 	![](Landscape_Applied_Material.png)
-1. Now with the material applied, it is now time to add some  **Layer Info** to our Landscape Actor. Add a Layer Info object for each layer you are going to have. For this example, we are going to have three. To read more about Layer Info objects, check out the [Layer Info Objects](Engine/Landscape/Editing/PaintMode#LayerInfOobjects) section of the [Paint Mode](Engine/Landscape/Editing/PaintMode) document.
+
+1. Now with the material applied, it is now time to add **Layer Info** to our Landscape Actor. Add a **Layer Info** object for each layer you are going to have. For this example, we are going to have three. To read more about **Layer Info** objects, check out the [Layer Info Objects](Engine/Landscape/Editing/PaintMode#LayerInfOobjects) section of the [Paint Mode](Engine/Landscape/Editing/PaintMode) document.
 
 	![](Landscape_Create_Layer.png)
+
 1. When completed, your **Target Layers** section should look like this. 
 
 	![](Landscape_Target_Layers.png)
-1. Now it is time to import our layer. To do this, all you need to do is **Right-Click** on the Target Layer you want to import a layer for and then select the layer that you want to import from the dialogue box that pops up. 
+
+1. Now it is time to import our layer. To do this, all you need to do is right-click on the **Target Layer** you want to import a layer for, and then select the layer that you want to import from the dialogue box that pops up. 
 
 	![](Landscape_Import_Layer_Option.png)
+
 1. If your layers are not output at the correct size, you will receive the following warning.  
 
 	![](Landscape_Layer_Import_Error.png)
  
-	 To fix this issue, you need to determine how big the Landscape layer you are trying to import over is supposed to be. To find this out, you can export any layer by **Right-Clicking** on the Target Layer you want to look at then select **Export** from the dialogue box that pops up.
+	 To fix this issue, you need to determine how big the **Landscape** layer you are trying to import over is supposed to be. To find this out, you can export any layer by right-clicking on the **Target Layer** you want to look at, and then select **Export** from the dialogue box that pops up.
 	 Once you do that, you will be asked to save the layer file somewhere on your PC. After the file has been saved, you can then open it up to determine what size your layers should be. 
 
-
-
-
-
 ## Height Maps
 
-Using external tools to create a base heightmap to work off of inside Unreal can be a good way of speeding up the Landscape creation process. The base heightmap can then be imported, cleaned up, or modified using the editing tools inside Unreal Editor to customize the Landscape and make it better fit into the world and desired game play. 
+Using external tools to create a base heightmap to work off of inside Unreal can be a good way of speeding up the Landscape creation process. The base heightmap can then be imported, cleaned up, or modified using the editing tools inside Unreal Editor to customize the Landscape, making it a better fit with the world and desired game play. 
 
 ### Height Map Formats
 
@@ -108,7 +99,7 @@ When exporting heightmaps from external programs, only the following formats are
 
 * 16-bit, grayscale .RAW file in little-endian byte order.
 
-If you are using Photoshop to export your heightmap from make sure that you set up your new images like the following. 
+If you are using Photoshop to export your heightmap, make sure that you set up your new images like the following. 
 ![](Photoshop_HeighMap_Example.png)
 
 [REGION:tip]
@@ -116,23 +107,21 @@ When making a heightmap in an external application, a value of  **White** (255 i
 [/REGION]
 
 ### Importing Custom Heightmaps 
-Importing custom Heightmap is quite simple. The process is initiated from the **New Landscape** section of the Landscape Modes panel. To active it, select the radio box labeled **Import from File** to expose the ability to import a heightmap from a file.
+Importing a Custom Heightmap is quite simple. The process is initiated from the **New Landscape** section of the Landscape Modes panel. To active it, select the radio box labeled **Import from File** to expose the ability to import a heightmap from a file.
 
 ![](Engine\Landscape\Creation\Landscape_Import_Section.png)
 
-
-To select a Heightmap to import, press the button next the box that says "Please specify a heightmap" and then select your Heightmap using the file browser. 
+To select a Heightmap to import, press the button next to the box that says "Please specify a heightmap", and then select your Heightmap using the file browser. 
 
 ![](Engine\Landscape\Creation\Landscape_Import_Section_Button.png)
 
-Now let's try importing a custom height map. For this example, we will be using the image below so make sure that you do load it.  
+Now let's try importing a custom height map. For this example, we will be using the image below.  
 
+![](Engine\Landscape\Creation\Landscape_Test_HM.png)(convert:false)
 [REGION:caption]
-Make sure you right-click Save as on the image to download it to your PC.
+Make sure you right-click **Save as** on the image (above) to download it to your PC.
 [/REGION]
 
-![](Engine\Landscape\Creation\Landscape_Test_HM.png)(convert:false)
-
 Once you have the heightmap saved to your PC, it is now time to use it in the Landscape tools. 
 1. Open up the Landscape tool and go into [Manage Mode.](Engine/Landscape/Editing/ManageMode)
 	
@@ -159,4 +148,3 @@ After pressing the Import Button:
 
 
 Now you have a new Landscape to work with that is based off the Heightmap that you just imported. 
-
diff --git a/Engine/Documentation/Source/Engine/Landscape/Editing/Editing.INT.udn b/Engine/Documentation/Source/Engine/Landscape/Editing/Editing.INT.udn
--- a/Engine/Documentation/Source/Engine/Landscape/Editing/Editing.INT.udn
+++ b/Engine/Documentation/Source/Engine/Landscape/Editing/Editing.INT.udn
@@ -63,6 +63,10 @@ For a quick reference guide to the various Landscape tool modes and options, see
 
 For examples of the Landscape tool in action, see the [](Resources\ContentExamples\Landscapes).
 
+[REGION:tip]
+The Landscape Tools also work in the VR Editor. For the controls for using Landscape in VR, see the [](Engine/Editor/VR/Controls).
+[/REGION]
+
 
 [REGION:topics third]
 [OBJECT:TopicCompact]
diff --git a/Engine/Documentation/Source/Engine/Landscape/Editing/ManageMode/ManageMode.INT.udn b/Engine/Documentation/Source/Engine/Landscape/Editing/ManageMode/ManageMode.INT.udn
--- a/Engine/Documentation/Source/Engine/Landscape/Editing/ManageMode/ManageMode.INT.udn
+++ b/Engine/Documentation/Source/Engine/Landscape/Editing/ManageMode/ManageMode.INT.udn
@@ -31,6 +31,11 @@ In **Manage** mode, you can create new Landscapes, and you can use the component
 
 For information about creating Landscapes, see [](Engine/Landscape\Creation).
 
+[REGION:tip]
+Landscape Manage Mode also works in the VR Editor. For the controls for using Landscape in VR, see the [](Engine/Editor/VR/Controls).
+[/REGION]
+
+
 ## Landscape Selection
 
 In the **Landscape** tool, in any mode, you can select which Landscape you want to modify from a list of all the Landscapes currently in your level.
diff --git a/Engine/Documentation/Source/Engine/Landscape/Editing/PaintMode/PaintMode.INT.udn b/Engine/Documentation/Source/Engine/Landscape/Editing/PaintMode/PaintMode.INT.udn
--- a/Engine/Documentation/Source/Engine/Landscape/Editing/PaintMode/PaintMode.INT.udn
+++ b/Engine/Documentation/Source/Engine/Landscape/Editing/PaintMode/PaintMode.INT.udn
@@ -30,6 +30,10 @@ The tools in **Paint** mode enable you to modify the appearance of your Landscap
 
 For more information about Landscape Materials, see [Landscape Materials](Engine/Landscape\Materials).
 
+[REGION:tip]
+Landscape Paint Mode also works in the VR Editor. For the controls for using Landscape in VR, see the [](Engine/Editor/VR/Controls).
+[/REGION]
+
 
 ## Painting Tools
 
diff --git a/Engine/Documentation/Source/Engine/Landscape/Editing/SculptMode/SculptMode.INT.udn b/Engine/Documentation/Source/Engine/Landscape/Editing/SculptMode/SculptMode.INT.udn
--- a/Engine/Documentation/Source/Engine/Landscape/Editing/SculptMode/SculptMode.INT.udn
+++ b/Engine/Documentation/Source/Engine/Landscape/Editing/SculptMode/SculptMode.INT.udn
@@ -29,6 +29,11 @@ version: 4.9
 
 The process of sculpting a Landscape involves using one or more tools that modify the underlying heightmap. These tools range from the simple Sculpt tool to other tools that use complex algorithms to apply interesting effects to the heightmap, such as erosion. Each tool has a set of properties that determine how the tool affects the Landscape.
 
+[REGION:tip]
+Landscape Paint Mode also works in the VR Editor. For the controls for using Landscape in VR, see the [](Engine/Editor/VR/Controls).
+[/REGION]
+
+
 The currently selected tool determines the operation that will be performed on the Landscape. The tools use [brushes](Engine\Landscape\Editing\Brushes) to apply their effects using strokes.
 
 | **Common Controls** | **Operation** |
diff --git a/Engine/Documentation/Source/Engine/Matinee/MatineeAndCinematics.INT.udn b/Engine/Documentation/Source/Engine/Matinee/MatineeAndCinematics.INT.udn
--- a/Engine/Documentation/Source/Engine/Matinee/MatineeAndCinematics.INT.udn
+++ b/Engine/Documentation/Source/Engine/Matinee/MatineeAndCinematics.INT.udn
@@ -5,45 +5,8 @@ Description:Animation tool that provides the ability to animate the properties o
 Version: 4.9
 parent:Engine
 order:15
+topic-image:matinee_topic.png
 
-[VAR:Topic]
-[OBJECT:Topic]
-	[PARAM:image]
-		![%Engine/Matinee:title%](Engine/Matinee/matinee_topic.png)
-	[/PARAM]
-	[PARAM:icon]
-		![](%ROOT%/cine_icon.png)(convert:false)
-	[/PARAM]
-	[PARAM:title]
-		%Engine/Matinee:title%
-	[/PARAM]
-	[PARAM:description]
-		%Engine/Matinee:description%
-	[/PARAM]
-	[PARAM:path]
-		[RELATIVE:Engine/Matinee]
-	[/PARAM]
-[/OBJECT]
-[/VAR]
-[VAR:TopicCompact]
-[OBJECT:TopicCompact]
-	[PARAM:image]
-		![%Engine/Matinee:title%](Engine/Matinee/matinee_topic.png)
-	[/PARAM]
-	[PARAM:icon]
-		![](%ROOT%/cine_icon.png)(convert:false)
-	[/PARAM]
-	[PARAM:title]
-		%Engine/Matinee:title%
-	[/PARAM]
-	[PARAM:description]
-		%Engine/Matinee:description%
-	[/PARAM]
-	[PARAM:path]
-		[RELATIVE:Engine/Matinee]
-	[/PARAM]
-[/OBJECT]
-[/VAR]
 
 [REGION:banner]
 ![matinee.jpg](matinee_landing_banner.png)
diff --git a/Engine/Documentation/Source/Engine/MediaFramework/HowTo/ControllingPlayback/ControllingPlayback.INT.udn b/Engine/Documentation/Source/Engine/MediaFramework/HowTo/ControllingPlayback/ControllingPlayback.INT.udn
--- a/Engine/Documentation/Source/Engine/MediaFramework/HowTo/ControllingPlayback/ControllingPlayback.INT.udn
+++ b/Engine/Documentation/Source/Engine/MediaFramework/HowTo/ControllingPlayback/ControllingPlayback.INT.udn
@@ -5,21 +5,32 @@ Description:Example of how you can control playback of your media assets.
 parent:Engine/MediaFramework/HowTo
 order:3
 checkpoint: editorqs
-Version: 4.13
+Version: 4.15
 SkillLevel: Intermediate
 Related: Engine\MediaFramework
 
 In addition to playing a video with Media Framework, you can grant players the ability to control the Playback Rate to rewind, pause or fast forward videos. 
-In this how-to we allow the player to press a button to rewind, a different button to fast forward and another button to play a video playing in the level. 
+In this how-to we allow the player to press a button to rewind, a different button to fast forward and another button to play a video playing in the level similar to below. 
+
+[OBJECT:EmbeddedVideo]
+[PARAMLITERAL:width]
+640
+[/PARAMLITERAL]
+[PARAMLITERAL:height]
+360
+[/PARAMLITERAL]
+[PARAMLITERAL:videoid]
+s6rTrV5-Mdo
+[/PARAMLITERAL]
+[/OBJECT]
+
+## Steps
 
 [REGION:note]
 For this how-to we are using the **Blueprint Third Person Template** project with **Starter Content** enabled.
 You will also need a few [Supported Video Files](Engine\MediaFramework\TechReference) saved on our computer that you want to add to the playlist for playback. You can use your own or **Right-click** and download this [Sample Video](Infiltrator Demo.mp4). 
 [/REGION]
 
-
-## Steps
-
 1.	In your project create a folder called **Movies**, then in explorer add your video to the **Movies** folder. 
 
 	![](ControllingPlayback_01.png)
@@ -42,18 +53,24 @@ You will also need a few [Supported Video Files](Engine\MediaFramework\TechRefer
 
 1.	Use the **Translate (W)**, **Rotation (E)** and **Scale (R)** tools to adjust the size and location of the Static Mesh as desired. 
 
-1.	Hold **Ctrl** then select both the **Sound Wave** and **Media Texture** assets and drop them on top of the **Static Mesh** in the level. This will automatically create a **Material** and apply it to the Static Mesh. 
+1.	Hold **Ctrl** then select both the **Sound Wave** and **Media Texture** assets and drop them on top of the **Static Mesh** in the level. 
 
 	![](ControllingPlayback_06.png)
 
+	This will automatically create a **Material** and apply it to the Static Mesh. 
+
 1.	From the Toolbar, click the **Blueprints** button and **Open Level Blueprint**. 
 
 	![](ControllingPlayback_07.png)
 
-1.	Add a **Variable** of the **Media Player Reference** type called **MediaPlayer** and set it to your **VideoPlayer** Media Player asset. You may need to **Compile** before you can assign the **Default Value**. 
+1.	Add a **Variable** of the **Media Player Reference** type called **MediaPlayer** and set it to your **VideoPlayer** Media Player asset. 
 
 	![](ControllingPlayback_08.png)
 
+	[REGION:note]
+	You may need to **Compile** before you can assign the **Default Value**. 
+	[/REGION]
+
 1.	Add an **Event Begin Play** node, your **Media Player** variable and use **Open Source** to open the **MyVideo** File Media Source asset.
 
 	![](ControllingPlayback_09.png)
@@ -62,10 +79,12 @@ You will also need a few [Supported Video Files](Engine\MediaFramework\TechRefer
 
 	![](ControllingPlayback_10.png)
 
-1.	Add **1** and **3** Keyboard Events both using **Set Rate** off the **MediaPlayer** asset and adjust each node's **Rate** value. By setting Rate to **-1.0** this will play the video in reverse while a Rate of **2.0** doubles the forward playback speed, scrubbing forward. 
+1.	Add **1** and **3** Keyboard Events both using **Set Rate** off the **MediaPlayer** asset and adjust each node's **Rate** value. 
 
 	![](ControllingPlayback_11.png)
 
+	By setting Rate to **-1.0** this will play the video in reverse while a Rate of **2.0** doubles the forward playback speed, scrubbing forward. 
+
 	[REGION:warning]
 	Currently you may experience audio synch loss when issuing commands to the Media Player asset. You can insert short delays between commands to allow the player has time to process them. We are aware of these issues and are investigating possible fixes to accompany future engine updates. 
 	[/REGION]
@@ -88,11 +107,10 @@ s6rTrV5-Mdo
 [/OBJECT]
 
 When playing in the editor, the video playback will not initially start when the level is loaded. When the player presses the **2** key, the video will start playing. 
-Pressing the **1** key, the video will start to rewind while pressing the **3** key fast forwards the video. 
-
+Pressing the **1** key, the video will start to rewind while pressing the **3** key fast forwards the video.
 The rate at which we rewind or fast forward is determined by the **Rate** value on the **Set Rate** node. Adjust those values to speed up the play back rate when going in reverse or forward. 
 
-While our example is bound to keyboard presses, you could execute the same functionality within a User Interface built with UMG where the player clicks on a button to execute the same commands. 
+While our example uses keyboard presses, you could use the same functionality within a User Interface built with UMG where the player clicks on a button to issue commands. 
 
 ## Other Considerations
 
diff --git a/Engine/Documentation/Source/Engine/MediaFramework/HowTo/FileMediaSource/UsingFileMedia.INT.udn b/Engine/Documentation/Source/Engine/MediaFramework/HowTo/FileMediaSource/UsingFileMedia.INT.udn
--- a/Engine/Documentation/Source/Engine/MediaFramework/HowTo/FileMediaSource/UsingFileMedia.INT.udn
+++ b/Engine/Documentation/Source/Engine/MediaFramework/HowTo/FileMediaSource/UsingFileMedia.INT.udn
@@ -5,21 +5,32 @@ Description:Describes how to use the File Media asset to play videos inside a le
 parent:Engine/MediaFramework/HowTo
 order:3
 checkpoint: editorqs
-Version: 4.13
+Version: 4.15
 SkillLevel: Beginner
 Related: Engine\MediaFramework
 
 If you are looking to have video playing inside your level on an object whether its a TV or a monitor or some other object, you can do so using the **Media Framework** assets. 
-In this how-to we place a video file inside the **Content** folder of our project, then use the **Media Player** and related assets (with a little **Blueprint Visual Script**) to play the video during gameplay. 
+In this how-to we place a video file inside the **Content** folder of our project, then use the **Media Player** and related assets (with a little **Blueprint Visual Script**) to play the video during gameplay as shown below. 
+
+[OBJECT:EmbeddedVideo]
+[PARAMLITERAL:width]
+640
+[/PARAMLITERAL]
+[PARAMLITERAL:height]
+360
+[/PARAMLITERAL]
+[PARAMLITERAL:videoid]
+zxyY2s6ZfrU
+[/PARAMLITERAL]
+[/OBJECT]
+
+## Steps
 
 [REGION:note]
 For this how-to we are using the **Blueprint Third Person Template** project with **Starter Content** enabled.
 You will also need a [Supported Video File](Engine\MediaFramework\TechReference) saved on our computer that you want to play. You can use your own or **Right-click** and download this [Sample Video](SampleVideo.mp4). 
 [/REGION]
 
-
-## Steps
-
 1.	In the **Content Browser** expand the **Sources Panel**, then under **Content** create a folder called **Movies**. 
 
 	![](FileMedia_01.png)
@@ -28,10 +39,12 @@ You will also need a [Supported Video File](Engine\MediaFramework\TechReference)
 
 	![](FileMedia_01b.png)
 
-1.	Drag your video file into the **Content/Movies** folder of your project. Placing your video files inside the **Content/Movies** folder of your project will ensure the videos are properly packaged. 
+1.	Drag your video file into the **Content/Movies** folder of your project. 
 
 	![](FileMedia_02.png)
 
+	Placing your video files inside the **Content/Movies** folder of your project will ensure the videos are properly packaged. 
+
 1.	Inside Unreal Engine and your project, **Right-click** in the **Movies** folder and under **Media** select **File Media Source**. 
 
 	![](FileMedia_03.png)
@@ -44,18 +57,23 @@ You will also need a [Supported Video File](Engine\MediaFramework\TechReference)
 
 	![](FileMedia_05.png)
 
-1.	In the **Create Media Player** window, check **Audio output SoundWave asset** and **Video output Media Texture asset**. This will automatically create a Sound Wave and Media Texture asset linked to this Media Player asset which is needed for playback. 
+1.	In the **Create Media Player** window, check **Audio output SoundWave asset** and **Video output Media Texture asset**. 
 
 	![](FileMedia_06.png)
 
+	This will automatically create a Sound Wave and Media Texture asset linked to this Media Player asset which is needed for playback. 
+
 1.  Name the new Media Player asset. We have called our Media Player asset **SampleMedia** which is also applied to the Sound Wave and Media Texture. 
 
 	![](FileMedia_07.png)
 
-1.  Open the Media Player asset and **Double-click** on your Media Source asset and it will playback. In the lower-right in the **Details** panel the **Output** section will also have auto-assigned the Sound Wave and Media Texture. 
+1.  Open the Media Player asset and **Double-click** on your Media Source asset and it will playback. 
 
 	![](FileMedia_08.png)
 
+	In the lower-right in the **Details** panel, **Play on Open** is checked by default, which will indicate that the asset will automatically start playing the content when the source is opened. 
+	Since we created a Media Texture and Sound Wave asset when we imported our source, they will automatically be assigned in the **Output** section. 
+
 1.	In the **Content Browser** under **StarterContent/Shapes** drag a **Shape_Plane** into the level. We will use this Static Mesh to playback our video on. 
 
 	![](FileMedia_09.png)
@@ -64,18 +82,24 @@ You will also need a [Supported Video File](Engine\MediaFramework\TechReference)
 
 	![](FileMedia_10.png)
 
-1.	Hold **Ctrl** then select both the **Sound Wave** and **Media Texture** assets and drop them on top of the **Static Mesh** in the level. This will automatically create a **Material** and apply it to the Static Mesh. 
+1.	Hold **Ctrl** then select both the **Sound Wave** and **Media Texture** assets and drop them on top of the **Static Mesh** in the level. 
 
 	![](FileMedia_11.png)
 
+	This will automatically create a **Material** and apply it to the Static Mesh. 
+
 1.	From the Toolbar, click the **Blueprints** button and **Open Level Blueprint**. 
 
 	![](FileMedia_12.png)
 
-1.	Add a **Variable** of the **Media Player Reference** type called **Media Player** and set it to your **SampleMedia** Media Player asset. You may need to **Compile** before you can assign the **Default Value**. 
+1.	Add a **Variable** of the **Media Player Reference** type called **Media Player** and set it to your **SampleMedia** Media Player asset. 
 
 	![](FileMedia_13.png)
 
+	[REGION:note]
+	You may need to **Compile** before you can assign the **Default Value**. 
+	[/REGION]
+
 1.	Hold **Ctrl** and drag in the **MediaPlayer** variable and also **Right-click** and add an **Event Begin Play** node. 
 
 	![](FileMedia_14.png)
