Copying //UE4/Dev-Documentation to //UE4/Samples-Main (Source: //UE4/Dev-Documentation @ 3372893)

diff --git a/Engine/Documentation/Source/Engine/MediaFramework/HowTo/MediaPlaylists/UsingMediaPlaylists.INT.udn b/Engine/Documentation/Source/Engine/MediaFramework/HowTo/MediaPlaylists/UsingMediaPlaylists.INT.udn
--- a/Engine/Documentation/Source/Engine/MediaFramework/HowTo/MediaPlaylists/UsingMediaPlaylists.INT.udn
+++ b/Engine/Documentation/Source/Engine/MediaFramework/HowTo/MediaPlaylists/UsingMediaPlaylists.INT.udn
@@ -2,35 +2,52 @@ Availability:Public
 Title: Working with Media Playlists
 Crumbs:%ROOT%, Engine
 Description:Example of how you can use a Media Playlist to playback multiple Media Source assets. 
-Version: 4.13
+Version: 4.15
 SkillLevel: Intermediate
 parent:Engine/MediaFramework/HowTo
 order:3
 checkpoint: editorqs
 Related: Engine\MediaFramework
 
-Suppose you have a series of videos that you want to playback in a specific order or you want to allow players to access videos within a specific series of videos. 
-You can do so by using a **Media Playlist** asset which points to **Media Source** assets assigned in a predefined order. 
+When working with media assets in Unreal Engine 4, there may be times where you have a series of videos that you want to playback in a specific order or you want to allow players to access videos within a specific series of videos. 
+This can be accomplished by using a **Media Playlist** asset which points to **Media Source** assets assigned in a predefined order. 
 
-In this how-to we create a Media Playlist and allow the player to playback each asset sequentially or access a specific asset inside the playlist through a key press. 
+In this how-to we create a Media Playlist and allow the player to playback each asset or access a specific asset inside the playlist through a key press similar to the video below.
 
-[REGION:note]
-For this how-to we are using the **Blueprint Third Person Template** project with **Starter Content** enabled.
-You will also need a few [Supported Video Files](Engine\MediaFramework\TechReference) saved on our computer that you want to add to the playlist for playback. You can use your own or **Right-click** and download these two sample videos: [Video Example 1](Infiltrator Demo.mp4) and [Video Example 2](SampleVideo.mp4)
-[/REGION]
+[OBJECT:EmbeddedVideo]
+[PARAMLITERAL:width]
+640
+[/PARAMLITERAL]
+[PARAMLITERAL:height]
+360
+[/PARAMLITERAL]
+[PARAMLITERAL:videoid]
+kjbgHVnQxBI
+[/PARAMLITERAL]
+[/OBJECT]
 
 
 ## Steps
 
+[REGION:note]
+For this how-to we are using the **Blueprint Third Person Template** project with **Starter Content** enabled.
+You will also need a few [Supported Video Files](Engine\MediaFramework\TechReference) saved on our computer that you want to add to the playlist for playback. 
+You can use your own or **Right-click** and download these two sample videos: [Video Example 1](Infiltrator Demo.mp4) and [Video Example 2](SampleVideo.mp4)
+[/REGION]
+
 1.	In the **Content Browser** expand the **Sources Panel**, then in the **Content** create a new folder called **Movies**. 
 
 	![](Playlist_00.png)
 
-1.	**Right-click** on the **Movies** folder and choose **Show in Explorer** then place your video files in the folder. Placing your video files inside the **Content/Movies** folder of your project will ensure the videos are properly packaged. 
+1.	**Right-click** on the **Movies** folder and choose **Show in Explorer** then place your video files in the folder. 
 
 	![](Playlist_01.png)
 
-1.	Create two **File Media Source** assets (called **Video_01** and **Video_02**) point the **File Path** to each respective video. 
+	[REGION:note]
+	Placing your video files inside the **Content/Movies** folder of your project will ensure the videos are properly packaged. 
+	[/REGION]
+
+1.	Create two **File Media Source** assets (called **Video_01** and **Video_02**) and point the **File Path** to each respective video. 
 
 	![](Playlist_02.png)
 
@@ -38,18 +55,22 @@ You will also need a few [Supported Video Files](Engine\MediaFramework\TechRefer
 
 	![](Playlist_03.png)
 
-1.	Inside **MyPlaylist**, **Double-click** on each **File Media Source** to add it to the playlist. You can also use the **+** button to add an empty element which you can then assign your source media. In this example we are using **File Media Source** assets, however you can also include **Stream Media Source** assets in a Playlist. 
+1.	Inside **MyPlaylist**, **Double-click** on each **File Media Source** to add it to the playlist. 
 
 	![](Playlist_04.png)
 
+	You can also use the **+** button to add an empty element which you can then assign your source media. In this example we are using **File Media Source** assets, however you can also include **Stream Media Source** assets in a Playlist. 
+
 1.	Create a **Media Player** asset called **VideoPlayer**. 
 
 	![](Playlist_05.png)
 
-1.	In the **Create Media Player** window, check **Audio output SoundWave asset** and **Video output Media Texture asset**. This will automatically create a Sound Wave and Media Texture asset linked to this Media Player asset which is needed for playback. 
+1.	In the **Create Media Player** window, check **Audio output SoundWave asset** and **Video output Media Texture asset**. 
 
 	![](Playlist_06.png)
 
+	This will automatically create a Sound Wave and Media Texture asset linked to this Media Player asset which is needed for playback. 
+
 	![](Playlist_07.png)
 
 1.	In the **Content Browser** under **StarterContent/Shapes** drag a **Shape_Plane** into the level.  
@@ -58,18 +79,24 @@ You will also need a few [Supported Video Files](Engine\MediaFramework\TechRefer
 
 1.	Use the **Translate (W)**, **Rotation (E)** and **Scale (R)** tools to adjust the size and location of the Static Mesh as desired. 
 
-1.	Hold **Ctrl** then select both the **Sound Wave** and **Media Texture** assets and drop them on top of the **Static Mesh** in the level. This will automatically create a **Material** and apply it to the Static Mesh. 
+1.	Hold **Ctrl** then select both the **Sound Wave** and **Media Texture** assets and drop them on top of the **Static Mesh** in the level. 
 
 	![](Playlist_10.png)
 
+	This will automatically create a **Material** and apply it to the Static Mesh. 
+
 1.	From the Toolbar, click the **Blueprints** button and **Open Level Blueprint**. 
 
 	![](Playlist_11.png)
 
-1.	Add a **Variable** of the **Media Player Reference** type called **Media Player** and set it to your **VideoPlayer** Media Player asset. You may need to **Compile** before you can assign the **Default Value**. 
+1.	Add a **Variable** of the **Media Player Reference** type called **Media Player** and set it to your **VideoPlayer** Media Player asset.
 
 	![](Playlist_12.png)
 
+	[REGION:note]
+	You may need to **Compile** before you can assign the **Default Value**. 
+	[/REGION]
+
 1.	**Right-click** in the **Event Graph** and add a **1 Keyboard Event**, **2 Keyboard Event** and **3 Keyboard Event**.  
 
 	![](Playlist_13.png)
@@ -87,15 +114,6 @@ You will also need a few [Supported Video Files](Engine\MediaFramework\TechRefer
 	
 ## End Result
 
-When playing in the editor, the video playback will not initially start when the level is loaded. When the player presses the **1** key, the playlist will start playing from the first video in the playlist. 
-When pressing the **2** key, the video playback will jump to the second video in the playlist. When pressing **3**, we are opening and playing the video from the start of the playlist which is the same as opening the playlist at index 0. 
-
-Our **Media Player** asset is set to **Play on Open** which is why each video is played when we call the open playlist function. 
-If you want to open a Media Source but not start playback immediately, you can uncheck the Play on Open option inside the Media Player asset then use the **Play** function to play an open Media Source asset. 
-
-You can also navigate an opend playlist with the **Next** function (to pick the next item in the playlist) and **Previous** (to pick the previous item in the playlist). 
-The Media Player also has a property called **Shuffle** that when enabled will pick a random item from the playlist when the next source is to be opened (also affects Previous and Next). 
-
 [OBJECT:EmbeddedVideo]
 [PARAMLITERAL:width]
 640
@@ -107,3 +125,14 @@ The Media Player also has a property called **Shuffle** that when enabled will p
 kjbgHVnQxBI
 [/PARAMLITERAL]
 [/OBJECT]
+
+
+When playing in the editor, the video playback will not initially start when the level is loaded. When the player presses the **1** key, the playlist will start playing from the first video in the playlist. 
+When pressing the **2** key, the video playback will jump to the second video in the playlist. When pressing **3**, we are opening and playing the video from the start of the playlist which is the same as opening the playlist at index 0. 
+
+Our **Media Player** asset is set to **Play on Open** which is why each video is played when we call the open playlist function. 
+If you want to open a Media Source but not start playback immediately, you can uncheck the Play on Open option inside the Media Player asset then use the **Play** function to play an open Media Source asset. 
+
+You can also navigate an opend playlist with the **Next** function (to pick the next item in the playlist) and **Previous** (to pick the previous item in the playlist). 
+The Media Player also has a property called **Shuffle** that when enabled will pick a random item from the playlist when the next source is to be opened (also affects Previous and Next). 
+
diff --git a/Engine/Documentation/Source/Engine/MediaFramework/HowTo/StreamMediaSource/UsingStreamMedia.INT.udn b/Engine/Documentation/Source/Engine/MediaFramework/HowTo/StreamMediaSource/UsingStreamMedia.INT.udn
--- a/Engine/Documentation/Source/Engine/MediaFramework/HowTo/StreamMediaSource/UsingStreamMedia.INT.udn
+++ b/Engine/Documentation/Source/Engine/MediaFramework/HowTo/StreamMediaSource/UsingStreamMedia.INT.udn
@@ -5,46 +5,63 @@ Description:Describes how to use the Stream Media asset to stream videos from th
 parent:Engine/MediaFramework/HowTo
 order:3
 checkpoint: editorqs
-Version: 4.13
+Version: 4.15
 SkillLevel: Beginner
 Related: Engine\MediaFramework
 
-In this how-to we play a video on a Static Mesh inside the level when we start the game. Our source video will come from a provided URL stream and we use some **Blueprint Visual Script** to open the stream and start playback once we load into the level. 
+In this how-to we play a video on a Static Mesh inside the level when we start the game. 
+Our source video will come from a provided URL stream and we use some **Blueprint Visual Script** to open the stream and start playback once we load into the level. 
+The end result is a video playing inside our level where the source comes from our assigned URL similar to below. 
+
+[OBJECT:EmbeddedVideo]
+[PARAMLITERAL:width]
+640
+[/PARAMLITERAL]
+[PARAMLITERAL:height]
+360
+[/PARAMLITERAL]
+[PARAMLITERAL:videoid]
+ncHIQ0QteC8
+[/PARAMLITERAL]
+[/OBJECT] 
+
+## Steps
 
 [REGION:note]
 For this how-to we are using the **Blueprint Third Person Template** project with **Starter Content** enabled.
 [/REGION]
 
-
-## Steps
-
 1.	In the **Content Browser** expand the **Sources Panel** then create a new folder called **Media** inside the **Content** folder.
 
 	![](StreamMedia_01.png)
 	
-1.	**Right-click** inside the **Content/Media** folder then under **Media** create a **Stream Media Source** and call it **MediaStream**. 
+1.	Click the **Add New** button, then under **Media**, select **Stream Media Source** and name the asset **MediaStream**. 
 
 	![](StreamMedia_02.png)
 
-1.	Open **MediaStream** then enter your desired **Stream URL**. If you do not have a file to link to, **Right-click** this [Sample Video](Infiltrator Demo.mp4) and copy the link address and paste it in the Stream URL field.
+1.	Open **MediaStream** then enter your desired **Stream URL**. 
 
 	![](StreamMedia_03.png) 
+
+	If you do not have a URL file to link to, **Right-click** this [Sample Video](Infiltrator Demo.mp4) and copy the link address and paste it in the Stream URL field.
 	
 	[REGION:note]
-	We are using a fake address, however your Stream URL must link directly to a supported format in order to play back the video. 
+	Your Stream URL must link directly to a supported format in order to play back the video. 
 	For example, PS4Media (PS4) only supports MP4 through HLS right now, while WmfMedia (Windows) can support a number of other stream sources. 
 	For futher information on supported formats per platform/player plug-in, see the [](Engine\MediaFramework\TechReference\) page. 
 	[/REGION]
 
-1.	**Right-click** in the **Content Browser** then under **Media** select the **Media Player** asset. 
+1.	Click the **Add New** button, then under **Media**, select the **Media Player** asset. 
 
 	![](StreamMedia_04.png)
 
-1.	In the **Create Media Player** window, check **Audio output SoundWave asset** and **Video output Media Texture asset**. This will automatically create a Sound Wave and Media Texture asset linked to this Media Player asset which is needed for playback.
+1.	In the **Create Media Player** window, check **Audio output SoundWave asset** and **Video output Media Texture asset**. 
 
 	![](StreamMedia_05.png)
+
+	This will automatically create a Sound Wave and Media Texture asset linked to this Media Player asset which is needed for playback.
 	
-1. Name the new Media PLayer Asset. We have called our Media Player asset **SampleStream** which is also applied to the Sound Wave and Media Texture. 
+1. Name the new Media Player Asset. We have called our Media Player asset **SampleStream** which is also applied to the Sound Wave and Media Texture. 
 
 	![](StreamMedia_06.png) 
 
@@ -54,18 +71,24 @@ For this how-to we are using the **Blueprint Third Person Template** project wit
 
 1.	Use the **Translate (W)**, **Rotation (E)** and **Scale (R)** tools to adjust the mesh to your desired size and location.
 
-1.	Hold **Ctrl** then select both the **Sound Wave** and **Media Texture** assets and drop them on top of the **Static Mesh** in the level. This will automatically create a **Material** and apply it to the Static Mesh.
+1.	Hold **Ctrl** then select both the **Sound Wave** and **Media Texture** assets and drop them on top of the **Static Mesh** in the level. 
 
 	![](StreamMedia_08.png) 
 
+	This will automatically create a **Material** and apply it to the Static Mesh.
+
 1.	From the Toolbar, click the **Blueprints** button and **Open Level Blueprint**. 
 
 	![](StreamMedia_09.png)
 
-1.	Add a **Variable** of the **Media Player Reference** type called **Media Player** and set it to your **SampleStream** Media Player asset. You may need to **Compile** before you can assign the **Default Value**. 
+1.	Add a **Variable** of the **Media Player Reference** type called **Media Player** and set it to your **SampleStream** Media Player asset. 
 
 	![](StreamMedia_10.png)
 
+	[REGION:note]
+	You may need to **Compile** before you can assign the **Default Value**. 
+	[/REGION]
+
 1.	Hold **Ctrl** and drag in the **MediaPlayer** variable and also **Right-click** and add an **Event Begin Play** node. 
 
 	![](StreamMedia_11.png)
diff --git a/Engine/Documentation/Source/Engine/MediaFramework/Overview/MF_Overview.INT.udn b/Engine/Documentation/Source/Engine/MediaFramework/Overview/MF_Overview.INT.udn
--- a/Engine/Documentation/Source/Engine/MediaFramework/Overview/MF_Overview.INT.udn
+++ b/Engine/Documentation/Source/Engine/MediaFramework/Overview/MF_Overview.INT.udn
@@ -5,7 +5,7 @@ Description:Overview of Media Framework functionality in Unreal Engine 4.
 Navigation:topic
 parent:Engine/MediaFramework
 order:2
-Version: 4.13
+Version: 4.15
 Related: Engine\MediaFramework\TechReference
 Related: Engine\MediaFramework\HowTo
 
@@ -32,12 +32,13 @@ Please see the [](Engine\MediaFramework\TechReference\) page for information on
 
 The **File Media Source** asset is used for media files stored on your device or in a shared local network file. The **File Path** section is where you specify where the media file resides. 
 While a media file can exist anywhere, it's typically a good practice to move your media files to the **Content/Movies** folder of your project and then point to the media file. 
-This will ensure that your media is packaged correctly with your project. You will also see a yellow exclamation point when pointing to a file outside of this location as a warning. 
+This will ensure that your media is packaged correctly with your project. You will also see a yellow exclamation point (as shown above) when pointing to a file outside of this location as a warning. 
 
-This type of source asset can be used to load the entire media file into memory and play it from there by enabling the **Precache File** option. 
+This type of source asset can be used to load the entire media file into memory and play it from there by enabling the **Precache File** option (under the Advanced Options section). 
 Based on the size of your file, caching time may vary, so take that into consideration when choosing this option. 
+Each File Media Source asset can have its native media player plug-in overridden on a per platform basis (depicted below) or the player plug-in can be seleclted automatically. 
 
-Each source asset supports a name, in the **Default Player** field, to identify the appropriate native media player for that asset. If you leave this field empty, the media player will be detected automatically. 
+![](PlatformOverride.png)
 
 [REGION:note]
 Please see the [](Engine\MediaFramework\HowTo\FileMediaSource) how-to for working with File Media Source assets. 
@@ -48,7 +49,7 @@ Please see the [](Engine\MediaFramework\HowTo\FileMediaSource) how-to for workin
 ![](MediaSource_Stream.png)
 
 This **Stream Media Source** asset takes a URL as the source of your media, and pulls it from the internet. Currently, linking to a YouTube or DailyMotion style URL is not supported, and you will need to link directly to a hosted file. 
-Please note that the URL shown above is just an example, and cannot be used.
+Like the File Media Source asset, Stream Media Source assets can have their player plug-in overridden on a per platform basis or the player plug-in can be selected automatically. 
 
 [REGION:note]
 Please see the [](Engine\MediaFramework\HowTo\StreamMediaSource) how-to for working with Stream Media Source assets. 
@@ -58,15 +59,9 @@ Please see the [](Engine\MediaFramework\HowTo\StreamMediaSource) how-to for work
 
 ![](MediaSource_Platform.png)
 
-The **Platform Media Source** asset supports overriding Media Sources on a per-platform basis. You can specify the Media Source to be used if no source was specified for a target platform, as well as define the Media Sources per platform. 
-
+The **Platform Media Source** asset supports overriding a Media Source on a per-platform basis. 
 For example, supppose you want a particluar video to play only on Android Devices or only on PS4. Inside the **Media Sources** section, you can designate which videos play on which platforms. Leaving the Media Source blank will cause the **Default Source** to be used. 
 
-[REGION:warning]
-The Platform Media Source asset is experimental and is not expected to work in version 4.13 of the engine. It is planned for a 4.14 release.
-[/REGION]
-
-
 ## Media Playlists
 
 While you can play Media Source assets individually, the **Media Playlist** asset can play multiple Media Sources in order. 
@@ -96,7 +91,7 @@ From the **Add New** button under **Media**, select the **Media Sound Wave** ass
 
 ![](MediaPlayerSoundWindow.png)
 
-The Media Sound Wave asset provides **Debug**, **Concurrency** and **Attenuation** settings that can be adjusted to define how the asset is processed. 
+The Media Sound Wave asset provides **Curve**, **Debug**, **Concurrency** and **Attenuation** settings that can be adjusted to define how the asset is processed. 
 When linked to a Media Player asset, audio attached to a video source will automatically play back along with the video. 
 Typically, you can leave the default settings for a newly created Media Sound Wave asset. However, if you need more control over how the sound is perceived, you can adjust the Concurrency and Attenuation settings. 
 
@@ -147,10 +142,6 @@ You can play, pause, rewind, or skip to the beginning or end of a video from the
 In the lower-right **Details** panel, you can define options for **Playback** or assign the **Sound Wave** and **Video Textures** (or an **Image Texture**) associated with the Media Player. 
 There is also an **Info** panel which will display info about the source such as the stream type, codecs, sample rate, dimensions and more. 
 
-[REGION:warning]
-Still Image Texture output is currently not implemented by any of the Media Player plug-ins and this feature is considered experimental.
-[/REGION]
-
 ## Media Assets & Scripting
 
 Once your Media Player asset is set up and linked with the appropriate assets, it can be played during a game session. In order to play media in-game, you will first need to open the Media Source asset through Blueprint or C++. 
diff --git a/Engine/Documentation/Source/Engine/OpenWorldTools/Reference/OpenWorldToolsReference.INT.udn b/Engine/Documentation/Source/Engine/OpenWorldTools/Reference/OpenWorldToolsReference.INT.udn
--- a/Engine/Documentation/Source/Engine/OpenWorldTools/Reference/OpenWorldToolsReference.INT.udn
+++ b/Engine/Documentation/Source/Engine/OpenWorldTools/Reference/OpenWorldToolsReference.INT.udn
@@ -21,7 +21,7 @@ This page contains a reference listing of properties you will find while working
 | Property |  Description |
 | --- |  :--- |
 |**Mesh** | Defines which Static Mesh will be used. |
-|**Component Glass** |The component class to use for foliage instances. You can make a Blueprint subclass of FoliagedInstancedStaticMeshComponent to implement custom behavior and assign that class here.|
+|**Component Class** |The component class to use for foliage instances. You can make a Blueprint subclass of FoliagedInstancedStaticMeshComponent to implement custom behavior and assign that class here.|
 |**Override Materials** | Material overrides for foliage instances. |
 
 ### Placement
@@ -43,7 +43,7 @@ This page contains a reference listing of properties you will find while working
 
 | Property |  Description |
 | --- |  :--- |
-|**Collision Radius** | The CollisionTadius determines when two instances overlap. When two instances overlap a winner will be picked based on rules and priority.|
+|**Collision Radius** | The CollisionRadius determines when two instances overlap. When two instances overlap a winner will be picked based on rules and priority.|
 |**Shade Radius** | The ShadeRadius determines when two instances overlap. If an instance can grow in the shade this radius is ignored.|
 |**Num Steps** | The number of times we age the species and spread its seeds.|
 |**Initial Seed Density** | Specifies the number of seeds to populate along 10 meters. The number is implicitly squared to cover a 10m x 10m area.|
@@ -56,7 +56,6 @@ This page contains a reference listing of properties you will find while working
 |**Max Initial Age** | Allows a new seed to be older than 0 when created. New seed will be randomly assigned an age in the range[0,MaxInitialAge].|
 |**Max Age** | Specifies the oldest a seed can be. After reaching this age the instance will still spread seeds, but will not get any older.|
 |**Overlap Priority** | When two instances overlap we must determine which instance to remove. The instance with a lower OverlapPriority will be removed. In the case where OverlapPriority is the same regular simulation rules apply.|
-|**Overlap Priority** | When two instances overlap we must determine which instance to remove. The instance with a lower OverlapPriority will be removed. In the case where OverlapPriority is the same regular simulation rules apply.|
 |**Procedural Scale** | The scale range of this type when being procedurally generated. Configured with the Scale Curve.|
 |**Scale Curve** | Instance scale factor as a function of normalized age(i.e. Current Age / Max Age). X = 0 corresponds to Age = 0, X = 1 corresponds to Age = Max Age. Y = 0 corresponds to Min Scale, Y = 1 corresponds to Max Scale.|
 |**External Curve** | For and external curve.|
@@ -70,7 +69,7 @@ This page contains a reference listing of properties you will find while working
 |**Cull Distance** | The distance where instances will begin to fade out if using a PerInstanceFadeAmount Material node. 0 disables. When the entire cluster is beyond this distance, the cluster is completely culled and not rendered at all.|
 |**Cast Shadow** | Controls whether the foliage should cast a shadow or not.|
 |**Affect Dynamic Indirect Lighting** | Controls whether the foliage should inject light into the Light Propagation Volume. This flag is only used if CastShadow is true. |
-|**Affect Distance Field Lighting** | Controls whether the primitive should affect dynamic distance field lighting methods. This flag is only used if CasrShadow is true. |
+|**Affect Distance Field Lighting** | Controls whether the primitive should affect dynamic distance field lighting methods. This flag is only used if CastShadow is true. |
 |**Cast Dynamic Shadow** | Controls whether the foliage should cast shadows in the case of non precomputed shadowing. This flag is only used if CastShadow is true. |
 |**Cast Static Shadow** | Whether the foliage should cast a static shadow from shadow casting lights. This flag is only used if CastShadow is true. |
 |**Light Map Resolution** | Overrides the lightmap resolution defined in the static mesh. |
diff --git a/Engine/Documentation/Source/Engine/Performance/ForwardRenderer/ForwardRenderer.INT.udn b/Engine/Documentation/Source/Engine/Performance/ForwardRenderer/ForwardRenderer.INT.udn
--- a/Engine/Documentation/Source/Engine/Performance/ForwardRenderer/ForwardRenderer.INT.udn
+++ b/Engine/Documentation/Source/Engine/Performance/ForwardRenderer/ForwardRenderer.INT.udn
@@ -1,6 +1,6 @@
-Availability:Docs
+Availability:Public
 Title: Forward Shading Renderer for VR
-Crumbs: %ROOT%, Engine, Engine/Performance
+Crumbs:
 Description:Describes the advantages of using the Forward Shading Renderer for increased VR performance.
 type: overview
 version: 4.14
@@ -14,26 +14,32 @@ related: Engine/Performance
 related: Engine/Performance/GPU
 related: Engine/Rendering
 
+[REGION:warning]
+This feature is currently considered experimental and some aspects may not work as expected or may be altered in future updates. 
+[/REGION]
+
 [REGION:banner]
 ![](ForwardRenderingBanner.png)
 [/REGION]
 
-[REGION:warning]
-This feature is currently considered experimental and some aspects may not work as expected or may be altered in future updates. 
-[/REGION]
+[TOC(start:2 end:2)]
 
 [EXCERPT:intro]
 By default, Unreal Engine 4 uses a Deferred Renderer as it provides the most versatility and grants access to more rendering features. 
-However, there are some trade-offs in using the Deferred Render that might not be right for all VR experiences. 
-**Forward Rendering** provides a faster baseline with faster rendering passes which may lead to better performance on VR platforms. 
-Not only is Forward Rendering faster, it also provides better anti-aliasing options than the Deferred Render which may lead to more crisp visuals. 
+However, there are some trade-offs in using the Deferred Renderer that might not be right for all VR experiences. 
+**Forward Rendering** provides a faster baseline, with faster rendering passes, which may lead to better performance on VR platforms. 
+Not only is Forward Rendering faster, it also provides better anti-aliasing options than the Deferred Renderer, which may lead to better visuals. 
 [/EXCERPT]
 
-[TOC(start:2 end:2)]
-
 ## Enabling Forward Shading
 
-To enable the Forward Shading Renderer, open your **Project Settings** then under **Rendering**, enable **Forward Shading**. 
+To enable the Forward Shading Renderer:
+
+1. In the **Edit** menu, open the **Project Settings**.
+
+1. Select the **Rendering** tab on the left and locate the **Forward Shading** category.
+
+1. Enable **Forward Shading**.
 
 ![](EnableForwardShading.png)
 
@@ -41,7 +47,13 @@ You will be prompted to restart the editor where after the restart, you can begi
 
 ## Enabling Multisample Anti-Aliasing
 
-To use Multisample Anti-Aliasing (MSAA), open your **Project Settings** then under **Rendering** and **Default Settings**, change the **Anti-Aliasing Method** to **MSAA**.
+To use Multisample Anti-Aliasing (MSAA):
+
+1. In the **Edit** menu, open the **Project Settings**.
+
+1. Select the **Rendering** tab on the left and locate the **Default Settings** category.
+
+1. Set the **Anti-Aliasing Method** property to **MSAA**.
 
 ![](EnableMSAA.png)
 
@@ -96,22 +108,23 @@ This list comprises the currently supported features of the Forward Renderer:
 * Planar Reflections of a partial scene composited into Reflection Captures
 * DBuffer Decals
 * Precomputed Lighting and Sky Lights
-* Unshadowed Movable Lights
+* Any number of unshadowed Movable Lights
 * Capsule Shadows
 * Alpha-to-coverage for Masked Materials
 * Instanced Stereo Compatible
+* Up to 4 overlapping Shadow Casting Movable Lights
+* Light Functions on shadow casting lights
 
 ## Known Issues & Common Questions
 
 The following features are **not yet supported** while using the Forward Renderer: 
 
 * Screen Space Techniques (SSR, SSAO, Contact Shadows)
-* Shadow Casting Movable Lights
 * Dynamically Shadowed Translucency
 * Translucency receiving environment shadows from a Stationary Light
-* Light Functions and IES Profiles
-* Alpha-to-coverage
-* MSAA on D-Buffer Decals, Motion Blur, Dynamic Shadows and Capsule Shadows
+* MSAA on D-Buffer Decals and Motion Blur
+
+
 
 Below are some common questions/issues that may arise when using the experimental version of the Forward Renderer. 
 
@@ -120,8 +133,6 @@ Below are some common questions/issues that may arise when using the experimenta
 * **Does the Forward Renderer allow fewer texture samples than the Deferred Render?** 
 	* The Forward Renderer packs all features into one shader so there are fewer texture samplers available for your material. 
 	* You can usually solve this by using Shared Samplers. 
-* **Does Forward Rendering MSAA support alpha to coverage for alpha masked materials?** 
-	* This is on our radar, however, in this experimental release, MSAA does not support alpha to coverage in 4.14.
 * **How come I have issues with MSAA not being able to find samples against Atmospheric Fog?** 
 	* Atmospheric Fog does not handle MSAA correctly yet. 
 	* For Height Fog, you can get around with Vertex Fogging (which is enabled by default when enabling Forward Shading). 
diff --git a/Engine/Documentation/Source/Engine/Physics/Collision/HowTo/AddCollisionPreset/AddCollisionPreset.INT.udn b/Engine/Documentation/Source/Engine/Physics/Collision/HowTo/AddCollisionPreset/AddCollisionPreset.INT.udn
--- a/Engine/Documentation/Source/Engine/Physics/Collision/HowTo/AddCollisionPreset/AddCollisionPreset.INT.udn
+++ b/Engine/Documentation/Source/Engine/Physics/Collision/HowTo/AddCollisionPreset/AddCollisionPreset.INT.udn
@@ -1,4 +1,4 @@
-Availability: Public
+Availability: Docs
 Title:Static Mesh Collision Reference
 Crumbs: %ROOT%, Engine, Engine/Physics, Engine/Physics/Collision
 Description:Content guide to creating and setting up collision geometry.
diff --git a/Engine/Documentation/Source/Engine/Physics/Collision/HowTo/AddCustomTraceType/AddCustomTraceType.INT.udn b/Engine/Documentation/Source/Engine/Physics/Collision/HowTo/AddCustomTraceType/AddCustomTraceType.INT.udn
--- a/Engine/Documentation/Source/Engine/Physics/Collision/HowTo/AddCustomTraceType/AddCustomTraceType.INT.udn
+++ b/Engine/Documentation/Source/Engine/Physics/Collision/HowTo/AddCustomTraceType/AddCustomTraceType.INT.udn
@@ -2,30 +2,40 @@ Availability: Docs
 Title:Add a Custom Trace Type to your Project
 Crumbs: %ROOT%, Engine, Engine/Physics, Engine/Physics/Collision
 Description:Content guide to creating and setting up collision geometry.
+parent:Engine/Physics/Collision/HowTo
+topic-image:Engine/Physics/Collision/Reference/physics_topic.png
+tags:Physics
+tags:Tracing
 
-[VAR:Topic]
-[OBJECT:Topic]
-	[PARAM:image]
-		![%Engine/Physics/Collision/Reference:title%](Engine/Physics/Collision/Reference/physics_topic.png)
-	[/PARAM]
-	[PARAM:icon]
-		![](%ROOT%\rendering_icon.png)(convert:false)
-	[/PARAM]
-	[PARAM:title]
-		%Engine/Physics/Collision:/Reference:title%
-	[/PARAM]
-	[PARAM:description]
-		%Engine/Physics/Collision/Reference:description%
-	[/PARAM]
-	[PARAM:path]
-		[RELATIVE:Engine/Physics/Collision/Reference]
-	[/PARAM]
-[/OBJECT]
-[/VAR]
+Often you will find a need for more than just the two default **Trace Response** channels (Visibility and Camera), perhaps you have a special laser that needs to pass through a special opaque object that you can't see through or have the camera clip through. When you encounter a situation like this you can add your own custom **Trace Response** channels by following the steps below.
 
+## Steps
 
+1. Open your project settings: **Edit Menu** -> **Project Settings**.
 
+	![](Engine/Physics/Collision/COL_projectSettings_1.png)
 
+1. Under **Engine** select **Collision**:
+
+	![](Engine/Physics/Collision/COL_projectSettings_2.png)
+	
+1. Click **New Trace Channel...**:
+
+	![](Engine\Physics\Collision\COL_addTraceChannel.png)
+
+1. Name your new **Trace Channel** and set its **Default Response**:
+
+	![](newTracePrompt.png)
+
+	The **Default Response** can be **Block**, **Overlap**, or **Ignore**. Depending on your use case you'll want to choose carefully to prevent as much additional work adjusting collision profiles on your Actors.
+
+1. Click **Accept**.
+
+1. Any **Blueprints** open in the **Blueprint Editor** will have to be closed down and re-opened for the new Trace Channel to appear on any components or nodes.
+
+## Result
+
+You'll now have a new Trace Channel to use anywhere in the editor. Any Actors or Components you want to be able to Trace with your new Channel will have to be set to block the new Channel.
 
 
 
diff --git a/Engine/Documentation/Source/Engine/Physics/Constraints/ConstraintsReference/Constraints.INT.udn b/Engine/Documentation/Source/Engine/Physics/Constraints/ConstraintsReference/Constraints.INT.udn
--- a/Engine/Documentation/Source/Engine/Physics/Constraints/ConstraintsReference/Constraints.INT.udn
+++ b/Engine/Documentation/Source/Engine/Physics/Constraints/ConstraintsReference/Constraints.INT.udn
@@ -9,131 +9,72 @@ topic-image:Engine/Physics/physics_topic.png
 tags:Physics
 SkillLevel: Beginner
 
-[VAR:spacer]![](%ROOT%/spacer.png)(w:150 convert:false)[/VAR]
+[TOC(end:2)]
 
-[VAR:spacer2]![](%ROOT%/spacer.png)(w:40 convert:false)[/VAR]
+This page contains a reference listing of properties for Physics Constraints separated by major category.
 
-This page contains a reference listing of properties in the Physics Constraints. 
-
-## Properties
-
-Below are the properties for Physics Constraints separated by major category.
-
-[OBJECT:modulesectionexpanded]
-	[PARAMLITERAL:id]
-		cons
-	[/PARAMLITERAL]
-	[PARAM:heading]
-		Constraint
-	[/PARAM]
-	[PARAM:content]
-		%p1%
-	[/PARAM]
-[/OBJECT]
-
-[VAR:p1]
+## Constraint
 
+[EXCERPT:p1]
 ![](propConstraint.png)
 
 [REGION:raw]
-| Property | %spacer% | Description |
-| --- | :---: | :--- |
-| **Constraint Actor 1** || When a **Physics Constraint** is placed into the level editor you must specify **Actors** to constrain. This is the first of 2 **Actors**. |
-| **Component Name 1** || The first target component that will be constrained. This can be specified when using **Actors** to constrain only a specific component and not the **Actor** root. |
-| **Constraint Actor 2** || When a **Physics Constraint** is placed into the level editor you must specify **Actors** to constrain. This is the second of 2 **Actors**. |
-| **Component Name 2** || The second target component that will be constrained. This can be specified when using **Actors** to constrain only a specific component and not the **Actor** root.  |
-| **Joint Name** || When constraining in the **Physics Asset Tool**, this is the name of the Bone that is initially constrained. |
-| **Constraint Bone 1** || When constraining in the **Physics Asset Tool**, this is the name of the first joint to constrain. |
-| **Constraint Bone 2** || When constraining in the **Physics Asset Tool**, this is the name of the second joint to constrain. |
+| Property  | Description |
+| ---  | :--- |
+| **Constraint Actor 1** | When a **Physics Constraint** is placed into the level editor you must specify **Actors** to constrain. This is the first of 2 **Actors**. |
+| **Component Name 1** | The first target component that will be constrained. This can be specified when using **Actors** to constrain only a specific component and not the **Actor** root. |
+| **Constraint Actor 2** | When a **Physics Constraint** is placed into the level editor you must specify **Actors** to constrain. This is the second of 2 **Actors**. |
+| **Component Name 2** | The second target component that will be constrained. This can be specified when using **Actors** to constrain only a specific component and not the **Actor** root.  |
+| **Joint Name** | When constraining in the **Physics Asset Tool**, this is the name of the Bone that is initially constrained. |
+| **Constraint Bone 1** | When constraining in the **Physics Asset Tool**, this is the name of the first joint to constrain. |
+| **Constraint Bone 2** | When constraining in the **Physics Asset Tool**, this is the name of the second joint to constrain. |
 [/REGION]
+[/EXCERPT:p1]
 
-[/VAR]
-
-
-
-[OBJECT:modulesectionexpanded]
-	[PARAMLITERAL:id]
-		proj
-	[/PARAMLITERAL]
-	[PARAM:heading]
-		Constraint Behavior
-	[/PARAM]
-	[PARAM:content]
-		%p2%
-	[/PARAM]
-[/OBJECT]
-
-[VAR:p2]
+## Constraint Behavior
 
+[EXCERPT:p2]
 ![](propProjection.png)
 
 [REGION:raw]
-| Property | %spacer% | Description |
-| --- | :---: | :--- |
-| **Disable Collision** || This disables collision between the constrained components. |
-| **Enable Projection** || If the distance error between bodies exceeds **Projection Linear Tolerance** value, or rotational error exceeds **Projection Angular Tolerance**, the Physics Body will be projected to fix this. This fixes cases where the visuals look disconnected, such as individual chain links breaking away from each other. |
-| **Projection Linear Tolerance** || If the distance error exceeds this value, in world units, the body will be projected. |
-| **Projection Angular Tolerance** || If the rotational error exceeds this value, in degrees, the body will be projected. |
+| Property  | Description |
+| ---  | :--- |
+| **Disable Collision** | This disables collision between the constrained components. |
+| **Enable Projection** | If the distance error between bodies exceeds **Projection Linear Tolerance** value, or rotational error exceeds **Projection Angular Tolerance**, the Physics Body will be projected to fix this. This fixes cases where the visuals look disconnected, such as individual chain links breaking away from each other. |
+| **Projection Linear Tolerance** | If the distance error exceeds this value, in world units, the body will be projected. |
+| **Projection Angular Tolerance** | If the rotational error exceeds this value, in degrees, the body will be projected. |
 [/REGION]
-
-[/VAR]
-
-
+[/EXCERPT]
 
 <!--
-[OBJECT:modulesectionexpanded]
-	[PARAMLITERAL:id]
-		join
-	[/PARAMLITERAL]
-	[PARAM:heading]
-		Joint Presets
-	[/PARAM]
-	[PARAM:content]
-		%p3%
-	[/PARAM]
-[/OBJECT]
-
-[VAR:p3]
+## Joint Presets
 
+[EXCERPT:p3]
 ![](propJointPresets.png)
 
 [REGION:raw]
-| Property | %spacer% | Description |
-| --- | :---: | :--- |
-| **Hinge** || Setup joint as hinge. A hinge joint allows motion only in one plane. |
-| **Prismatic** || Setup joint as prismatic. A prismatic joint allows only linear sliding movement. |
-| **Ball and Socket** || Setup joint as ball and socket. A Ball and Socket joint allows motion around an indefinite number of axes, which have one common center. |
+| Property  | Description |
+| ---  | :--- |
+| **Hinge** | Setup joint as hinge. A hinge joint allows motion only in one plane. |
+| **Prismatic** | Setup joint as prismatic. A prismatic joint allows only linear sliding movement. |
+| **Ball and Socket** | Setup joint as ball and socket. A Ball and Socket joint allows motion around an indefinite number of axes, which have one common center. |
 [/REGION]
-
-[/VAR]
+[/EXCERPT:p3]
 -->
 
+## Linear Limits
 
-
-[OBJECT:modulesectionexpanded]
-	[PARAMLITERAL:id]
-		linlim
-	[/PARAMLITERAL]
-	[PARAM:heading]
-		Linear Limits
-	[/PARAM]
-	[PARAM:content]
-		%p4%
-	[/PARAM]
-[/OBJECT]
-
-[VAR:p4]
-
+[EXCERPT:p4]
 ![](propLinearLimits.png)
 
 [REGION:raw]
-| Property | %spacer% | Description |
-| --- | :---: | :--- |
-| **XMotion** || Indicates the linear constraint applied along the X-axis. [INCLUDE:#linearMotionEnum] |
-| **YMotion** || Indicates the linear constraint applied along the Y-axis. [INCLUDE:#linearMotionEnum] |
-| **ZMotion** || Indicates the linear constraint applied along the Z-axis. [INCLUDE:#linearMotionEnum] |
-| **Limit** || The distance allowed between the two joint reference frames. |
-| **Scale Linear Limits** || If true, linear limits scale using the absolute min of the 3D scale of the owning component. |
+| Property  | Description |
+| ---  | :--- |
+| **XMotion** | Indicates the linear constraint applied along the X-axis. [INCLUDE:#linearMotionEnum] |
+| **YMotion** | Indicates the linear constraint applied along the Y-axis. [INCLUDE:#linearMotionEnum] |
+| **ZMotion** | Indicates the linear constraint applied along the Z-axis. [INCLUDE:#linearMotionEnum] |
+| **Limit** | The distance allowed between the two joint reference frames. |
+| **Scale Linear Limits** | If true, linear limits scale using the absolute min of the 3D scale of the owning component. |
 [/REGION]
 
 ### Advanced
@@ -141,50 +82,34 @@ Below are the properties for Physics Constraints separated by major category.
 ![](propLinearLimitsADV.png)
 
 [REGION:raw]
-| Property | %spacer% | Description |
-| --- | :---: | :--- |
-| **Soft Constraint** || Whether we want to use a soft constraint (spring). |
-| **Stiffness** || Stiffness of the soft constraint. Only used when Soft Constraint is on. |
-| **Damping** || Damping of the soft constraint. |
-| **Restitution** || Controls the amount of bounce when the constraint is violated. |
-| **Contact Distance** || Determines how close to the limit we have to get before turning the joint on. |
-| **Linear Breakable** || Whether it is possible to break the joint with linear force. |
-| **Linear Break Threshold** || Force needed to break the distance constraint.  |
+| Property  | Description |
+| ---  | :--- |
+| **Soft Constraint** | Whether we want to use a soft constraint (spring). |
+| **Stiffness** | Stiffness of the soft constraint. Only used when Soft Constraint is on. |
+| **Damping** | Damping of the soft constraint. |
+| **Restitution** | Controls the amount of bounce when the constraint is violated. |
+| **Contact Distance** | Determines how close to the limit we have to get before turning the joint on. |
+| **Linear Breakable** | Whether it is possible to break the joint with linear force. |
+| **Linear Break Threshold** | Force needed to break the distance constraint.  |
 [/REGION]
+[/EXCERPT:p4]
 
-[/VAR]
+## Angular Limits
 
-
-
-
-
-
-[OBJECT:modulesectionexpanded]
-	[PARAMLITERAL:id]
-		anglim
-	[/PARAMLITERAL]
-	[PARAM:heading]
-		Angular Limits
-	[/PARAM]
-	[PARAM:content]
-		%p5%
-	[/PARAM]
-[/OBJECT]
-
-[VAR:p5]
+[EXCERPT:p5]
 
 ![](propsAngularLimits.png)
 
 [REGION:raw]
-| Property | %spacer% | Description |
-| --- | :---: | :--- |
-| **Swing 1Motion** || Indicates whether the Swing1 limit is used. [INCLUDE:#angularMotionEnum]  |
-| **Swing 2Motion** || Indicates whether the Swing2 limit is used. [INCLUDE:#angularMotionEnum] |
-| **Twist Motion** || Indicates whether the Twist limit is used. [INCLUDE:#angularMotionEnum] |
-| **Swing 1Limit Angle** || Angle of movement along the **XY** Plane. |
-| **Twist Limit Angle** || Angle of movement along the **XZ** Plane. |
-| **Swing 2Limit Angle** || Symmetric angle of roll along the X-axis. |
-| **Angular Rotation Offset** || Specifies the angular offset between the two frames of reference. By default this limit goes from (-Angle, +Angle). |
+| Property  | Description |
+| ---  | :--- |
+| **Swing 1Motion** | Indicates whether the Swing1 limit is used. [INCLUDE:#angularMotionEnum]  |
+| **Swing 2Motion** | Indicates whether the Swing2 limit is used. [INCLUDE:#angularMotionEnum] |
+| **Twist Motion** | Indicates whether the Twist limit is used. [INCLUDE:#angularMotionEnum] |
+| **Swing 1Limit Angle** | Angle of movement along the **XY** Plane. |
+| **Twist Limit Angle** | Angle of movement along the **XZ** Plane. |
+| **Swing 2Limit Angle** | Symmetric angle of roll along the X-axis. |
+| **Angular Rotation Offset** | Specifies the angular offset between the two frames of reference. By default this limit goes from (-Angle, +Angle). |
 [/REGION]
 
 ### Advanced
@@ -192,53 +117,40 @@ Below are the properties for Physics Constraints separated by major category.
 ![](propsAngularLimitsADV.png)
 
 [REGION:raw]
-| Property | %spacer% | Description |
-| --- | :---: | :--- |
-| [REGION:tablesection]**Swing Limits**[/REGION] |||
-| **Soft Constraint** || Whether we want to use a soft constraint (spring). |
-| **Stiffness** || Stiffness of the soft constraint. Only used when Soft Constraint is on. |
-| **Damping** || Damping of the soft constraint. |
-| **Restitution** || Controls the amount of bounce when the constraint is violated. |
-| **Contact Distance** || Determines how close to the limit we have to get before turning the joint on. |
-| [REGION:tablesection]**Twist Limits**[/REGION] |||
-| **Soft Constraint** || Whether we want to use a soft constraint (spring). |
-| **Stiffness** || Stiffness of the soft constraint. Only used when Soft Constraint is on. |
-| **Damping** || Damping of the soft constraint. |
-| **Restitution** || Controls the amount of bounce when the constraint is violated. |
-| **Contact Distance** || Determines how close to the limit we have to get before turning the joint on. |
-| [REGION:tablesection]**Breakable**[/REGION] |||
-| **Angular Breakable** || Whether it is possible to break the joint with angular force. |
-| **Angular Break Threshold** || Torque needed to break the joint. |
+| Property  | Description |
+| ---  | :--- |
+| [REGION:tablesection]**Swing Limits**[/REGION] ||
+| **Soft Constraint** | Whether we want to use a soft constraint (spring). |
+| **Stiffness** | Stiffness of the soft constraint. Only used when Soft Constraint is on. |
+| **Damping** | Damping of the soft constraint. |
+| **Restitution** | Controls the amount of bounce when the constraint is violated. |
+| **Contact Distance** | Determines how close to the limit we have to get before turning the joint on. |
+| [REGION:tablesection]**Twist Limits**[/REGION] ||
+| **Soft Constraint** | Whether we want to use a soft constraint (spring). |
+| **Stiffness** | Stiffness of the soft constraint. Only used when Soft Constraint is on. |
+| **Damping** | Damping of the soft constraint. |
+| **Restitution** | Controls the amount of bounce when the constraint is violated. |
+| **Contact Distance** | Determines how close to the limit we have to get before turning the joint on. |
+| [REGION:tablesection]**Breakable**[/REGION] ||
+| **Angular Breakable** | Whether it is possible to break the joint with angular force. |
+| **Angular Break Threshold** | Torque needed to break the joint. |
 [/REGION]
 
-[/VAR]
-
-
+[/EXCERPT:p5]
 
+## Linear Motor
 
-[OBJECT:modulesectionexpanded]
-	[PARAMLITERAL:id]
-		linmot
-	[/PARAMLITERAL]
-	[PARAM:heading]
-		Linear Motor
-	[/PARAM]
-	[PARAM:content]
-		%p6%
-	[/PARAM]
-[/OBJECT]
-
-[VAR:p6]
+[EXCERPT:p6]
 
 ![](propLinearMotor.png)
 
 [REGION:raw]
-| Property | %spacer% | Description |
-| --- | :---: | :--- |
-| **Position Target** || Enables the position linear motor on one or more axes and sets the local position desired. |
-| **Strength** || How much force to apply to reach the desired position. |
-| **Velocity Target** || Enables the velocity linear motor on one or more axes and sets the desired velocity. |
-| **Strength** || How much force to apply to reach the desired velocity. |
+| Property  | Description |
+| ---  | :--- |
+| **Position Target** | Enables the position linear motor on one or more axes and sets the local position desired. |
+| **Strength** | How much force to apply to reach the desired position. |
+| **Velocity Target** | Enables the velocity linear motor on one or more axes and sets the desired velocity. |
+| **Strength** | How much force to apply to reach the desired velocity. |
 [/REGION]
 
 
@@ -247,42 +159,27 @@ Below are the properties for Physics Constraints separated by major category.
 ![](PhysCon_MotorADV.png)
 
 [REGION:raw]
-| Property | %spacer% | Description |
-| --- | :---: | :--- |
-| **Max Force** || The force limit of the drive.|
+| Property  | Description |
+| ---  | :--- |
+| **Max Force** | The force limit of the drive.|
 [/REGION]
 
-[/VAR]
-
-
-
+[/EXCERPT:p6]
 
-[OBJECT:modulesectionexpanded]
-	[PARAMLITERAL:id]
-		angmot
-	[/PARAMLITERAL]
-	[PARAM:heading]
-		Angular Motor
-	[/PARAM]
-	[PARAM:content]
-		%p7%
-	[/PARAM]
-[/OBJECT]
-
-[VAR:p7]
+## Angular Motor
 
 ![](propsAngularMotor.png)
 
 [REGION:raw]
-| Property | %spacer% | Description |
-| --- | :---: | :--- |
-| **Angular Drive Mode** || [INCLUDE:#AngMotorMode] |
-| **Target Orientation** || Target proemtatopm relative to the body reference frame. |
-| **Drives** || Depending on the Angular Drive Mode setting, you can enable or disable the different motors for that mode. |
-| **Strength** || The force applied to reach the target orientation. |
-| **Target Velocity** || Target angular velocity relative to the body reference frame |
-| **Drives** || Depending on the Angular Drive Mode setting, you can enable or disable the different motors for that mode. |
-| **Stength** || The force applied to reach the target velocity. |
+| Property  | Description |
+| ---  | :--- |
+| **Angular Drive Mode** | [INCLUDE:#AngMotorMode] |
+| **Target Orientation** | Target proemtatopm relative to the body reference frame. |
+| **Drives** | Depending on the Angular Drive Mode setting, you can enable or disable the different motors for that mode. |
+| **Strength** | The force applied to reach the target orientation. |
+| **Target Velocity** | Target angular velocity relative to the body reference frame |
+| **Drives** | Depending on the Angular Drive Mode setting, you can enable or disable the different motors for that mode. |
+| **Stength** | The force applied to reach the target velocity. |
 [/REGION]
 
 
@@ -291,14 +188,11 @@ Below are the properties for Physics Constraints separated by major category.
 ![](PhysCon_MotorADV.png)
 
 [REGION:raw]
-| Property | %spacer% | Description |
-| --- | :---: | :--- |
-| **Max Force** || The force limit of the drive.|
+| Property  | Description |
+| ---  | :--- |
+| **Max Force** | The force limit of the drive.|
 [/REGION]
 
-[/VAR]
-
-
 
 <!--
 
@@ -306,10 +200,10 @@ Below are the properties for Physics Constraints separated by major category.
 
 Whether this angular motor uses SLERP (Sphereical Linear Interpolation) or decompose into a Swing Motor and a Twist Motor (cone and roll constraints). SLERP will not work if any of the angular constraints are locked.
 
-| Property | %spacer2% | Description |
-| --- | :---: | :--- |
-| **SLERP** || Sets the motor to SLERP (Sphereical Linear Interpolation) mode. SLERP mode will not work if any of axis of the angular constraint is locked. |
-| **Twist and Swing** || Sets the motor to **Twist and Swing** (cone and roll constraints) mode.|
+| Property  | Description |
+| ---  | :--- |
+| **SLERP** | Sets the motor to SLERP (Sphereical Linear Interpolation) mode. SLERP mode will not work if any of axis of the angular constraint is locked. |
+| **Twist and Swing** | Sets the motor to **Twist and Swing** (cone and roll constraints) mode.|
 
 [/EXCERPT]
 
@@ -332,7 +226,7 @@ Whether this angular motor uses SLERP (Sphereical Linear Interpolation) or decom
 
 | Property | Description |
 | -------- | ----------- |
-|[REGION:tablesection]Constraint[/REGION]||
+|[REGION:tablesection]Constraint[/REGION]|
 | **Component Name 1** |  |
 | **Component Name 2** |   |
 | **Joint Name** |  |
@@ -346,7 +240,7 @@ Whether this angular motor uses SLERP (Sphereical Linear Interpolation) or decom
 
 | Property | Description |
 | -------- | ----------- |
-|[REGION:tablesection]Constraint[/REGION]||
+|[REGION:tablesection]Constraint[/REGION]|
 | **Joint Name** |  |
 | **Constraint Bone 1** |  |
 | **Constraint Bone 2** |  |
@@ -362,7 +256,7 @@ Whether this angular motor uses SLERP (Sphereical Linear Interpolation) or decom
 
 | Property | Description |
 | -------- | ----------- |
-|[REGION:tablesection]Linear Limits[/REGION]||
+|[REGION:tablesection]Linear Limits[/REGION]|
 | **Linear XMotion** |  |
 | **Linear YMotion** |  |
 | **Linear ZMotion** |  |
@@ -373,11 +267,11 @@ Whether this angular motor uses SLERP (Sphereical Linear Interpolation) or decom
 | **Linear Limit Damping** |  |
 | **Linear Breakable** |  |
 | **Linear Break Threshold** |  |
-|[REGION:tablesection]Projection[/REGION]||
+|[REGION:tablesection]Projection[/REGION]|
 | **Enable Projection** |  |
 | **Projection Linear Tolerance** |  |
 | **Projection Angular Tolerance** |  |
-|[REGION:tablesection]Angular Limits[/REGION]||
+|[REGION:tablesection]Angular Limits[/REGION]|
 | **Angular Swing 1Motion** |  |
 | **Angular Swing 2Motion** |  |
 | **Angular Twist Motion** |  |
@@ -393,7 +287,7 @@ Whether this angular motor uses SLERP (Sphereical Linear Interpolation) or decom
 | **Twist Limit Damping** |  |
 | **Angular Breakable** |  |
 | **Angular Break Threshold** |  |
-|[REGION:tablesection]Linear Motor[/REGION]||
+|[REGION:tablesection]Linear Motor[/REGION]|
 | **Linear Position Drive** |  |
 | **Linear XPosition Drive** |  |
 | **Linear YPosition Drive** |  |
@@ -404,7 +298,7 @@ Whether this angular motor uses SLERP (Sphereical Linear Interpolation) or decom
 | **Linear Velocity Target** |  |
 | **Linear Velocity Strength** | Damping to apply to the linear drive. |
 | **Max Linear Force** | Limit to the force the linear drive can apply. |
-|[REGION:tablesection]Angular Motor[/REGION]||
+|[REGION:tablesection]Angular Motor[/REGION]|
 | **Angular Orientation Drive** |  |
 | **Angular Orientation Target** |  |
 | **Angular Position Strength** | Spring value to apply for the angular drive. |
diff --git a/Engine/Documentation/Source/Engine/Physics/PhAT/HowTo/ApplyConstraintProfile/ApplyConstraintProfile.INT.udn b/Engine/Documentation/Source/Engine/Physics/PhAT/HowTo/ApplyConstraintProfile/ApplyConstraintProfile.INT.udn
--- a/Engine/Documentation/Source/Engine/Physics/PhAT/HowTo/ApplyConstraintProfile/ApplyConstraintProfile.INT.udn
+++ b/Engine/Documentation/Source/Engine/Physics/PhAT/HowTo/ApplyConstraintProfile/ApplyConstraintProfile.INT.udn
@@ -1,4 +1,4 @@
-﻿Availability:Docs
+﻿Availability:Public
 Title:How-To Apply a Physics Constraint Profile
 Crumbs: %ROOT%, Engine
 Description:This how-to will show you how to apply a Physics Constraint Profile to a Skeletal Mesh Component using Blueprints.
diff --git a/Engine/Documentation/Source/Engine/Physics/PhAT/HowTo/ApplyPhysicalAnimationProfile/ApplyPhysicalAnimProfile.INT.udn b/Engine/Documentation/Source/Engine/Physics/PhAT/HowTo/ApplyPhysicalAnimationProfile/ApplyPhysicalAnimProfile.INT.udn
--- a/Engine/Documentation/Source/Engine/Physics/PhAT/HowTo/ApplyPhysicalAnimationProfile/ApplyPhysicalAnimProfile.INT.udn
+++ b/Engine/Documentation/Source/Engine/Physics/PhAT/HowTo/ApplyPhysicalAnimationProfile/ApplyPhysicalAnimProfile.INT.udn
@@ -1,4 +1,4 @@
-﻿Availability:Docs
+﻿Availability:Public
 Title:How-To Apply a Physical Animation Profile
 Crumbs: %ROOT%, Engine
 Description:This how-to will show you how to apply a Physical Animation Profile to a Skeletal Mesh Component using Blueprints.
diff --git a/Engine/Documentation/Source/Engine/Physics/PhAT/Reference/Properties.INT.udn b/Engine/Documentation/Source/Engine/Physics/PhAT/Reference/Properties.INT.udn
--- a/Engine/Documentation/Source/Engine/Physics/PhAT/Reference/Properties.INT.udn
+++ b/Engine/Documentation/Source/Engine/Physics/PhAT/Reference/Properties.INT.udn
@@ -9,460 +9,216 @@ topic-image:Engine/Physics/physics_topic.png
 tags:Physics
 SkillLevel: Beginner
 
-[VAR:spacer]![](%ROOT%/spacer.png)(w:150 convert:false)[/VAR]
+[TOC(end:3)]
 
-[VAR:spacer2]![](%ROOT%/spacer.png)(w:40 convert:false)[/VAR]
+This page contains a reference listing of properties for Physics Asset Tool, separated into 3 expandable sections based on what is selected in the viewport.
 
-This page contains a reference listing of properties you will find while working in the Physics Asset Tool.
-
-## Properties
-
-Below are the properties for Physics Asset Tool, separated into 3 expandable sections based on what is selected in the viewport.
-
-[OBJECT:modulesection]
-	[PARAMLITERAL:id]
-		PhatP
-	[/PARAMLITERAL]
-	[PARAM:heading]
-		Physics Asset Tool Properties
-	[/PARAM]
-	[PARAM:content]
-		%c1%
-	[/PARAM]
-[/OBJECT]
-
-[VAR:c1]
+## Physics Asset Tool Properties
 
 When no Physics Body or Physics Constraint is selected (regardless of mode), the properties for the Physics Asset Tool are show in the Details panel.
 
-[OBJECT:modulesectionexpanded]
-	[PARAMLITERAL:id]
-		Anim1
-	[/PARAMLITERAL]
-	[PARAM:heading]
-		Anim
-	[/PARAM]
-	[PARAM:content]
-		%p11%
-	[/PARAM]
-[/OBJECT]
-
-[OBJECT:modulesectionexpanded]
-	[PARAMLITERAL:id]
-		Sim1
-	[/PARAMLITERAL]
-	[PARAM:heading]
-		Simulation
-	[/PARAM]
-	[PARAM:content]
-		%p12%
-	[/PARAM]
-[/OBJECT]
-
-[OBJECT:modulesectionexpanded]
-	[PARAMLITERAL:id]
-		MousS
-	[/PARAMLITERAL]
-	[PARAM:heading]
-		Mouse Spring
-	[/PARAM]
-	[PARAM:content]
-		%p13%
-	[/PARAM]
-[/OBJECT]
-
-[OBJECT:modulesectionexpanded]
-	[PARAMLITERAL:id]
-		Poke
-	[/PARAMLITERAL]
-	[PARAM:heading]
-		Poking
-	[/PARAM]
-	[PARAM:content]
-		%p14%
-	[/PARAM]
-[/OBJECT]
-
-[OBJECT:modulesectionexpanded]
-	[PARAMLITERAL:id]
-		Adv
-	[/PARAMLITERAL]
-	[PARAM:heading]
-		Advanced
-	[/PARAM]
-	[PARAM:content]
-		%p15%
-	[/PARAM]
-[/OBJECT]
-
-[OBJECT:modulesectionexpanded]
-	[PARAMLITERAL:id]
-		Draw
-	[/PARAMLITERAL]
-	[PARAM:heading]
-		Drawing
-	[/PARAM]
-	[PARAM:content]
-		%p16%
-	[/PARAM]
-[/OBJECT]
-
-[/VAR]
-
-[VAR:p11]
+### Anim
 
 ![](PhATDef1.png)
 
 [REGION:raw]
-| Property | %spacer% | Description |
-| --- | :---: | :--- |
-| **Physics Blend** || Enables you to manually control the physics/animation blend while simulating in the **Physics Asset Tool**. |
-| **Update Joints from Animation** || Let's you manually control the physics/animation. |
-| **Physics Update Mode** || [INCLUDE:#UpdateJointsfromAnimationEnum] |
-| **Poke Pause Time** || Time between poking a ragdoll and starting to blend back. |
-| **Poke Blend Time** || Time taken to blend from physics to animation. |
+| Property | Description |
+| --- | :--- |
+| **Physics Blend** | Enables you to manually control the physics/animation blend while simulating in the **Physics Asset Tool**. |
+| **Update Joints from Animation** | Let's you manually control the physics/animation. |
+| **Physics Update Mode** | [INCLUDE:#UpdateJointsfromAnimationEnum] |
+| **Poke Pause Time** | Time between poking a ragdoll and starting to blend back. |
+| **Poke Blend Time** | Time taken to blend from physics to animation. |
 [/REGION]
 
-[/VAR]
-
-<!--
-[EXCERPT:UpdateJointsfromAnimationEnum]
-
-Determines whether simulation of root body updates component transform.
-
-| Value | %spacer2% | Description |
-| --- | :---: | :--- |
-| **Simulation Updates Component Transfrom** || Simulation updates the component's transform. |
-| **Component Trasform is Kinematic** || Simulation doesn't update the component's transform. Instead it is driven by animation, blueprint, etc... |
-[/EXCERPT]
--->
-
-[VAR:p12]
+### Simulation
 
 ![](PhATDef2.png)
 
 [REGION:raw]
-| Property | %spacer% | Description |
-| --- | :---: | :--- |
-|**Floor Gap** || The gap between the floor and physics objects in the **Physics Asset Tool**. Changing this requires a restart of the **Physics Asset Tool**. |
-|**Grav Scale** || Scale factor for the gravity used in the simulation while in the **Physics Asset Tool**. |
-|**Max FPS** || Max FPS for simulation in PhAT. This is helpful for targeting the same FPS as your game. -1 means this is disabled.  |
-|**Time Dilation** || This will dilate the time by scale.  |
+| Property | Description |
+| --- | :--- |
+|**Floor Gap** | The gap between the floor and physics objects in the **Physics Asset Tool**. Changing this requires a restart of the **Physics Asset Tool**. |
+|**Grav Scale** | Scale factor for the gravity used in the simulation while in the **Physics Asset Tool**. |
+|**Max FPS** | Max FPS for simulation in PhAT. This is helpful for targeting the same FPS as your game. -1 means this is disabled.  |
+|**Time Dilation** | This will dilate the time by scale.  |
 [/REGION]
 
-[/VAR]
-
-[VAR:p13]
+### Mouse Spring
 
 ![](PhATDef3.png)
 
 [REGION:raw]
-| Property | %spacer% | Description |
-| --- | :---: | :--- |
-|**Handle Linear Damping** || Linear damping of mouse spring forces. |
-|**Handle Linear Stiffness** || Linear stiffness of mouse spring forces. |
-|**Handle Angular Damping** || Angular damping of mouse spring forces. |
-|**Handle Angular Stiffness** || Angular stiffness of mouse spring forces. |
-|**Interpolation Speed** || How quickly we interpolate the physics target transform for mouse spring forces. |
+| Property | Description |
+| --- | :--- |
+|**Handle Linear Damping** | Linear damping of mouse spring forces. |
+|**Handle Linear Stiffness** | Linear stiffness of mouse spring forces. |
+|**Handle Angular Damping** | Angular damping of mouse spring forces. |
+|**Handle Angular Stiffness** | Angular stiffness of mouse spring forces. |
+|**Interpolation Speed** | How quickly we interpolate the physics target transform for mouse spring forces. |
 [/REGION]
 
-[/VAR]
-
-[VAR:p14]
+### Poking
 
 ![](PhATDef4.png)
 
 [REGION:raw]
-| Property | %spacer% | Description |
-| --- | :---: | :--- |
-| **Poke Strength** || Strength of the impulse used when poking with the left mouse button. |
+| Property | Description |
+| --- | :--- |
+| **Poke Strength** | Strength of the impulse used when poking with the left mouse button. |
 [/REGION]
 
-[/VAR]
-
-[VAR:p15]
+### Advanced
 
 ![](PhATDef5.png)
 
 [REGION:raw]
-| Property | %spacer% | Description |
-| --- | :---: | :--- |
-| **Show Constraints as Points** || Whether or not to draw constraints as points. |
-| **Show Names in Hierarchy** || Whether or not to draw bone names in the viewport. |
+| Property | Description |
+| --- | :--- |
+| **Show Constraints as Points** | Whether or not to draw constraints as points. |
+| **Show Names in Hierarchy** | Whether or not to draw bone names in the viewport. |
 [/REGION]
 
-[/VAR]
-
-[VAR:p16]
+### Drawing
 
 ![](PhATDef6.png)
 
 [REGION:raw]
-| Property | %spacer% | Description |
-| --- | :---: | :--- |
-| **Constraint Draw Size** || Controls how large constraints are drawn in the **Physics Asset Tool**. |
+| Property | Description |
+| --- | :--- |
+| **Constraint Draw Size** | Controls how large constraints are drawn in the **Physics Asset Tool**. |
 [/REGION]
 
-[/VAR]
+<!--
+[EXCERPT:UpdateJointsfromAnimationEnum]
 
----
+Determines whether simulation of root body updates component transform.
 
-[OBJECT:modulesection]
-	[PARAMLITERAL:id]
-		BodyM
-	[/PARAMLITERAL]
-	[PARAM:heading]
-		Body Mode
-	[/PARAM]
-	[PARAM:content]
-		%c2%
-	[/PARAM]
-[/OBJECT]
+| Value | Description |
+| --- | :--- |
+| **Simulation Updates Component Transfrom** | Simulation updates the component's transform. |
+| **Component Trasform is Kinematic** | Simulation doesn't update the component's transform. Instead it is driven by animation, blueprint, etc... |
+[/EXCERPT]
+-->
 
-[VAR:c2]
+## Body Mode Properties
 
 When the Physics Asset Tool is in Body Mode and a Physics Body is selected, the Details panel will display the properties for that Physics Body.
 
-[OBJECT:modulesectionexpanded]
-	[PARAMLITERAL:id]
-		Phys1
-	[/PARAMLITERAL]
-	[PARAM:heading]
-		Physics
-	[/PARAM]
-	[PARAM:content]
-		%p22%
-	[/PARAM]
-[/OBJECT]
-
-[OBJECT:modulesectionexpanded]
-	[PARAMLITERAL:id]
-		Body1
-	[/PARAMLITERAL]
-	[PARAM:heading]
-		Body Setup
-	[/PARAM]
-	[PARAM:content]
-		%p21%
-	[/PARAM]
-[/OBJECT]
-
-[OBJECT:modulesectionexpanded]
-	[PARAMLITERAL:id]
-		Coll
-	[/PARAMLITERAL]
-	[PARAM:heading]
-		Collision
-	[/PARAM]
-	[PARAM:content]
-		%p23%
-	[/PARAM]
-[/OBJECT]
-
-[OBJECT:modulesectionexpanded]
-	[PARAMLITERAL:id]
-		PhysAnim
-	[/PARAMLITERAL]
-	[PARAM:heading]
-		Physical Animation
-	[/PARAM]
-	[PARAM:content]
-		%p24%
-	[/PARAM]
-[/OBJECT]
-
-[/VAR]
-
-[VAR:p21]
+### Physics
 
-![](PhATBody1.png)
+![](PhATBody2.png)
 
 [REGION:raw]
-| Property | %spacer% | Description |
-| --- | :---: | :--- |
-| **Primitives** || The basic shapes that compose this Bone's collision. |
-| **Bone Name** || Used in the PhysicsAsset case, associates this Body with a Bone in a Skeletal Mesh. |
-| **Consider for Bounds** || Should this bone be used in the calculation for the bounding volume of the Skeletal Mesh. |
+| Property | Description |
+| --- | :--- |
+| **Mass in KG** | Mass of the body in KG. |
+| **Linear Damping** | "Drag" force added to reduce linear movement. |
+| **Angular Damping** | "Drag" force added to reduce angular movement.  |
+| **Enable Gravity** | If object should have the force of gravity applied.  |
+| **Physics Type** | If this is set to **Unfixed** it will use physics. If this is set to **Fixed** it will use kinematic. **Default** will inherit from **OwnerComponent**'s behavior. |
+| **Double Sided Geometry** | If true, the physics triangle mesh will use double sided faces when doing scene queries. Thsi is useful for planes and single sided meshes that need traces to work on both sides. |
+| **Simple Collision Physical Material** | The **Physical Material** to use for simple collision on this body. |
 [/REGION]
 
-[/VAR]
+#### Advanced
 
-[VAR:p22]
-
-![](PhATBody2.png)
+![](PhATBody2_adv.png)
 
 [REGION:raw]
-| Property | %spacer% | Description |
-| --- | :---: | :--- |
-| **Mass in KG** || Mass of the body in KG. |
-| **Linear Damping** || "Drag" force added to reduce linear movement. |
-| **Angular Damping** || "Drag" force added to reduce angular movement.  |
-| **Enable Gravity** || If object should have the force of gravity applied.  |
-| **Physics Type** || If this is set to **Unfixed** it will use physics. If this is set to **Fixed** it will use kinematic. **Default** will inherit from **OwnerComponent**'s behavior. |
-| **Double Sided Geometry** || If true, the physics triangle mesh will use double sided faces when doing scene queries. Thsi is useful for planes and single sided meshes that need traces to work on both sides. |
-| **Simple Collision Physical Material** || The **Physical Material** to use for simple collision on this body. |
+| Property | Description |
+| --- | :--- |
+| **Start Awake** | If the physics body should start awake (and moving), or asleep (stationary until interacted with). |
+| **Center of Mass Offset** | User specified offset for the center of mass of this object, from the calculated location. |
+| **Mass Scale** | Per-instance scaling of mass. |
+| **Max Angular Velocity** | The maximum angular velocity for this instance. |
+| **Generate Wake Events** | Should **wake**/**sleep** events fire when this object is woken up or put to sleep by the physics simulation. |
+| **Max Depenetration Velocity** | The maximum velocity used to depenetrate this object.  |
+| **Walkable Slope Override** | Custom walkable slope setting for this body. See the [Walkable Slope](Engine/Physics/WalkableSlope) documentation for usage information. |
+| **Walkable Slope Behavior** | Behavior of this surface (whether we affect the walkable slope). Determines how the Walkable Slope Angle operates on the existing walkable angle, whether it increases or decreases it. |
+| **Walkable Slope Angle** | Override walkable slope, applying the rules of the Walkable Slope Behavior. |
+| **Custom Sleep Threshold Multiplier** | If the **SleepFamily** is set to custom, multiply the natural sleep threshold by this amount. A higher number will cause the body to sleep sooner. |
+| **Position Solver Iteration Count** | This physics body's solver iteration count for position. Increasing this will be more CPU intensive, but better stabilized. |
+| **Velocity Solver Iteration Count** | This physics body's solver iteration count for velocity. Increasing this will be more CPU intensive, but better stabilized. |
+| **Sleep Family** | The set of values used in considering when put this body to sleep. |
 [/REGION]
 
-### Advanced
+### Body Setup
 
-![](PhATBody2_adv.png)
+![](PhATBody1.png)
 
 [REGION:raw]
-| Property | %spacer% | Description |
-| --- | :---: | :--- |
-| **Start Awake** || If the physics body should start awake (and moving), or asleep (stationary until interacted with). |
-| **Center of Mass Offset** || User specified offset for the center of mass of this object, from the calculated location. |
-| **Mass Scale** || Per-instance scaling of mass. |
-| **Max Angular Velocity** || The maximum angular velocity for this instance. |
-| **Generate Wake Events** || Should **wake**/**sleep** events fire when this object is woken up or put to sleep by the physics simulation. |
-| **Max Depenetration Velocity** || The maximum velocity used to depenetrate this object.  |
-| **Walkable Slope Override** || Custom walkable slope setting for this body. See the [Walkable Slope](Engine/Physics/WalkableSlope) documentation for usage information. |
-| **Walkable Slope Behavior** || Behavior of this surface (whether we affect the walkable slope). Determines how the Walkable Slope Angle operates on the existing walkable angle, whether it increases or decreases it. |
-| **Walkable Slope Angle** || Override walkable slope, applying the rules of the Walkable Slope Behavior. |
-| **Custom Sleep Threshold Multiplier** || If the **SleepFamily** is set to custom, multiply the natural sleep threshold by this amount. A higher number will cause the body to sleep sooner. |
-| **Position Solver Iteration Count** || This physics body's solver iteration count for position. Increasing this will be more CPU intensive, but better stabilized. |
-| **Velocity Solver Iteration Count** || This physics body's solver iteration count for velocity. Increasing this will be more CPU intensive, but better stabilized. |
-| **Sleep Family** || The set of values used in considering when put this body to sleep. |
+| Property | Description |
+| --- | :--- |
+| **Primitives** | The basic shapes that compose this Bone's collision. |
+| **Bone Name** | Used in the PhysicsAsset case, associates this Body with a Bone in a Skeletal Mesh. |
+| **Consider for Bounds** | Should this bone be used in the calculation for the bounding volume of the Skeletal Mesh. |
 [/REGION]
 
-[/VAR]
-
-[VAR:p23]
+### Collision
 
 ![](PhATBody3.png)
 
 [REGION:raw]
-| Property | %spacer% | Description |
-| --- | :---: | :--- |
-| **Simulation Generates Hit Events** || Should 'Hit' events fire when this object collides during physics simulation. |
-| **Phys Material Override** || Allows you to override the PhysicalMaterial to use for simple collision on this body. |
-| **Collision Response** || See the [Collision Response Reference](Engine/Physics/Collision/Reference) documentation for more information. |
+| Property | Description |
+| --- | :--- |
+| **Simulation Generates Hit Events** | Should 'Hit' events fire when this object collides during physics simulation. |
+| **Phys Material Override** | Allows you to override the PhysicalMaterial to use for simple collision on this body. |
+| **Collision Response** | See the [Collision Response Reference](Engine/Physics/Collision/Reference) documentation for more information. |
 
 [/REGION]
-### Advanced
+
+#### Advanced
 
 ![](PhATBody3_adv.png)
 
 [REGION:raw]
-| Property | %spacer% | Description |
-| --- | :---: | :--- |
-| **Use CCD** || If true Continuous Collision Detection (CCD) will be used for this component. |
+| Property | Description |
+| --- | :--- |
+| **Use CCD** | If true Continuous Collision Detection (CCD) will be used for this component. |
 [/REGION]
 
-[/VAR]
-
-[VAR:p24]
-
+### Physical Animation
+		
 ![](PhATBody4.png)
 
 Please see the [](Engine/Physics/PhAT/HowTo) for more information on using **Physical Animations**.
 
-[/VAR]
-
----
-
-[OBJECT:modulesection]
-	[PARAMLITERAL:id]
-		ConsM
-	[/PARAMLITERAL]
-	[PARAM:heading]
-		Constraint Mode
-	[/PARAM]
-	[PARAM:content]
-		%c3%
-	[/PARAM]
-[/OBJECT]
-
-[VAR:c3]
+## Constraint Mode Properties
 
 When the Physics Asset Tool is in Body Mode and a Physics Body is selected, the Details panel will display the properties for that Physics Body.
 
-[OBJECT:modulesectionexpanded]
-	[PARAMLITERAL:id]
-		Cons2
-	[/PARAMLITERAL]
-	[PARAM:heading]
-		Constraint
-	[/PARAM]
-	[PARAM:content]
-		%p31%
-	[/PARAM]
-[/OBJECT]
-
-[OBJECT:modulesectionexpanded]
-	[PARAMLITERAL:id]
-		Line1
-	[/PARAMLITERAL]
-	[PARAM:heading]
-		Linear Limits
-	[/PARAM]
-	[PARAM:content]
-		%Engine/Physics/Constraints/ConstraintsReference:p2%
-	[/PARAM]
-[/OBJECT]
-
-[OBJECT:modulesectionexpanded]
-	[PARAMLITERAL:id]
-		Proj2
-	[/PARAMLITERAL]
-	[PARAM:heading]
-		Projection
-	[/PARAM]
-	[PARAM:content]
-		%Engine/Physics/Constraints/ConstraintsReference:p3%
-	[/PARAM]
-[/OBJECT]
-
-[OBJECT:modulesectionexpanded]
-	[PARAMLITERAL:id]
-		Angu2
-	[/PARAMLITERAL]
-	[PARAM:heading]
-		Angular Limits
-	[/PARAM]
-	[PARAM:content]
-		%Engine/Physics/Constraints/ConstraintsReference:p4%
-	[/PARAM]
-[/OBJECT]
-
-[OBJECT:modulesectionexpanded]
-	[PARAMLITERAL:id]
-		Lmoto
-	[/PARAMLITERAL]
-	[PARAM:heading]
-		Linear Motor
-	[/PARAM]
-	[PARAM:content]
-		%Engine/Physics/Constraints/ConstraintsReference:p5%
-	[/PARAM]
-[/OBJECT]
-
-[OBJECT:modulesectionexpanded]
-	[PARAMLITERAL:id]
-		Amoto
-	[/PARAMLITERAL]
-	[PARAM:heading]
-		Angular Motor
-	[/PARAM]
-	[PARAM:content]
-		%Engine/Physics/Constraints/ConstraintsReference:p6%
-	[/PARAM]
-[/OBJECT]
-
-[/VAR]
-
-[VAR:p31]
+### Constraint
 
 ![](PhATCon1.png)
 
 [REGION:raw]
-| Property | %spacer% | Description |
-| --- | :---: | :--- |
-| **Joint Name** || The name of the bone (**Physics Body**) this **Physics Constraint** is associated with. |
-| **Constraint Bone 1** || Name of the first bone (**Physics Body**) this **Physics Constraint** is connecting. This will be the "child" bone in a **Physics Asset**. |
-| **Constraint Bone 2** || Name of the second bone (**Physics Body**) this **Physics Constraint** is connecting. This will be the "parent" bone in a **Physics Asset**. |
+| Property | Description |
+| --- | :--- |
+| **Joint Name** | The name of the bone (**Physics Body**) this **Physics Constraint** is associated with. |
+| **Constraint Bone 1** | Name of the first bone (**Physics Body**) this **Physics Constraint** is connecting. This will be the "child" bone in a **Physics Asset**. |
+| **Constraint Bone 2** | Name of the second bone (**Physics Body**) this **Physics Constraint** is connecting. This will be the "parent" bone in a **Physics Asset**. |
 [/REGION]
-[/VAR]
+
+### Linear Limits
+
+[INCLUDE:Engine/Physics/Constraints/ConstraintsReference#p2(offset:1)]
+
+### Projection
+
+[INCLUDE:Engine/Physics/Constraints/ConstraintsReference#p3(offset:1)]
+
+### Angular Limits
+
+[INCLUDE:Engine/Physics/Constraints/ConstraintsReference#p4(offset:1)]
+
+### Linear Motor
+
+[INCLUDE:Engine/Physics/Constraints/ConstraintsReference#p5(offset:1)]
+
+### Angular Motor
+
+[INCLUDE:Engine/Physics/Constraints/ConstraintsReference#p6(offset:1)]
 
 
 
diff --git a/Engine/Documentation/Source/Engine/Physics/Physics.INT.udn b/Engine/Documentation/Source/Engine/Physics/Physics.INT.udn
--- a/Engine/Documentation/Source/Engine/Physics/Physics.INT.udn
+++ b/Engine/Documentation/Source/Engine/Physics/Physics.INT.udn
@@ -67,6 +67,7 @@ Unreal Engine 4 uses the PhysX 3.3 physics engine to drive its physical simulati
 		* [Physics Constraints User Guide](Engine/Physics/Constraints/ConstraintsUserGuide "%Engine/Physics/Constraints/ConstraintsUserGuide:description%")
 		* [](Engine/Physics/Constraints/ConstraintsBlueprints "%Engine/Physics/Constraints/ConstraintsBlueprints:description%")
 		* [](Engine/Physics/Constraints/ConstraintsReference "%Engine/Physics/Constraints/ConstraintsReference:description%")
+		* [](Engine/Physics/Tracing "%Engine/Physics/Tracing:description%")
 		* [Damping and Friction](Engine/Physics/Constraints/DampingAndFriction "%Engine/Physics/Constraints/DampingAndFriction:description%")
 	[/PARAM]
 [/OBJECT]
diff --git a/Engine/Documentation/Source/Engine/Physics/PhysicsBodies/Reference/PhysicsBodiesRef.INT.udn b/Engine/Documentation/Source/Engine/Physics/PhysicsBodies/Reference/PhysicsBodiesRef.INT.udn
--- a/Engine/Documentation/Source/Engine/Physics/PhysicsBodies/Reference/PhysicsBodiesRef.INT.udn
+++ b/Engine/Documentation/Source/Engine/Physics/PhysicsBodies/Reference/PhysicsBodiesRef.INT.udn
@@ -2,45 +2,7 @@ Availability: Public
 Title:Physics Bodies Reference
 Crumbs: %ROOT%
 Description:Physics Body, or Body Instance, property reference.
-
-[VAR:Topic]
-[OBJECT:Topic]
-	[PARAM:image]
-		![%Engine/Physics/PhysicsBodies:title%](Engine/Physics/physics_topic.png)
-	[/PARAM]
-	[PARAM:icon]
-		![](%ROOT%/physics_icon.png)(convert:false)
-	[/PARAM]
-	[PARAM:title]
-		%Engine/Physics/PhysicsBodies:title%
-	[/PARAM]
-	[PARAM:description]
-		%Engine/Physics/PhysicsBodies:description%
-	[/PARAM]
-	[PARAM:path]
-		[RELATIVE:Engine/Physics/PhysicsBodies]
-	[/PARAM]
-[/OBJECT]
-[/VAR]
-[VAR:TopicCompact]
-[OBJECT:TopicCompact]
-	[PARAM:image]
-		![%Engine/Physics/PhysicsBodies:title%](Engine/Physics/physics_topic.png)
-	[/PARAM]
-	[PARAM:icon]
-		![](%ROOT%/physics_icon.png)(convert:false)
-	[/PARAM]
-	[PARAM:title]
-		%Engine/Physics/PhysicsBodies:title%
-	[/PARAM]
-	[PARAM:description]
-		%Engine/Physics/PhysicsBodies:description%
-	[/PARAM]
-	[PARAM:path]
-		[RELATIVE:Engine/Physics/PhysicsBodies]
-	[/PARAM]
-[/OBJECT]
-[/VAR]
+topic-image:Engine/Physics/physics_topic.png
 
 This page contains a reference listing of properties in the Physics and Collision categories of properties. If you are looking for reference on Collision Responses or Collision Presets please see: [](Engine/Physics/Collision/Reference).
 
@@ -69,11 +31,11 @@ Below are the properties for Physics Bodies (Body Instances) separated by major
 [REGION:raw]
 | Property | %spacer% | Description |
 | --- | :---: | :--- |
-| **Simulate Physics** || %doxygen:FBodyInstance::bSimulatePhysics% |
+| **Simulate Physics** || [TEXT(type:"tooltip" namespace:"UObjectToolTips" key:"BodyInstance:bSimulatePhysics" output:"block")] |
 | **Mass in KG** || Mass of the body in KG. |
-| **Angular Damping** || %doxygen:FBodyInstance::AngularDamping% |
-| **Linear Damping** || %doxygen:FBodyInstance::LinearDamping% |
-| **Enable Gravity** || %doxygen:FBodyInstance::bEnableGravity% |
+| **Angular Damping** || [TEXT(type:"tooltip" namespace:"UObjectToolTips" key:"BodyInstance:AngularDamping" output:"block")] |
+| **Linear Damping** || [TEXT(type:"tooltip" namespace:"UObjectToolTips" key:"BodyInstance:LinearDamping" output:"block")] |
+| **Enable Gravity** || [TEXT(type:"tooltip" namespace:"UObjectToolTips" key:"BodyInstance:bEnableGravity" output:"block")] |
 |[REGION:tablesection]**Constraints**[/REGION]|||
 | **Lock Position** || [INCLUDE:#lockpos] |
 | **Lock Rotation** || [INCLUDE:#lockrot] |
@@ -87,22 +49,22 @@ Below are the properties for Physics Bodies (Body Instances) separated by major
 [REGION:raw]
 | Property | %spacer% | Description |
 | --- | :---: | :--- |
-| **Auto Weld** || %doxygen:FBodyInstance::bAutoWeld% |
-| **Start Awake** || %doxygen:FBodyInstance::bStartAwake% |
-| **Center Of Mass Offset** || %doxygen:FBodyInstance::COMNudge% |
-| **Mass Scale** || %doxygen:FBodyInstance::MassScale% |
-| **Max Angular Velocity** || %doxygen:FBodyInstance::MaxAngularVelocity% |
-| **Use Async Scene** || %doxygen:FBodyInstance::bUseAsyncScene% |
-| **Override Max Depenetration Velocity** || %doxygen:FBodyInstance::bOverrideMaxDepenetrationVelocity% |
-| **Max Depenetration Velocity** || %doxygen:FBodyInstance::MaxDepenetrationVelocity% |
-| **Override Walkable Slope on Instance** || %doxygen:FBodyInstance::bOverrideWalkableSlopeOnInstance% |
-| **Walkable Slope Override** || %doxygen:UBodySetup::WalkableSlopeOverride% See the **[](Engine/Physics/WalkableSlope)** documentation for usage information. |
-| **Walkable Slope Behavior** || %doxygen:FWalkableSlopeOverride::WalkableSlopeBehavior% |
-| **Walkable Slope Angle** || %doxygen:FWalkableSlopeOverride::WalkableSlopeAngle% 
-| **Sleep Family** || %doxygen:FBodyInstance::SleepFamily% |
-| **Position Solver Iteration Count** || %doxygen:FBodyInstance::PositionSolverIterationCount% |
-| **Velocity Solver Iteration Count** || %doxygen:FBodyInstance::VelocitySolverIterationCount% |
-| **Should Update Physics Volume** || %doxygen:USceneComponent::bShouldUpdatePhysicsVolume% |
+| **Auto Weld** || [TEXT(type:"tooltip" namespace:"UObjectToolTips" key:"BodyInstance:bAutoWeld" output:"block")] |
+| **Start Awake** ||[TEXT(type:"tooltip" namespace:"UObjectToolTips" key:"BodyInstance:bStartAwake" output:"block")] |
+| **Center Of Mass Offset** ||[TEXT(type:"tooltip" namespace:"UObjectToolTips" key:"BodyInstance:COMNudge" output:"block")] |
+| **Mass Scale** ||[TEXT(type:"tooltip" namespace:"UObjectToolTips" key:"BodyInstance:MassScale" output:"block")] |
+| **Max Angular Velocity** ||[TEXT(type:"tooltip" namespace:"UObjectToolTips" key:"BodyInstance:MaxAngularVelocity" output:"block")] |
+| **Use Async Scene** ||[TEXT(type:"tooltip" namespace:"UObjectToolTips" key:"BodyInstance:bUseAsyncScene" output:"block")] |
+| **Override Max Depenetration Velocity** ||[TEXT(type:"tooltip" namespace:"UObjectToolTips" key:"BodyInstance:bOverrideMaxDepenetrationVelocity" output:"block")] |
+| **Max Depenetration Velocity** ||[TEXT(type:"tooltip" namespace:"UObjectToolTips" key:"BodyInstance:MaxDepenetrationVelocity" output:"block")] |
+| **Override Walkable Slope on Instance** ||[TEXT(type:"tooltip" namespace:"UObjectToolTips" key:"BodyInstance:bOverrideWalkableSlopeOnInstance" output:"block")] |
+| **Walkable Slope Override** || [TEXT(type:"tooltip" namespace:"UObjectToolTips" key:"BodySetup:WalkableSlopeOverride" output:"block")] See the **[](Engine/Physics/WalkableSlope)** documentation for usage information. |
+| **Walkable Slope Behavior** || [TEXT(type:"tooltip" namespace:"UObjectToolTips" key:"WalkableSlopeOverride:WalkableSlopeBehavior" output:"block")] |
+| **Walkable Slope Angle** || [TEXT(type:"tooltip" namespace:"UObjectToolTips" key:"WalkableSlopeOverride:WalkableSlopeAngle" output:"block")] |
+| **Sleep Family** || [TEXT(type:"tooltip" namespace:"UObjectToolTips" key:"BodyInstance:SleepFamily" output:"block")] |
+| **Position Solver Iteration Count** || [TEXT(type:"tooltip" namespace:"UObjectToolTips" key:"BodyInstance:PositionSolverIterationCount" output:"block")] |
+| **Velocity Solver Iteration Count** || [TEXT(type:"tooltip" namespace:"UObjectToolTips" key:"BodyInstance:VelocitySolverIterationCount" output:"block")] |
+| **Should Update Physics Volume** || [TEXT(type:"tooltip" namespace:"UObjectToolTips" key:"SceneComponent:bShouldUpdatePhysicsVolume" output:"block")] |
 
 [/REGION]
 
@@ -128,9 +90,9 @@ Below are the properties for Physics Bodies (Body Instances) separated by major
 [REGION:raw]
 | Property | %spacer% | Description |
 | --- | :---: | :--- |
-| **Simulation Generates Hit Events** || %doxygen:FBodyInstance::bNotifyRigidBodyCollision% |
-| **Phys Material Override** || %doxygen:FBodyInstance::PhysMaterialOverride% |
-| **Generate Overlap Events** || %doxygen:UPrimitiveComponent::bGenerateOverlapEvents% |
+| **Simulation Generates Hit Events** || [TEXT(type:"tooltip" namespace:"UObjectToolTips" key:"BodyInstance:bNotifyRigidBodyCollision" output:"block")] |
+| **Phys Material Override** || [TEXT(type:"tooltip" namespace:"UObjectToolTips" key:"BodyInstance:PhysMaterialOverride" output:"block")] |
+| **Generate Overlap Events** || [TEXT(type:"tooltip" namespace:"UObjectToolTips" key:"PrimitiveComponent:bGenerateOverlapEvents" output:"block")] |
 | **Collision Responses** || See the **[](Engine/Physics/Collision/Reference)** documentation for more information. |
 [/REGION]
 
@@ -141,12 +103,12 @@ Below are the properties for Physics Bodies (Body Instances) separated by major
 [REGION:raw]
 | Property | %spacer% | Description |
 | --- | :---: | :--- |
-| **Use CCD** || %doxygen:FBodyInstance::bUseCCD% |
-| **Always Create Physics State** || %doxygen:UPrimitiveComponent::bAlwaysCreatePhysicsState% |
-| **Multi Body Overlap** || %doxygen:UPrimitiveComponent::bMultiBodyOverlap% | 
-| **Check Async Scene On Move** || %doxygen:UPrimitiveComponent::bCheckAsyncSceneOnMove% |
-| **Trace Complex On Move** || %doxygen:UPrimitiveComponent::bTraceComplexOnMove% |
-| **Return Material On Move** || %doxygen:UPrimitiveComponent::bReturnMaterialOnMove% |
+| **Use CCD** || [TEXT(type:"tooltip" namespace:"UObjectToolTips" key:"BodyInstance:bUseCCD" output:"block")] |
+| **Always Create Physics State** || [TEXT(type:"tooltip" namespace:"UObjectToolTips" key:"PrimitiveComponent:bAlwaysCreatePhysicsState" output:"block")] |
+| **Multi Body Overlap** || [TEXT(type:"tooltip" namespace:"UObjectToolTips" key:"PrimitiveComponent:bMultiBodyOverlap" output:"block")] | 
+| **Check Async Scene On Move** || [TEXT(type:"tooltip" namespace:"UObjectToolTips" key:"PrimitiveComponent:bCheckAsyncSceneOnMove" output:"block")] |
+| **Trace Complex On Move** || [TEXT(type:"tooltip" namespace:"UObjectToolTips" key:"PrimitiveComponent:bTraceComplexOnMove" output:"block")] |
+| **Return Material On Move** || [TEXT(type:"tooltip" namespace:"UObjectToolTips" key:"PrimitiveComponent:bReturnMaterialOnMove" output:"block")] |
 | **Can Ever Affect Navigation** || Whether this component can potentially influence navigation |
 
 [/REGION]
@@ -160,44 +122,44 @@ Below are the properties for Physics Bodies (Body Instances) separated by major
 <!--
 [EXCERPT:lockpos]
 
-%doxygen:FBodyInstance::bLockTranslation%
+[TEXT(type:"tooltip" namespace:"UObjectToolTips" key:"BodyInstance:bLockTranslation" output:"block")]
 
-[REGION:raw]
-| Property | %spacer2% | Description |
-| --- | :---: | :--- |
-| **Lock X Translation** || %doxygen:FBodyInstance::bLockXTranslation% |
-| **Lock Y Translation** || %doxygen:FBodyInstance::bLockYTranslation% |
-| **Lock Z Translation** || %doxygen:FBodyInstance::bLockZTranslation% |
+[REGION:simpletable]
+| Property | Description |
+| --- | :--- |
+| **Lock X Translation** | [TEXT(type:"tooltip" namespace:"UObjectToolTips" key:"BodyInstance:bLockXTranslation" output:"block")] |
+| **Lock Y Translation** | [TEXT(type:"tooltip" namespace:"UObjectToolTips" key:"BodyInstance:bLockYTranslation" output:"block")] |
+| **Lock Z Translation** | [TEXT(type:"tooltip" namespace:"UObjectToolTips" key:"BodyInstance:bLockZTranslation" output:"block")] |
 [/REGION]
 [/EXCERPT]
 
 [EXCERPT:lockrot]
 
-%doxygen:FBodyInstance::bLockRotation%
+[TEXT(type:"tooltip" namespace:"UObjectToolTips" key:"BodyInstance:bLockRotation" output:"block")]
 
-[REGION:raw]
-| Property | %spacer2% | Description |
-| --- | :---: | :--- |
-| **Lock X Rotation** || %doxygen:FBodyInstance::bLockXRotation% |
-| **Lock Y Rotation** || %doxygen:FBodyInstance::bLockYRotation% |
-| **Lock Z Rotation** || %doxygen:FBodyInstance::bLockZRotation% |
+[REGION:simpletable]
+| Property | Description |
+| --- | :--- |
+| **Lock X Rotation** | [TEXT(type:"tooltip" namespace:"UObjectToolTips" key:"BodyInstance:bLockXRotation" output:"block")] |
+| **Lock Y Rotation** | [TEXT(type:"tooltip" namespace:"UObjectToolTips" key:"BodyInstance:bLockYRotation" output:"block")] |
+| **Lock Z Rotation** | [TEXT(type:"tooltip" namespace:"UObjectToolTips" key:"BodyInstance:bLockZRotation" output:"block")] |
 [/REGION]
 [/EXCERPT]
 
 [EXCERPT:mode]
 
-%doxygen:FBodyInstance::DOFMode%
-
-[REGION:raw]
-| Property | %spacer2% | Description |
-| --- | :---: | :--- |
-| **Default** || Inherits the degrees of freedom from the project settings. |
-| **Six DOF** || Specifies which axis to freeze rotation and movement along. |
-| **YZPlane** || Allows 2D movement along the Y-Z plane. |
-| **XZPlane** || Allows 2D movement along the X-Z plane. |
-| **XYPlane** || Allows 2D movement along the X-Y plane. |
-| **Custom Plane** || Allows 2D movement along the plane of a given normal. |
-| **None** || No constraints. |
+[TEXT(type:"tooltip" namespace:"UObjectToolTips" key:"BodyInstance:DOFMode" output:"block")]
+
+[REGION:simpletable]
+| Property | Description |
+| --- | :--- |
+| **Default** | Inherits the degrees of freedom from the project settings. |
+| **Six DOF** | Specifies which axis to freeze rotation and movement along. |
+| **YZPlane** | Allows 2D movement along the Y-Z plane. |
+| **XZPlane** | Allows 2D movement along the X-Z plane. |
+| **XYPlane** | Allows 2D movement along the X-Y plane. |
+| **Custom Plane** | Allows 2D movement along the plane of a given normal. |
+| **None** | No constraints. |
 [/REGION]
 [/EXCERPT]
 -->
\ No newline at end of file
diff --git a/Engine/Documentation/Source/Engine/Physics/Tracing/HowTo/HowToLanding.INT.udn b/Engine/Documentation/Source/Engine/Physics/Tracing/HowTo/HowToLanding.INT.udn
new file mode 100644
--- /dev/null
+++ b/Engine/Documentation/Source/Engine/Physics/Tracing/HowTo/HowToLanding.INT.udn
@@ -0,0 +1,21 @@
+Availability: Public
+Title:Traces How to Guides
+Crumbs:
+Description: Guides covering Tracing (Raycasting) in Unreal Engine 4.
+version: 4.14
+Parent:Engine/Physics/Tracing
+type:landing
+tags:Physics
+tags:Tracing
+tags:Raycast
+tags:Gameplay
+order:1
+topic-image:Engine/Physics/Tracing/trace_topic.png
+related:Engine/Physics/Tracing/Overview
+
+The following guides show you how to work with Tracing (Raycasting) in Unreal Engine 4.
+
+##Intermediate 
+
+[DIR(output:"fancy" parent:"Engine/Physics/Tracing/HowTo" org:"hierarchy" end:"1")]
+
diff --git a/Engine/Documentation/Source/Engine/Physics/Tracing/HowTo/MultiLineTraceByChannel/MultiLineTraceByChannel.INT.udn b/Engine/Documentation/Source/Engine/Physics/Tracing/HowTo/MultiLineTraceByChannel/MultiLineTraceByChannel.INT.udn
new file mode 100644
--- /dev/null
+++ b/Engine/Documentation/Source/Engine/Physics/Tracing/HowTo/MultiLineTraceByChannel/MultiLineTraceByChannel.INT.udn
@@ -0,0 +1,51 @@
+Availability: Public
+Title:Using a Multi Line Trace (Raycast) by Channel
+Crumbs:
+Description:This how-to shows using a Multi Line Trace by Channel Blueprint node to return a the all Actors that respond on the Visilibity channel, and prints their names.
+parent:Engine/Physics/Tracing/HowTo
+version: 4.14
+skilllevel:Intermediate
+type:how-to
+tags:Physics
+tags:Tracing
+tags:Raycast
+tags:Gameplay
+topic-image:Engine/Physics/Tracing/trace_topic.png
+related:Engine/Physics/Tracing/HowTo/SingleLineTraceByChannel
+related:Engine/Physics/Tracing/HowTo/MultiLineTraceByObject
+
+**MultiLineTraceByChannel** will perform a collision trace along a given line and return all hits encountered (up to and including) the first blocking hit, returning only objects that respond to the specified Trace Channel. What this means (practically), is that if their are a number of **Actors** or **Components** with collision that **Overlap** the specified Trace Channel between the start and the end of the trace, you will recieve them all. But, if the first hit **Blocks** the specified Trace Channel, you will only recieve that one item. If you want to recieve all items regardless of overlap, or blocking a trace channel, you should use a [MultiLineTraceByObject](Engine/Physics/Tracing/HowTo/MultiLineTraceByObject) node. Below, you will find steps for setting up **MultiLineTraceByChannel**. 
+
+### Steps
+
+1.  Follow the steps used for the [LineTraceByChannel](Engine/Physics/Tracing/HowTo/SingleLineTraceByChannel) example to set up a trace. 
+
+1.  Replace the **LineTraceByChannel** node with the **MultiLineTraceByChannel** node. 
+
+1.  Drag off the **Out Hits** pin and add a **ForEachLoop** node. 
+
+	![](GHT2B_18.png)
+
+	Since multiple Actors are hit, for each of them, we want to do something (in this example, print the Actors to the screen). 
+
+1.  Drag off the **Array Element** and add a **Break Hit Result**. Then, drag off the **Hit Actor**, add a **Get Display Name (Object)** and connect to a **Print String**. 
+	
+	[REGION:lightbox]
+	[![](GHT2B_19.png)(w:800)](GHT2B_19.png)
+	[/REGION]
+
+	_Click image for a full view._
+
+	[REGION:note]
+	This will now print to the string each of the Actors hit by the array. 
+	[/REGION]
+
+## Result
+
+Here, we have a Glass Window in front of a Physics Actor. 
+
+![](GHT2B_17.png)
+
+The Glass Window is a **Destructible Mesh**, and we have set its **Trace Response** to **Visibility** in its collision settings, which is set to **Overlap**; while the Physics Actor (Cube) has its **Visibility** setting to **Block**. You could use this for situations where you may want to shoot through an object (destroying it) while also hitting an enemy.
+
+
diff --git a/Engine/Documentation/Source/Engine/Physics/Tracing/HowTo/MultiLineTraceByObject/MultiLineTraceByObject.INT.udn b/Engine/Documentation/Source/Engine/Physics/Tracing/HowTo/MultiLineTraceByObject/MultiLineTraceByObject.INT.udn
new file mode 100644
--- /dev/null
+++ b/Engine/Documentation/Source/Engine/Physics/Tracing/HowTo/MultiLineTraceByObject/MultiLineTraceByObject.INT.udn
@@ -0,0 +1,55 @@
+Availability: Public
+Title:Using a Multi Line Trace (Raycast) by Object
+Crumbs:
+Description:This how-to shows using a Multi Line Trace by Object Blueprint node to return a the all World Dynamic Actors it hits, and prints their names.
+parent:Engine/Physics/Tracing/HowTo
+version: 4.14
+skilllevel:Intermediate
+type:how-to
+tags:Physics
+tags:Tracing
+tags:Raycast
+tags:Gameplay
+topic-image:Engine/Physics/Tracing/trace_topic.png
+related:Engine/Physics/Tracing/HowTo/SingleLineTraceByChannel
+related:Engine/Physics/Tracing/HowTo/MultiLineTraceByChannel
+
+**MultiLineTraceForObjects** will perform a collision trace along a given line and return all hits encountered, returning only objects that match one of the specified Object Types. Below, you will find steps for setting up **MultiLineTraceForObjects**. 
+
+## Steps
+
+1.  Follow the steps used for the [LineTraceByChannel](Engine/Physics/Tracing/HowTo/SingleLineTraceByChannel) example to set up a trace. 
+
+1.  Replace the **LineTraceByChannel** node with the **MultiLineTraceForObjects** node. 
+
+1.  Drag off the **Object Types** pin and add a **Make Array** node, then use the drop-down menus to add Objects to the Array. 
+
+	![](GHT2B_21.png)
+
+	Here we assign **WorldDynamic** and **Physics Body** as our Object Types. You can use the **Add pin** button to add more Object Types to the Array. 
+	
+1.  Drag off the **Out Hits** pin of the trace node and add a **ForEachLoop** node. 
+
+	![](GHT2B_22.png)
+
+	This will allow us to do something for each of the Actors hit by the trace. 
+
+1.  Drag off the **Array Element** and add a **Break Hit Result**. Then, off the **Hit Actor**, add a **To String (Object)** and connect to a **Print String**. 
+
+	[REGION:lightbox]
+	[![](GHT2B_23.png)(w:800)](GHT2B_23.png)
+	[/REGION]
+
+	_Click image for a full view._
+
+	[REGION:note]
+	This will now print to the screen each of the Actors hit by the array. 
+	[/REGION]
+
+## Result
+
+Here, we have a hanging ceiling light (World Dynamic Object Type) in front of a Physics Actor (Physics Body Object Type). 
+
+![](GHT2B_20.png)
+
+**Multi Line Trace by Object**, unlike **Multi Line Trace by Channel**, does not stop at the first object it hits, which is why the trace continues through the light to the cube. 
