Copying //UE4/Dev-Documentation to //UE4/Samples-Main (Source: //UE4/Dev-Documentation @ 3372893)

diff --git a/Engine/Documentation/Source/Engine/Physics/Tracing/HowTo/SingleLineTraceByChannel/SingleLineTraceByChannel.INT.udn b/Engine/Documentation/Source/Engine/Physics/Tracing/HowTo/SingleLineTraceByChannel/SingleLineTraceByChannel.INT.udn
new file mode 100644
--- /dev/null
+++ b/Engine/Documentation/Source/Engine/Physics/Tracing/HowTo/SingleLineTraceByChannel/SingleLineTraceByChannel.INT.udn
@@ -0,0 +1,102 @@
+Availability: Public
+Title:Using a Single Line Trace (Raycast) by Channel
+Crumbs:
+Description:This how to covers using a Single Line Trace by Channel Blueprint node to return a the first Actor it hits that responds on the Visibility channel, and prints its name.
+parent:Engine/Physics/Tracing/HowTo
+version: 4.14
+skilllevel:Intermediate
+type:how-to
+tags:Physics
+tags:Tracing
+tags:Raycast
+tags:Gameplay
+topic-image:Engine/Physics/Tracing/trace_topic.png
+related:Engine/Physics/Tracing/HowTo/SingleLineTraceByObject
+
+**LineTraceByChannel** will perform a collision trace along a given line and return the first Object that the trace hits. Below, you will find steps for setting up **SingleLineTraceByChannel**. 
+
+## Steps
+
+1.  Create a new project using the **Blueprint First Person** template with **Include Starter Content** and open the project. 
+
+1.  In the **FirstPersonBP/Blueprints** folder, open the **FirstPersonCharacter** Blueprint. 
+
+1.  Right-click in the graph, search for, and add an **Event Tick** node. 
+
+	![](GHT2B_1.png)
+	
+	This will cause the trace to run every frame.
+	
+1.  Drag off the execute pin, then search for the **LineTraceByChannel** node. 
+
+	![](GHT2B_2.png)
+
+1.  While holding down the **Ctrl** key, drag in the **FirstPersonCamera** component. 
+
+	![](GHT2B_3.png)
+
+	This is where we will start our trace from, the camera. 
+
+1.  Drag off the **FirstPersonCamera** node, and add a **Get World Location** node, then connect it to the **Start** of the trace. 
+
+1.	Drag off the **FirstPersonCamera** node again and add a **Get World Rotation** node. 
+
+	![](GHT2B_4.png)(w:800)
+
+	Here, we are starting the trace from the location of the FirstPersonCamera, then we are getting the rotation of the FirstPersonCamera. 
+
+1.	Drag off the **Get World Rotation** node and add a **Get Forward Vector**, then drag off that and add a **Vector * Float** node set to **1500**. 
+
+	![](GHT2B_5.png)
+
+	We are getting the rotation and the forward vector, extending outward from it by 1500 (this value is the length of the trace). 
+
+1.	Drag off the **Get World Location** node and add a **Vector + Vector** node, connecting (as shown below) to the **End** of the trace node. 
+
+	![](GHT2B_6.png)(w:800)
+
+	Here, we are taking the location of the FirstPersonCamera and extending out from it, 1500 units based on its rotation and forward vector.
+
+1.	On the Trace node, set the **Draw Debug Type** to **For One Frame**. 
+
+	![](GHT2B_7.png)
+
+	This will allow us to see a debug line while playing in-game to see our line trace. 
+
+1.	Drag off the execution out pin of the trace and add a **Print String** node. 
+
+	![](GHT2B_8.png)
+
+1.	Drag off the **Out Hit** pin and add search for **Break Hit** then add a **Break Hit Result** node. 
+
+	![](GHT2B_9.png)
+
+1.	Drag off the **Hit Actor** pin (off of the **Break Hit Result**), add a **To String (Object)** node, and connect it to the **Print String** node. 
+
+	[REGION:lightbox]
+	[![](GHT2B_10.png)(w:800)](GHT2B_10.png)
+	[/REGION]
+
+	_Click image for a full view._
+	
+	[REGION:note]
+	This will allow us to debug print the object we hit with our trace.
+	[/REGION]
+
+1.	Click the **Compile** button, then play in the Editor and look at the cubes in the level. 
+
+	![](GHT2B_11.png)(w:800)
+
+	_Here, we have ejected from the First Person perspective so that you can see the view angle of the trace._
+	
+	[REGION:note]
+	You should see that when the trace hits a cube, it prints the cube to the screen. 
+	[/REGION]
+
+## End Result
+
+The example above returns all Objects that are set to respond to the provided Trace Channel, however there may be instances when you want to only return certain objects. In the example above, you can use the **Actors to Ignore** pin, taking an Array of Actors that should be ignored by the trace (this means that you have to specify each individual Actor to ignore). 
+
+[REGION:tip]
+You can also perform a **LineTraceByObject** where only the specified **ObjectTypes** are returned. This will allow you to target a specific set of Objects (only to be included in the trace). 
+[/REGION]
diff --git a/Engine/Documentation/Source/Engine/Physics/Tracing/HowTo/SingleLineTraceByObject/SingleLineTraceByObject.INT.udn b/Engine/Documentation/Source/Engine/Physics/Tracing/HowTo/SingleLineTraceByObject/SingleLineTraceByObject.INT.udn
new file mode 100644
--- /dev/null
+++ b/Engine/Documentation/Source/Engine/Physics/Tracing/HowTo/SingleLineTraceByObject/SingleLineTraceByObject.INT.udn
@@ -0,0 +1,50 @@
+Availability: Public
+Title:Using a Single Line Trace (Raycast) by Object
+Crumbs:
+Description:This how-to covers using a Single Line Trace by Object Blueprint node to return a the first World Dynamic Actor it hits, and prints its name.
+parent:Engine/Physics/Tracing/HowTo
+version: 4.14
+skilllevel:Intermediate
+type:how-to
+tags:Physics
+tags:Tracing
+tags:Raycast
+tags:Gameplay
+topic-image:Engine/Physics/Tracing/trace_topic.png
+related:Engine/Physics/Tracing/HowTo/SingleLineTraceByChannel
+
+**LineTraceForObjects** will perform a collision trace along a given line and return the first Object the trace hits that matches one of the specified Object Types. To set up **LineTraceForObjects** trace, do the following: 
+
+## Steps
+
+1.  Follow the steps used for the [LineTraceByChannel](Engine/Physics/Tracing/HowTo/SingleLineTraceByChannel) example to set up a trace. 
+
+1.  Replace the **LineTraceByChannel** node with the **LineTraceForObjects** node. 
+
+1.  Drag off the **Object Types** pin and add the **Make Array** node. 
+
+	![](GHT2B_12.png)
+
+1.  On the **Make Array** node, specify the **ObjectType** you want to trace for (via the drop-down menu). 
+
+	![](GHT2B_13.png)
+
+	[REGION:note]
+	Here, we are tracing for **WorldDyanmic** Objects. You can add more by clicking the **Add Pin** button. 
+	[/REGION]
+
+1.  You can set the rest of the trace up the same way a **LineTraceByChannel** is setup. 
+
+	[REGION:lightbox]
+	[![](GHT2B_15.png)(w:800)](GHT2B_15.png)
+	[/REGION]
+
+	_Click image for a full view._
+
+## Result
+
+ We have added a single **WorldDynamic** Object into our level. 
+
+![](GHT2B_16.png)
+
+Only the added Actor is now returning as a hit; consequently, the cubes (since they are Physics Actors) do not return a hit. 
diff --git a/Engine/Documentation/Source/Engine/Physics/Tracing/Overview/TracingOverview.INT.udn b/Engine/Documentation/Source/Engine/Physics/Tracing/Overview/TracingOverview.INT.udn
new file mode 100644
--- /dev/null
+++ b/Engine/Documentation/Source/Engine/Physics/Tracing/Overview/TracingOverview.INT.udn
@@ -0,0 +1,129 @@
+Availability: Public
+Title:Traces Overview
+Crumbs:
+Description:Overview of the Unreal Engine 4 tracing system.
+version: 4.14
+type:overview
+Parent:Engine/Physics/Tracing
+order:1
+tags:Physics
+tags:Tracing
+tags:Raycast
+tags:Gameplay
+topic-image:trace_overview_topic.png
+related:Engine/Physics/Tracing/HowTo
+
+[TOC(start:1 end:4)]
+
+<!-- Common Variables -->
+
+[VAR:spacer]![](%ROOT%/spacer.png)(w:150 convert:false)[/VAR]
+
+[VAR:spacer2]![](%ROOT%/spacer.png)(w:40 convert:false)[/VAR]
+
+<!-- Overview -->
+
+[REGION:fullwidth]
+![](traceTitle.png)(convert:false)
+[/REGION]
+
+**Traces** offer a method for reaching out in your levels and getting feedback on what is present along a line segment. You use them by providing two end points (a start and end location) and the physics system "traces" a line segment between those points, reporting any Actors (with collision) that it hits. Traces are essentially the same as **Raycasts** or **Raytraces** in other software packages.
+
+Whether you want to know if one **Actor** can "see" another, or figure out the normal of a specific polygon, or simulate high velocity weaponry, or if you need to know that an **Actor** has entered a space; Traces offer you a reliable and computationally cheap solution. This document covers the basic feature set of Traces in Unreal Engine 4 (UE4).
+
+## Trace by Channel or Object Type
+
+Since Traces use the physics system, you'll be able to define the category of thing you want to trace against. There are two broad categories to choose from: Channels and Object Types. Channels are used for things like visibility and the camera, and almost exclusively have to do with tracing. Object Types are the physics types of Actors with collision in your scene, such as Pawns, Vehicles, Destructible Actors, etc...
+
+You can add more Channels and Object types as you need them. See [](Engine/Physics/Collision/HowTo/AddCustomCollisionType) for more information on how to do this.
+
+## Returning Single or Multiple Hits
+
+When tracing, you can choose to return the first thing that matches the critera hit by the trace, or you can return everything hit by the trace that matches the critera.
+
+Special consideration is given to **Multi Trace by Channel** versus **Multi Trace by Object**. When using **Muli Trace by Channel** the trace will return all **overlaps** up to and including the first **block**. Imagine shooting a bullet through some tall grass that then impacts a wall.
+
+**Multi Trace by Object** will return everything that matches an object type the trace is looking for, assuming the component is set to return trace queries. Making it great for counting the number of objects between the start and end of the trace.
+
+[REGION:fullwidth]
+![](singleVSmulti.png)(convert:false)
+[/REGION]
+
+### Hit Results
+
+When a Trace hits something, it returns a "Hit Result" Struct. Made available in Blueprints (also in C++), this is what the struct looks like:
+
+| --- | --- |
+| ![](BP_HitStruct.png) | [INCLUDE:#ts1] |
+
+<!--
+[EXCERPT:ts1]
+| **Member** || **Definition** |
+| --- | :---: | :--- |
+| **Blocking Hit** || Whether or not the hit was a blocking hit. This is used when **Multi Tracing by Channel**, due to the ability to have traces simply overlap and not stop the trace. |
+| **Initial Overlap** || Whether this is the first overlap in a series of results. |
+| **Time** || This is the "Time" of impact along the trace direction ranging from [0.0 to 1.0]. If there was no hit, this will return 1.0. |
+| **Location** || The world space location of the hit that is modified based on the shape of the trace. |
+| **Impact Point** || The absolute location of the hit. Does not include the shape of the trace, only the point of the hit. |
+| **Normal** || The direction of the trace. |
+| **Impact Normal** || The normal of the hit surface. |
+| **Phys Mat** || The **Physical Material** of the hit surface. |
+| **Hit Actor** || The hit **Actor**. |
+| **Hit Component** || The specfic **Component** hit. |
+| **Hit Bone Name** || If traced against a **Skeletal Mesh**, this is the name of the bone that was hit. |
+| **Hit Item** || Primitive-specific data, recording which item in the primitive was hit. |
+| **Face Index** || If colliding with a trimesh or landscape, this is the index of the face that was hit. |
+| **Trace Start** || The start location of the trace. |
+| **Trace End** || The end location of the trace. |
+[/EXCERPT]
+-->
+
+## Using Shape Traces
+
+[REGION:fullwidth]
+![](shapeTracesEX.png)(convert:false)
+[/REGION]
+
+When a Line Trace is not enough, you may get the results you are after by using a Shape Trace instead. For example, maybe you are creating a "vision cone" for an enemy, and you want to detect players that enter it. A Line Trace may not be enough as players may be able to avoid detection by ducking under the Line Trace. 
+
+In this situation, you could possibly use a Box Trace, Capsule Trace, or Sphere Trace. 
+
+[REGION:fullwidth]
+![](TracesOverview_ShapeTraces.png)
+[/REGION]
+
+Shape Traces function like Line Traces, where you are sweeping and checking for collision from a start point to an end point; however, Shape Traces have an added layer of checking because you are using a shape as a volume (of sorts) in your Raycast. You can use a Shape Trace as a Single Trace, or as a Multi Trace, and each are set up in the same manner as a Line Trace; however, you will need to provide additional details pertaining to the size (or orientation) of the shape you use.
+
+## Getting UV Coordinates from Trace
+
+A trace can return the UV Coordinates of the Actor it hit, assuming trace complex is used. As of 4.14 this only works on **Static Mesh Components**, **Procedural Mesh Components**, and **BSP**. It will **not** work on **Skeletal Mesh Components** because you trace against the **Physics Asset**, which doesn't have UV coordinates (even if you choose to trace complex).
+
+[REGION:note]
+Using this feature will increase CPU memory usage, because UE4 needs to keep an additional copy of vertex positions and UV coordinates in main memory.
+[/REGION]
+
+### Enabling UV Coordinates from Trace
+
+To enable this feature, follow these steps:
+
+1. Access your **Project Settings** from the **Edit Menu**.
+
+	![](accessProj.png)
+
+1. Enable the **Support UV From Hit Results** feature in the **Physics Section** of your **Project Settings**.
+
+	![](ProjectSettings.png)
+
+1. Restart the Editor.
+
+	[REGION:note]
+	You can inspect the Blueprint **Find Collision UV** node using this feature prior to restarting the editor; however, the node will only return [0.0, 0.0] when you inspect it.
+	If you want the node to return the correct UV data, you'll have to restart the editor.
+	[/REGION]
+
+
+## Other Features	
+
+Tracing also has a number of minor features to limit what they return, making them easier to debug. They have the ability to trace against **Complex Collision** (if a Static Mesh or Procedural Mesh has it enabled). If they are called from an **Actor**, they can be told to ignore all attached components by enabling the **Actor** to trace through itself. Finally, they have the option to show a representation of the trace as red and green lines (with larger boxes representing hits).
+
+<!-- Draw Debug Type, Trace Complex, Ignore Self -->
\ No newline at end of file
diff --git a/Engine/Documentation/Source/Engine/Physics/Tracing/Tracing.INT.udn b/Engine/Documentation/Source/Engine/Physics/Tracing/Tracing.INT.udn
new file mode 100644
--- /dev/null
+++ b/Engine/Documentation/Source/Engine/Physics/Tracing/Tracing.INT.udn
@@ -0,0 +1,29 @@
+Availability: Public
+Title:Traces (Raycasts)
+Crumbs:
+Description:The landing page for Physics Traces (Raycasts).
+version: 4.14
+Parent:Engine/Physics
+type:landing
+order:4
+tags:Physics
+tags:Tracing
+tags:Raycast
+tags:Gameplay
+topic-image:trace_topic.png
+social-image:trace_social.png
+
+![](traceTop.png)(convert:false)
+
+There may be instances in your game where you want to determine if the player character is looking at something, and if so, alter the game state in some way (for example, highlight something when a player looks at it). Or maybe you want to determine if an enemy can see the player character, and if so, start shooting or engaging them in some way. You can accomplish both scenarios by using **Traces** (or **Raycasts**) to "shoot" out an invisible ray which will detect geometry between two points and if geometry is hit, return what was hit so that you may then do something with it. 
+
+There are different options available when running a Trace. You can run a Trace to check for collision with any Objects where hit Objects are returned, or you can run a Trace by Trace Channel where any Objects hit will return hit information if the Object is set to specifically respond to a specified Trace Channel (which can be set via Collision Settings). 
+
+In addition to running traces by Objects or by Trace Channel, you can run your Trace to detect Single hits or Multi hits, where a Single Trace returns a singular hit result and a Multi Trace returns multiple hits resulting from the Trace. With Traces, you can also specify the type of ray that is used: a straight line, a box, a capsule, or a sphere. 
+
+##Topics
+
+To learn more about Raycasts using Blueprints, follow the links below:
+
+%Engine/Physics/Tracing/Overview:topic%
+%Engine/Physics/Tracing/HowTo:topic%
\ No newline at end of file
diff --git a/Engine/Documentation/Source/Engine/Physics/Vehicles/VehicleCenterOfMass/VehicleCenterOfMass.INT.udn b/Engine/Documentation/Source/Engine/Physics/Vehicles/VehicleCenterOfMass/VehicleCenterOfMass.INT.udn
new file mode 100644
--- /dev/null
+++ b/Engine/Documentation/Source/Engine/Physics/Vehicles/VehicleCenterOfMass/VehicleCenterOfMass.INT.udn
@@ -0,0 +1,81 @@
+Availability: Docs
+Title: Vehicle Center of Mass
+Crumbs:%ROOT%
+Description: An overview of how Center of Mass works with Vehicles.
+Type: Reference
+Parent: Engine/Physics/Vehicles/
+Order: 
+Version: 4.15
+Tags: Vehicles
+SkillLevel: Beginner
+
+
+Weight distribution of a vehicle is important to its control, as it affects a variety of its characteristics, like handling, acceleration, and traction. Different vehicle types will require different weight distributions depending on their intended use. For the purposes of game development, these characteristics can also define what style of game youâ€™re making, whether that is an arcade-style racer, a simulation, or even a hybrid of the two. Altering **Center of Mass** enables you to change weight distribution of the vehicle.
+
+In games, the primary use of Center of Mass is for vehicles, but it can also be used for large Physics Bodies that encompass an irregular shape. In your 
+[Physics Asset](Engine/Physics/PhAT), you will often use one large Physics Body to define the majority of the Mass for the vehicle or large object. The 
+Center of Mass will be generated at the center of this Physics Body, which can make the vehicle handle oddly, so you can adjust the Center of Mass to 
+account for where the mass of the vehicle is really located. 
+
+[![](CenterOfMass.png)(w:400)](CenterOfmass.png)
+
+Depending on where the Center of Mass is located for your vehicle, you can shift it to be predominantly front-heavy, causing the vehicle to have a tendency toward **understeering** (not turning enough when going around corners), or to be predominantly rear-heavy, causing it to tend toward **oversteering** (turning more sharply than intended). In most cases, it is ideal to find a neutral balance for your center of mass so that the vehicle is more easily controlled. 
+
+[REGION:imagetable]
+| ![Understeer](Understeering.png)(w:450) | ![Oversteer](Oversteering.png)(w:450) |
+| --- | --- |
+| Understeer | Oversteer |
+[/REGION]
+
+Also, when you consider where to place the Center of Mass, it is worth noting that this choice can even have an effect on how your vehicle handles in the air. In this example, the Center of Mass has been lowered and moved closer to the rear of the car. Since it's low to the ground and can get to high speed quickly, the lower rear center of mass helps stabilize the car, especially with jumps!
+
+[REGION:imagetable]
+| [INCLUDE:#originalCoM] | [INCLUDE:#adjustedCoM] |
+| --- | --- |
+| Center of Mass: X: 0, Y: 0, Z: 0 | Center of Mass: X: -25, Y: 0, Z: -10 |
+[/REGION]
+
+<!--- 
+
+[EXCERPT:originalCoM]
+[OBJECT:EmbeddedVideo]
+	[PARAMLITERAL:width]
+	480
+	[/PARAMLITERAL]
+	[PARAMLITERAL:height]
+	270
+	[/PARAMLITERAL]
+	[PARAMLITERAL:videoid]
+	CPjxbM05m9c
+	[/PARAMLITERAL]
+[/OBJECT]
+[/EXCERPT:originalCoM]
+
+[EXCERPT:adjustedCoM]
+[OBJECT:EmbeddedVideo]
+	[PARAMLITERAL:width]
+	480
+	[/PARAMLITERAL]
+	[PARAMLITERAL:height]
+	270
+	[/PARAMLITERAL]
+	[PARAMLITERAL:videoid]
+	EzMv0EsWTj0
+	[/PARAMLITERAL]
+[/OBJECT]
+[/EXCERPT:adjustedCoM]
+--->
+
+## Debugging Center of Mass
+
+To help you better debug mass properties and inertia tensors associated with physics objects while in the Level Editor, you can enable the showflag by going to **Show** > **Advanced** > **Mass Properties**. 
+
+[REGION:imagetable]
+| ![](VehicleCoM1.png)(w:500) | ![](VehicleCoM2.png)(w:500) |
+| --- | --- |
+| CoM: X: 0, Y: 0, Z: 0 | CoM: X: -25, Y: 0, Z: -10 |
+[/REGION]
+
+The thickness of each axis indicates the moment of inertia magnitude along that axis. 
+
+
diff --git a/Engine/Documentation/Source/Engine/QuickStart/EditorQuickStart.INT.udn b/Engine/Documentation/Source/Engine/QuickStart/EditorQuickStart.INT.udn
--- a/Engine/Documentation/Source/Engine/QuickStart/EditorQuickStart.INT.udn
+++ b/Engine/Documentation/Source/Engine/QuickStart/EditorQuickStart.INT.udn
@@ -9,6 +9,7 @@ Version: 4.9
 parent:Engine/Editor
 order:1
 checkpoint: editorqs
+type:quick start
 
 [VAR:Topic]
 [OBJECT:Topic]
diff --git a/Engine/Documentation/Source/Engine/Rendering/HDRDisplayOutput/HDRDispalyOutput.INT.udn b/Engine/Documentation/Source/Engine/Rendering/HDRDisplayOutput/HDRDispalyOutput.INT.udn
new file mode 100644
--- /dev/null
+++ b/Engine/Documentation/Source/Engine/Rendering/HDRDisplayOutput/HDRDispalyOutput.INT.udn
@@ -0,0 +1,133 @@
+Availability: Docs
+Crumbs:
+Title: High Dynamic Range Display Output
+Description: Overview for HDR Display Output and available options.
+Type: Overview
+SkillLevel: Intermediate
+Version: 4.15
+Parent: Engine/Rendering/
+Order:
+Tags: Rendering
+
+
+
+In Unreal Engine 4, you can now output to High Dynamic Range (HDR) displays taking advantage of features such as higher contrast and wider color gamut! The goal here is to give the displayed images the
+characteristics that are more like natural light conditions experienced in the "real world". This is part of the move to the **Academy Color Encoding System** (ACES) standard that is a pipeline to ensure
+consistent color is preserved across multiple formats and displays, and also as a way to _future-proof_ the source material used without having to adjust it for another medium. 
+
+[REGION:fullwidth raw]
+![](LDR_HDRcomparison.jpg)
+[/REGION]
+[REGION:caption]
+_These are simulated images and are for demonstration purposes only since it is impossible to convey HDR difference on an LDR screen._
+[/REGION]
+
+with the current implementation, the full processing of the rendered scene is handled through what is known as the **ACES Viewing Transform**. This process works by using scene-referred and display-referred images.
+
+* **Scene-referred** images maintain the original _linear light_ values of the source materials without limiting the exposure range. 
+* **Display-referred** is the final image transformed into the color space of the display being used. 
+
+By using this pipeline, the original source files do not have to be edited each time they are used with a different display in order to guarantee correct coloring. Instead, the display being output to maps it to the 
+correct color space.
+
+The ACES Viewing Transform, more specifically, works in the following order in the viewing pipeline:
+
+* **Look Modification Transform (LMT)** - This part of the process takes the ACES color-encoded image that has had a creative "look" (color grading and correction) applied to it and then outputs the images rendered with
+ACES in combination with the Reference Rendering Transform (RRT) and an Output Device Transform (ODT).
+* **Reference Rendering Transform (RRT)** - Then, this part takes the scene-referred color values and converts them to display-referred. In this process, it enables the rendered images to not rely on a specific display
+and, instead, can ensure a larger gamut and dynamic range that is correct for the specific display it's being output to, including ones that have yet to be created.
+* **Output Device Transform (ODT)** - Finally, this takes the HDR Data output of the RRT to then match it with the different devices and color spaces they can display. Because of this, each target requires its own
+ODT to be matched against (eg Rec709, Rec2020, DCI-P3, etc.).
+
+
+## Enabling HDR Output
+
+You can enable HDR output at runtime by toggling the console variables or by using the **GameUserSettings** node in Blueprints. 
+
+![](HDRGameUserSettings.png)(w:800)
+
+The **GameUserSettings** controls will automatically clamp to the closest available match for the current output device and set all flags accordingly. Alternatively, you can use the following console variables to enable
+and change the available options for the HDR device and color gamut output need. 
+
+[REGION:simpletable]
+| Console Variable | Description |
+| --- | --- |
+| `r.HDR.EnableHDROutput` | When set to 1, this recreates the swap-chain and enables HDR output. |
+| `r.HDR.Display.OutputDevice` | This is the device format of the output display [INCLUDE:#deviceOuptut] |
+| `r.HDR.Display.ColorGamut` | This is the color gamut of the output display. [INCLUDE:#colorGamut] |
+[/REGION]
+
+<!---
+[EXCERPT:deviceOuptut]
+|||
+| --- | --- |
+| 0 | sRGB (LDR) (default) |
+| 1 | Rec709 (LDR) |
+| 2 | Explicit gamma mapping (LDR) |
+| 3 | ACES 1000-nit ST-2084 (Dolby PQ) (HDR) |
+| 4 | ACES 2000-nit ST-2084 (Dolby PQ) (HDR) |
+| 5 | ACES 1000-nit ScRGB (HDR) |
+| 6 | ACES 2000-nit ScRGB (HDR) |
+[/EXCERPT:deviceOuptut]
+
+[EXCERPT:colorGamut]
+|||
+| --- | --- |
+| 0 | Rec709 / sRGB, D65 (default) |
+| 1 | DCI-P3, D65 |
+| 2 | Rec2020 / BT2020, D65 |
+| 3 | ACES, D60 |
+| 4 | ACEScg, D60 |
+[/EXCERPT:colorGamut]
+--->
+
+Once you've setup the Blueprint or C++ call for GameUserSettings, you can run any cooked project, use the command line argument `-game mode`, use Standalone game mode, or Play-in-Editor (PIE) in a new window using 
+exclusive fullscreen (press **Alt** + **Enter** or use the command `fullscreen`). Using borderless windowed or windowed mode will not currently work for HDR output. 
+
+
+### Low Dynamic Range (LDR) UI support in HDR
+
+[REGION:note]
+This feature is experimental and may change in a future release.
+[/REGION]
+
+There is support for bringing Low Dynamic Range (LDR) 
+
+When HDR output has been enabled your User Interface (UI) may not look correct. We've added experimental support for LDR UI composition for this reason. It attempts to match the
+LDR look as closely as possible, although, it's recommended to boost the UI slightly to avoid looking washed out next to the vibrancy of the main scene. 
+
+To do this, you can use the following console variables:
+
+[REGION:simpletable]
+| Console Variable | Description |
+| --- | --- |
+| `r.HDR.UI.CompositeMode` | When set to 1, this enables HDR UI composition, which attempts to preserve LDR visuals and blending. |
+| `r.HDR.UI.Level` | This adjusts the brightness of the composited UI. It is recommended to use a value of 1.5 or 2. |
+[/REGION]
+
+
+## HDR Hardware and Setup Considerations
+
+With the many different types of displays and TVs that are available, along with console and PC considerations, you may find that you need to use a specific piece of hardware or change a setting to get 
+HDR output to work properly. The following are some things to consider when you're running through your setup process:
+
+* Make sure that your system is connected to the HDR display using a high-speed HDMI 2.0 cable so that you can take advantage of the transfers speeds. (HDMI 1.4 may work, but may have issues displaying the content.)
+Not all HDMI ports may accept HDMI 2.0 and HDR. Refer to your display's manual if you're unsure. 
+* Ensure that the TV port is HDR enabled. Sometimes this may be listed in your TV or display settings as "HDMI Deep Color" or "Enhanced Format." Refer to your display's manual if you're unsure.
+* With some consoles, like the PS4, you may need to disable the system setting for **Enable HDCP** for HDR output to work properly.
+* For some displays, you may need to adjust display settings to get the correct output. 
+	* With NVIDIA GPUs, use the NVIDIA Control Panel and adjust the display resolution Output Color Format to use **RGB** and the Output Color Depth to use **10-bit** or **12-bit**.
+
+## Notes and Limitations
+
+* By default, the film-mapping curve is disabled during HDR output due to poor compatibility with LDR controls. 
+* Currently, only paths for 1000-nit and 2000-nit display outputs are implemented. 
+* D3D11 restrictions limit HDR output to exclusive fullscreen support only. With D3D12 support in Windows 10, this can likely be extended to enable individual viewports for HDR Output. The implementation on Mac currently supports this, already.
+
+## Reference Material
+
+* [http://www.oscars.org/science-technology/sci-tech-projects/aces](http://www.oscars.org/science-technology/sci-tech-projects/aces)
+* [http://www.oscars.org/science-technology/aces/aces-documentation](http://www.oscars.org/science-technology/aces/aces-documentation)
+
+
+
diff --git a/Engine/Documentation/Source/Engine/Rendering/LightingAndShadows/CapsuleShadows/CapsuleShadowsLanding.INT.udn b/Engine/Documentation/Source/Engine/Rendering/LightingAndShadows/CapsuleShadows/CapsuleShadowsLanding.INT.udn
--- a/Engine/Documentation/Source/Engine/Rendering/LightingAndShadows/CapsuleShadows/CapsuleShadowsLanding.INT.udn
+++ b/Engine/Documentation/Source/Engine/Rendering/LightingAndShadows/CapsuleShadows/CapsuleShadowsLanding.INT.udn
@@ -2,27 +2,22 @@
 title: Capsule Shadows
 description:Landing page for Capsule Shadows in Unreal Engine 4
 crumbs:
-version:4.13
+version:4.15
+parent:Engine/Rendering/LightingAndShadows/
+order:
+related: Engine/Rendering/LightingAndShadows/DistanceFields/HowTo/DFHT_4
+type:landing
+tags: Lighting
 
-**Capsule Shadows** support soft shadow casting for Skeletal Meshes by creating a capsule representation of your character using a custom Physics Asset. 
-The soft shadows serve to ground the character in indirectly-lit areas and give soft shadows in directly-lit areas. This shadowing technique is available in Unreal Engine 4.11 and higher. 
+[REGION:fullwidth raw]
+![CapsuleShadowsBanner](CapsuleShadowsBanner.png)
+[/REGION]
 
+As of Unreal Engine 4.11 and later, **Capsule Shadows** can be used to enable support for soft shadow casting for Skeletal Meshes by using a capsule representation of your character made 
+from a Physics Asset. The soft shadows serve to ground the character in indirectly lit areas and give soft shadows in directly lit areas.  
 
-[OBJECT:TopicButtonList]
-	[PARAM:icon]
-		![](%ROOT%/tech_icon.png)(convert:false)
-	[/PARAM]
-	[PARAM:title]
-		Capsule Shadows Topics
-	[/PARAM]
-	[PARAM:description]
-	[/PARAM]
-	[PARAM:links]
-	* [](Engine\Rendering\LightingAndShadows\CapsuleShadows\Overview "%Engine\Rendering\LightingAndShadows\CapsuleShadows\Overview:description%")
-	* [](Engine\Rendering\LightingAndShadows\CapsuleShadows\BestPractices "%Engine\Rendering\LightingAndShadows\CapsuleShadows\BestPractices:description%")
-	* [](Engine\Rendering\LightingAndShadows\CapsuleShadows\Reference "%Engine\Rendering\LightingAndShadows\CapsuleShadows\Reference:description%")
-	* [](Engine\Rendering\LightingAndShadows\CapsuleShadows\HowTo "%Engine\Rendering\LightingAndShadows\CapsuleShadows\HowTo:description%")
-	[/PARAM]
-[/OBJECT]
+## Getting Started
 
+%Engine/Rendering/LightingAndShadows/CapsuleShadows/Overview:topic%
+%Engine/Rendering/LightingAndShadows/CapsuleShadows/QuickStart:topic%
 
diff --git a/Engine/Documentation/Source/Engine/Rendering/LightingAndShadows/CapsuleShadows/Overview/CapsuleShadowsOverview.INT.udn b/Engine/Documentation/Source/Engine/Rendering/LightingAndShadows/CapsuleShadows/Overview/CapsuleShadowsOverview.INT.udn
--- a/Engine/Documentation/Source/Engine/Rendering/LightingAndShadows/CapsuleShadows/Overview/CapsuleShadowsOverview.INT.udn
+++ b/Engine/Documentation/Source/Engine/Rendering/LightingAndShadows/CapsuleShadows/Overview/CapsuleShadowsOverview.INT.udn
@@ -1,208 +1,174 @@
 ﻿Availability: Docs
 Title: Capsule Shadows Overview
-Crumbs: %ROOT%, Engine, 
-Description: Capsule Shadows Overview
-Version: 4.13
-
-
-[REGION:center]
-![](image_9.png)(w:900)
-[/REGION]
+Crumbs:  
+Description: An overview of Capsules Shadows used in Unreal Engine 4.
+Version: 4.15
+Type: overview
+Parent: Engine/Rendering/LightingAndShadows/CapsuleShadows
+topic-image:cs_overview_topic.png
+tags: Lighting
 
 [TOC(start:2 end:3)]
 
+In Unreal Engine 4 (UE4), we now support soft shadowing for characters in lit areas with **Capsule Shadows** which is a capsule representation of your character with a Physics Asset. 
+This shadowing method can make a character feel grounded in these indirectly lit areas and have soft shadows in directly lit areas, which is not possible when using traditional
+shadow mapping techniques.
 
-**Capsule Shadows** are used to generate soft shadows in lit areas by using a capsule representation of a character with a Physics Asset.
-This feature can make a character feel grounded in both direct and indirect lighting by producing consistent soft shadows.
-
-
-## How do they work?
 
-Capsule Shadows start by using a unique Physics Asset that can be assigned to any Skeletal Mesh as a **Shadow Physics Asset**. This asset represents an approximation of the Skeletal Mesh 
-with respect to the shadows it should cast. Once this information is known, the [Indirect Lighting Samples](Engine/Rendering/LightingAndShadows/IndirectLightingCache/) that are 
-generated after a light build can be used to cast a cone towards the direction of the light source. As the character nears this point, the cone is checked against the capsules to 
-see if there are any overlaps. If there are, then all the capsule shadows can be accumulated to build soft shadows for the character.
+## Character Capsule Representation
 
-[REGION:imagetable]
-| ![](CS_Without.png)(w:370) | ![](CS_With.png)(w:370) |
-| ---- | ---- |
-| Without Capsule Shadows | With Capsule Shadows |
-[/REGION]
-
-[REGION:left]
-[INCLUDE:Engine/Rendering/LightingAndShadows/CapsuleShadows\Overview#CapsuleShadowSlider]
-[/REGION]
+A [Physics Asset](Engine/Physics/PhAT/) is used to create an appoximate representation of the character to support very soft shadowing. Since the capsules are attached to the bones of the character, this 
+allows for the shadow to move and cast accurately within the scene. 
 
-<!---
-[EXCERPT:CapsuleShadowSlider]
-[OBJECT:SliderSequence]
-	[PARAMLITERAL:max]
-	5
-	[/PARAMLITERAL]
-	[PARAMLITERAL:ext]
-	jpg
-	[/PARAMLITERAL]
-	[PARAMLITERAL:id]
-	FS
-	[/PARAMLITERAL]
-	[PARAM:caption]
-	**Drag the slider to see what it takes to render the Capsule Shadows for a character.**
+[OBJECT:ComparisonSlider]
+	[PARAM:before]
+	![Character with Capsule Representation](CapsuleShadowPA.png)(w:500)
 	[/PARAM]
-	[PARAM:image_list]
-	![](CS_1.png)(w:750)
-	![](CS_2.png)(w:750)
-	![](CS_3.png)(w:750)
-	![](CS_4.png)(w:750)
-	![](CS_6.png)(w:750)
+	[PARAM:after]
+	![Capsule Representation Only](CapsuleShadowPA1.png)(w:500)
 	[/PARAM]
 [/OBJECT]
-[/EXCERPT:CapsuleShadowSlider]
--->
-
-The Physics asset that is created for Capsule Shadows and assigned to the Skeletal Mesh is used to represent the character in these indirectly-lit areas. Note that this is a mock-up representation of the capsules for the character. 
-
-1. Once lighting has been built for the scene, the character will have no shadowing in areas that have indirect lighting. 
-2. The Physics asset that is created for Capsule Shadows and assigned to the Skeletal Mesh is used to represent the character in these indirectly-lit areas. Note that this is a mock-up representation of the capsules for the character.
-3. With lighting built, the precomputed [volume lighting samples](Engine/Rendering/LightingAndShadows/IndirectLightingCache/) are placed so that Skeletal Actors can blend in with the indirectly-lit environment by referencing the 
-lighting information stored in the placed samples.
-4. These light samples are used to draw a cone towards the capsule representaiton of the character based on the directional intensity of the light. The cone can be used to mask out shadowed areas to get a soft shadow representation for 
-the character.  
-5. The intesnity and directionality of the capsule shadow representation can be computed to ground the character in the environment giving us the final render with soft shadowing for movable actors in indirectly-lit areas. 
 
 [REGION:note]
-Capsule Shadows can also work in Dynamic lights, when **Capsule Direct Shadow** is enabled, to produce a soft shadow in directly-lit areas. See the [Capsule Direct Shadowing](Engine\Rendering\LightingAndShadows\CapsuleShadows\Overview#CapsulesforDirectShadowing) section of this page.
+The shadows on the ground are not representative of Capsule Shadows.
 [/REGION]
 
-#### Capsule Representation of Character
 
-A Physics Asset is used to create a accurate representation of the character. Since the capsules are attached to the bones of the character this allows for the shadows to move with the character.
+## Usage
 
-For steps to setup your Phsyics Asset see the [How to Create a Shadow Phsyics Assets](Engine/Rendering/LightingAndShadows/CapsuleShadows/HowTo/CreateShadowAsset) page.
+![](ShadowPhysicsAsset.png)(w:400)
 
-For tips on setting up your Shadow Physics Asset see the [Best Practices](Engine/Rendering/LightingAndShadows/CapsuleShadows/BestPractices) page for Capsule Shadows.
-
-[REGION:lightbox]
-[![](PA_WithoutArms.png)(w:500)](PA_withoutArms.png)
-[/REGION]
-
-[REGION:caption]
-_Capsule representation of a Skeletal Mesh character_
-[/REGION]
-
-[REGION:note]
-The shadows on the ground are not representative Capsule Shadows. 
-[/REGION]
+Open your Skeletal Mesh, and in the **Asset Details** panel, use the **Shadow Physics Asset** selection to assign your Physics Asset that will be used for Capsule Shadows.
 
 
 ## Capsule Shadow Settings
 
-You can enable Capsule Shadows for each Skeletal Mesh in your level by enabling the option for **Capsule Direct Shadows** and/or **Capsule Indirect Shadows** in its **Detail's** panel.
+![](CapsuleShadowSettings.png)(w:400)
+[REGION:simpletable]
+| Property | Description |
+| -------- | ----------- |
+| **Capsule Direct Shadow** | This will enable soft shadowing from direct (movable) lights when the capsule representation is assigned to the Skeletal Mesh's Shadow Physics Asset slot. |
+| **Capsule Indirect Shadow** | This will enable soft shadowing from precomputed lighting (lightmaps and skylight) when the capsule representation is assigned to the Skeletal Mesh's Shadow Physics Asset slot. |
+| **Capsule Indirect Shadow Min Visibility** | This will allow artists to control how dark the capsule shadow, in indirectly lit areas, can be. |
+[/REGION]
 
-![](EnableCapsuleShadows.png)(w:350)
 
+### Capsule Indirect Shadows
 
-#### Capsules Direct Shadowing
+![](CapsuleShadowsSetting2.png)(w:400)
 
-**Capsule Direct Shadows** enables soft shadowing for your character when in areas that are directly lit by any **Movable** or **Stationary** light source. 
+When you enable **Capsule Indirect Shadow**, the capsule representation of the character will be used to cast a directional soft shadow based on the [volume lighting samples](https://docs.unrealengine.com/latest/INT/Engine/Rendering/LightingAndShadows/IndirectLightingCache/)
+placed and computed by Lightmass during a lighting build. Capsule Indirect Shadows enables your characters to feel grounded in these areas where traditional shadow mapping will not work very well.
 
-[REGION:lightbox]
-[![](DirectCapsuleShadows.png)(w:500)](DirectCapsuleShadows.png)
-[/REGION]
-[REGION:caption]
-_Click the image for full size_
-[/REGION]
-
-##### Light Source Angle
+[OBJECT:ComparisonSlider]
+	[PARAM:before]
+	![Capsule Indirect Shadow Enabled](CSEnabled.png)(w:500)
+	[/PARAM]
+	[PARAM:after]
+	![Capsule Indirect Shadow Disabled](CSDisabled.png)(w:500)
+	[/PARAM]
+[/OBJECT]
 
-For directly lit Capsule Shadows you can adjust the **Light Source Angle** of the Directional Light to soften the shadows. 
+When Capsule Indirect Shadow is enabled, your character will cast a soft shadow that can ground them in areas with only bounce lighting. 
 
 [OBJECT:ComparisonSlider]
 	[PARAM:before]
-	![Light Source Angle: 1.0 (Default)](DirectLightSourceAngle1.png)(w:500)
+	![Indirect Capsule Shadows Enabled](SkylightCS1.png)(w:500)
 	[/PARAM]
 	[PARAM:after]
-	![Light Source Angle: 2.0](DirectLightSourceAngle2.png)(w:500)
+	![Indirect Capsule Shadows Disabled](SkylightCS2.png)(w:500)
 	[/PARAM]
 [/OBJECT]
 
-#### Capsules Indirect Shadowing
-
-**Capsule Indirect Shadows** enables soft shadowing for your character in areas that are indirectly lit by any **Static** or **Stationary** light source. Capsule shadowing is created using the 
-[Indirect Lighting Cache](Engine/Rendering/LightingAndShadows/IndirectLightingCache/) that is generated during a light build. 
+In an open area with only a Sky Light as the light source, there will be little directionality, since light is coming from all around. This produces a subtle soft "blob" shadow below
+the character when using precomputed lighting. 
 
-[REGION:lightbox]
-[![](IndirectDefault.png)(w:500)](IndirectDefault.png)
-[/REGION]
-[REGION:caption]
-_Click the image for full size_
-[/REGION]
+[OBJECT:EmbeddedVideo]
+	[PARAMLITERAL:width]
+	640
+	[/PARAMLITERAL]
+	[PARAMLITERAL:height]
+	360
+	[/PARAMLITERAL]
+	[PARAMLITERAL:videoid]
+	ns4QW1BFSHg
+	[/PARAMLITERAL]
+[/OBJECT]
 
+In an enclosed area where light is coming through an opening, the [indirect lighting cache](https://docs.unrealengine.com/latest/INT/Engine/Rendering/LightingAndShadows/IndirectLightingCache/) is used to 
+give directionality and softness to the capsule shadow by interpolating between the placed volume lighting samples as the character moves through the space. At the doorway, the character's shadow will be
+less intense with little directionality, but as the character moves further away from the doorway, the shadow's intensity and directionality will increase.
 
-## Directionality and Intensity of Indirect Shadows
 
-The directionality and intensity of the capsule shadowing is determined by the light samples generated after a light build. The brighter intensity of the lighting sample the less intense and softer the shadow will be since light is coming from many directions. 
+#### Indirect Minimum Shadow Visibility
 
-### Closed Area Shadowing
+Artists are afforded additional control by adjusting the **Capsule Indirect Shadow Min Visibility**. This enables you to adjust how dark the capsule shadow will be in indirectly lit areas, using precomputed lighting. 
+This can be be useful to reduce any self-shadowing from a capsule in these areas, or to soften the shadow's intensity, so that it blends nicely with surrounding shadows.
 
-When you have a opening for light to pour through, the intensity of the light will be brighter at the opening and dimmer at the back part of this area. 
+[OBJECT:ComparisonSlider]
+	[PARAM:before]
+	![Capsule Shadow Indirect Min Visibility: 0.1 (Default)](CSEnabled.png)(w:500)
+	[/PARAM]
+	[PARAM:after]
+	![Capsule Shadow Indirect Min Visibility: 0.45](CSMinVisibility_45.png)(w:500)
+	[/PARAM]
+[/OBJECT]
 
-[REGION:lightbox]
-[![](CS_CaveFront.png)(w:500)](CS_CaveFront.png)
-[/REGION]
-[REGION:caption]
-_Click the image for full size_
-[/REGION]
 
-In the front of this area the character's shadow is much softer since there is more light and the shadow has less directionality since more light is coming from above. 
+### Capsule Direct Shadows
 
-[REGION:lightbox]
-[![](CS_CaveBack.png)(w:500)](CS_CaveBack.png)
-[/REGION]
-[REGION:caption]
-_Click the image for full size_
-[/REGION]
+![](CapsuleShadowsSetting1.png)(w:400)
 
-However, at the back of this area the shadow is more intense, since it is more dimly lit and has more directionality. The majority of the light is coming from a single opening towards the front which means there will be
-less light traveling to the back of this area, and the character is lit more strongly from one direction.
+When **Direct Capsule Shadow** is enabled, the **Light Source Angle** for Directional Lights and **Source Radius** for Spot or Point Lights is used to define how soft the shadows will be from its receiving surface. 
+This enables you to directly control the softness of your Capsule Shadows when directly lit.
 
-### Open Area Shadowing
+#### Light Source Angle 
 
-In open areas that receive indirect shadowing, Capsule Shadows will appear much softer, since the light can come from all around.
+The **Light Source Angle** property of a Directional Light is used for dynamic shadowing methods to create soft area shadows. When adjusting this property, the shadow will soften the further it is from the 
+surface it's casting onto. 
 
-[REGION:lightbox]
-[![](CS_openarea.png)(w:500)](CS_openarea.png)
-[/REGION]
-[REGION:caption]
-_Click the image for full size_
+[REGION:imagetable]
+| ![](CS_LSA1.png)(w:400) | ![](CS_LSA2.png)(w:400) |
+| --- | --- |
+| Light Source Angle: 1.0 (Default) | Light Source Angle: 2.0 |
 [/REGION]
 
+#### Source Radius
 
+Similarly to how the Light Source Angle of a Directional Light works, the **Source Radius** of a Stationary Spot or Point Light can be used to give your character soft shadows. When you use a larger Source Radius, the light will soften the further the shadow caster is from the surface.
 
+[REGION:imagetable]
+| ![](CS_SR1.png)(w:400) | ![](CS_SR2.png)(w:400) |
+| --- | --- |
+| Source Radius: 5.0 | Source Radius: 15.0 |
+[/REGION]
+
+[REGION:note]
+Once lighting has been built for your scene, you can adjust the **Source Radius** property without having to rebuild lighting. This property only affects movable actors with Capsule Shadows or Mesh Distance Fields enabled.
+[/REGION]
 
 ## Performance
 
-Capsule Shadows are more efficient than standard shadow mapping techniques because of tile deferred culling which cull the capsule influences to only the tiles that could possibly 
-receive shadows from that capsule. This means that most of the areas on the screen were unaffected and did not require additional rendering work. The tiles that are potentially 
-receiving shadows will have to compute, for each pixel, the number of capsules that directly affect that pixel. Capsule Shadows perform better when there are fewer affected tiles on 
-the screen, as well as when there are fewer capsules.
+The GPU performance cost of Capsule Shadows is proportional to the number of capsules used, the number of characters, and the screen size that their shadow affects. 
 
-[OBJECT:ComparisonSlider]
-	[PARAM:before]
-	![Standard Shdow Mapping](DirectShadowMapping.png)(w:500)
-	[/PARAM]
-	[PARAM:after]
-	![Capsule Direct Shadows Enabled](DirectCapsuleShadows.png)(w:500)
-	[/PARAM]
-[/OBJECT]
-
-The performance benefit of Capsule Shadows is determined by the number of capsule bodies that are being used and the tiles influenced. Whereas, standard shadow mapping is more expensive because
-it needs to consider all the triangles in the shadow map, which means you can only get sharp shadows that have a little bit of blur to them.
+Using a Radeon 7870 at 1080p, for example: 
+ 
+[REGION:simpletable]
+| GPU Cost with 10 Capsules | Time in milliseconds (ms) |
+| ---- | --- |
+| **A single character on screen** | 0.29 ms |
+| **Each additional character on screen** | 0.05 ms |
+[/REGION]
 
+This implementation is very efficient, since it computes shadowing at half resolution with a depth-aware upsample, while using screen tile culling to reduce the amount of shadowing
+work to only happen where it is needed. 
 
 ## Limitations
 
-* Requires DX 11 due to tiled deferred implementation using compute shaders.
+* Requires DirectX 11, due to tiled deferred implementation using compute shaders.
 * Self-shadowing artifacts can exist for arbitrary mesh shapes.
-* Can only use Spheres and Sphylls for capsule representation.
-* Only support for Skeletal Meshes, but could be made to work with static meshes.
-* When capsule shadows get so soft that they become ambient occlusion, there's an artifact in the shadowing that causes a hard line. 
\ No newline at end of file
+* Can only use Sphyls and Spheres for capsule representation.
+* When capsule shadows get so soft that they become ambient occlusion, there's an artifact in the shadowing that causes a hard line. 
+
+
+
diff --git a/Engine/Documentation/Source/Engine/Rendering/LightingAndShadows/CapsuleShadows/QuickStart/1/CSQS_1.INT.udn b/Engine/Documentation/Source/Engine/Rendering/LightingAndShadows/CapsuleShadows/QuickStart/1/CSQS_1.INT.udn
new file mode 100644
--- /dev/null
+++ b/Engine/Documentation/Source/Engine/Rendering/LightingAndShadows/CapsuleShadows/QuickStart/1/CSQS_1.INT.udn
@@ -0,0 +1,41 @@
+Availability:Docs
+Title:1 - Required Setup
+Crumbs:
+Description: In this step we will use the Project Browser to create a New Project.
+SkillLevel:Intermediate
+Version:4.15
+checkpoint: capsuleshadowsqs
+parent:Engine/Rendering/LightingAndShadows/CapsuleShadows/QuickStart
+type:quick start
+order:1
+tags: Lighting
+ 
+[Nav]
+
+Before we get started, the first thing we'll need to do is to create a new project that we can work from for the rest of this guide. 
+In the steps below, we'll use Unreal Engine 4's Project Browser to create the project. After completing this step, you will have a 
+template that you can build onto for your future projects. 
+
+## Steps
+
+1. Start by opening **Unreal Engine** from the **Epic Games Launcher**. The [Project Browser](https://docs.unrealengine.com/latest/INT/Engine/Basics/Projects/Browser/index.html) will appear.
+1. Click on the **New Project** tab, then in the **Blueprint** tab, select **Third Person**. We will need to give our project a name, so we will use "Quick Start." For the purpose of this tutorial we 
+will not need Starter Content, so you can set that property to **No Starter Content**. You can now click **Create Project**.
+
+	![](NewProject.png)(w:800)
+	
+## End Result
+
+After creating our project, and once the Unreal Editor has finished loading, you will be presented with the Main Editor Interface shown below. In the Content Browser, you will see various 
+folders for the Blueprint Third Person Template that we will use in the next steps.
+
+[REGION:lightbox]
+[![](LoadedProject.png)(w:800)](LoadedProject.png)
+[/REGION]
+[REGION:caption]
+_Click image for full size._
+[/REGION]
+
+In the next step, we will start by setting up a new Physics Asset for our Third Person Character that will be used for our Capsule Shadows.
+
+[Nav]
\ No newline at end of file
diff --git a/Engine/Documentation/Source/Engine/Rendering/LightingAndShadows/CapsuleShadows/QuickStart/2/CSQS_2.INT.udn b/Engine/Documentation/Source/Engine/Rendering/LightingAndShadows/CapsuleShadows/QuickStart/2/CSQS_2.INT.udn
new file mode 100644
--- /dev/null
+++ b/Engine/Documentation/Source/Engine/Rendering/LightingAndShadows/CapsuleShadows/QuickStart/2/CSQS_2.INT.udn
@@ -0,0 +1,72 @@
+Availability:Docs
+Title:2 - Your Character - Creating a Shadow Physics Asset
+Crumbs:
+Description:In this step we will create a new Physics Asset for our Skeletal Mesh that will be used for shadow representation of our character. 
+SkillLevel:Intermediate
+Version:4.15
+checkpoint: capsuleshadowsqs
+parent:Engine/Rendering/LightingAndShadows/CapsuleShadows/QuickStart
+Prereq: Engine/Physics/PhAT/HowTo/
+type:quick start
+order:2
+tags:Lighting
+
+
+[Nav]
+
+With your project ready to go, in this step, we'll create a new Physics Asset for the Third Person Character's Skeletal mesh that will be used for shadow representation with Capsule Shadows. 
+
+
+## Steps
+
+1. Inside of the **Content Browser**, use the folder hierarchy to navigate to **Mannequin** > **Character** > **Mesh** folder. Here you will find the asset named **SK_Mannequin**, which you'll use. Select it, and then 
+right-click, to bring up the context menu. From there, select **Create** > **Physics Asset** > **Create**.
+
+	[REGION:raw]
+	![](ContextMenuCreatePA.png)(w:960)
+	[/REGION]
+
+1. After you click Create, the **New Asset** window will open to create the new Physics Asset. Set the **Minimum Bone Size** to 15 and then click **OK**.
+
+	![](NewPhysicsAsset.png)
+
+	[REGION:note]
+	Only use Sphyl bodies for the best result. Sphere bodies will also work, but they are not as flexible in their use with character Skeletal Meshes.
+	[/REGION]
+
+1. Next, the [Physics Asset Tool](Engine/Physics/PhAT) (PhAT) will open with your new Physics Asset for the Skeletal Mesh `SK_Mannequin`. 
+
+	![](PA_Character.png)(w:500)
+
+1. This is an optional step. Before moving forward, it is recommended that you take a moment and give your newly created Physics Asset a name. Do this by minimizing the PhAT window, and in the **Content Browser**, locate
+the Physics Asset that was created for SK_Mannequin in your **Mesh** folder. Give it a name so that you can locate it later ("SPA_Mannequin" would be a good name). Once you've done this, you can maximize the PhAT window again.
+		
+	![](CB_SPA_Mannequin.png)(w:390)
+	
+1. In the PhAT window, you'll now want to adjust the various Physics Bodies, so that the capsule represents the character more accurately. This will require deleting some existing bodies, along with scaling and rotating
+the remaining ones to better fit. Below, you'll find a list of things you should keep in mind when adjusting and removing your capsules: 
+
+	* Limit the number of bodies by removing unnecessary capsules; like hands, arms, multiple capsules for the torso and neck, etc.
+	* Foot capsule placement is very important to make your character look grounded, and may need to be tweaked later to look more accurate.
+	* Use a little bit of overlap between the capsules's joints to avoid gaps in the shadow.
+	* Capsules for the arms are not necessary for this Quick Start and can be removed. 
+
+1. When you are finished with your adjustments, you should have something that looks similar to this. The placement does not have to be perfect and can easily be tweaked later to refine any shadowing issues 
+	that you may run into.
+	
+	![](PA_CharacterAdjusted.png)(w:500)
+
+## End Result
+
+Now that the Physics Asset for the character's capsule representation is complete, you should be familiar with the process of creating a new Physics Asset and setting it up to represent your character with 
+as few bodies as possible. 
+
+[REGION:imagetable]
+| ![](Before.png)(w:325) | ![](After.png)(w:325)  |
+| --- | --- |
+| Before: 21 Capsule Bodies | After: 10 Capsule Bodies |
+[/REGION]
+
+In the next step, we'll assign our new Physics Asset to our Skeletal Mesh. 
+
+[Nav]
diff --git a/Engine/Documentation/Source/Engine/Rendering/LightingAndShadows/CapsuleShadows/QuickStart/3/CSQS_3.INT.udn b/Engine/Documentation/Source/Engine/Rendering/LightingAndShadows/CapsuleShadows/QuickStart/3/CSQS_3.INT.udn
new file mode 100644
--- /dev/null
+++ b/Engine/Documentation/Source/Engine/Rendering/LightingAndShadows/CapsuleShadows/QuickStart/3/CSQS_3.INT.udn
@@ -0,0 +1,44 @@
+Availability:Docs
+Title:3 - Your Character - Assigning the Physics Asset
+Crumbs:
+Description:In this step we will assign our newly created Physics Asset to our Skeletal Mesh Character. 
+SkillLevel:Intermediate
+Version:4.15
+checkpoint: capsuleshadowsqs
+parent:Engine/Rendering/LightingAndShadows/CapsuleShadows/QuickStart
+type:quick start
+order:3
+tags:Lighting
+ 
+
+[Nav]
+
+In the previous page, we created a new Physics Asset for our Skeletal Mesh character that we'll use in this step to assign in the [Skeleton Editor](Engine/Animation/Persona/Modes/Skeleton). 
+
+
+## Steps
+
+1. Inside of the **Content Browser**, use the folder hierarchy and navigate to the **Mannequin** > **Mesh** folder. Select **SK_Mannequin** asset and double-click to open it. 
+
+	![](CB_Mannequin.png)(w:400)
+
+1. When the Skeleton Editor opens, navigate to the **Asset Details** panel and scroll down to locate the tab for **Lighting**. Click the selection box next to **Shadow Physics Asset** and select the Physics Asset we created in the previous step. 
+
+	[REGION:raw]
+	![](AssetDetailsLighting2.png)(w:950)
+	[/REGION]
+
+	[REGION:tip]
+	Be vigilant in naming your Physics Assets clearly because when you use the Shadow Physics Asset selection box, all Physics Assets for the project will be listed here.  
+	[/REGION]
+		
+1. You can now **Save** & **Close** the Skeleton Editor. 
+
+## End Result
+
+You have now assigned the Physics Asset that you created in the previous step to your character's Shadow Physics Asset slot. This will enable your character to use this Physics Asset for shadowing purposes.
+
+In the next step, we will add our character to the level, and then enable Capsule Shadowing.
+
+
+[Nav]
\ No newline at end of file
diff --git a/Engine/Documentation/Source/Engine/Rendering/LightingAndShadows/CapsuleShadows/QuickStart/4/CSQS_4.INT.udn b/Engine/Documentation/Source/Engine/Rendering/LightingAndShadows/CapsuleShadows/QuickStart/4/CSQS_4.INT.udn
new file mode 100644
--- /dev/null
+++ b/Engine/Documentation/Source/Engine/Rendering/LightingAndShadows/CapsuleShadows/QuickStart/4/CSQS_4.INT.udn
@@ -0,0 +1,87 @@
+Availability:docs
+Title:4 - Your Level - Enabling Capsule Shadows For Your Character
+Crumbs:
+Description:In this step we create a second shot for our example cinematic. 
+SkillLevel:Intermeidate
+Version:4.15
+checkpoint: capsuleshadowsqs
+parent:Engine/Rendering/LightingAndShadows/CapsuleShadows/QuickStart
+type:quick start
+order:4
+tags:Lighting
+ 
+
+[Nav]
+ 
+ In the previous step, you assigned the Physics Asset to your character in the Skeleton Editor as its Shadow Physics Asset, which, in this step, will enable you to 
+ use Capsule Shadowing for your character in directly and indirectly lit areas when enabled.
+
+## Steps
+
+1. Inside the **Content Browser**, use the folder hierarchy to navigate to the **ThirdPersonBP** > **Blueprints** folder. Select the **ThirdPersonCharacter** Blueprint, and double-click to open it. 
+
+	![](CB_TPBP.png)(w:400)
+
+1. When the Blueprint Editor opens, use the **Components** window, and select **Mesh (Inherited)**. Then, in the **Details** panel, scroll down and locate the **Lighting** tab. Under the advanced rollout, locate and 
+enable the following:
+
+	* **Capsule Direct Shadows**
+	* **Capsule Indirect Shadows**
+
+	[REGION:lightbox]
+	[![](CapsuleSettings.png)(w:950)](CapsuleSettings.png)
+	[/REGION]
+	[REGION:caption]
+	_Click image for full size._
+	[/REGION]
+
+1. With Capsule Shadows enabled for Direct and Indirect lighting, you can **Save** & **Close** the ThirdPersonCharacter Blueprint. 
+
+1. Before we can fully test Capsule Shadows, you'll need to create some indirectly lit areas in your level to use Capsule Indirect Shadows. From your Level Viewport, select the **Floor** Actor,
+and then hold down the **ALT** key while left-clicking to drag a copy up along the Z-axis (Blue). Move the floor Actor to the top of the outer walls, and then use the X-axis (Red) on the translation tool to move
+it back, so that only part of the floor is covered. This will give you an area that uses indirect lighting when you rebuild the lighting for the scene. 
+	
+	[OBJECT:EmbeddedVideo]
+		[PARAMLITERAL:width]
+		640
+		[/PARAMLITERAL]
+		[PARAMLITERAL:height]
+		360
+		[/PARAMLITERAL]
+		[PARAMLITERAL:videoid]
+		dQu0wKv_c_k
+		[/PARAMLITERAL]
+	[/OBJECT]
+	
+1. Now, select the **Direcitonal Light**, named "Light Source", from the **World Outliner**. In its **Details** panel, set the **Intensity** for the light to **10**. This will make sure there is enough light to illuminate the area being covered.
+
+	![](LightSourceSettings.png)
+
+1. From the Main Toolbar, click the **Build** button so that lighting will be rebuilt for the scene. This ensures that you will be able to use Capsule Indirect Shadows for your character.
+
+	[REGION:raw]
+	![](BuildButton.png)
+	[/REGION]
+
+## End Result
+
+You now have a character that supports Capsule Shadowing when in direct and indirect lighting. You can **Play in Editor** (PIE) to move your character around the scene to see how the soft shadows look. 
+
+[OBJECT:EmbeddedVideo]
+	[PARAMLITERAL:width]
+	640
+	[/PARAMLITERAL]
+	[PARAMLITERAL:height]
+	360
+	[/PARAMLITERAL]
+	[PARAMLITERAL:videoid]
+	1oOPsSMTA6g
+	[/PARAMLITERAL]
+[/OBJECT]
+
+You can enable both Capsule Shadow settings individually, such that, if you don't need the soft shadowing while your character is lit in open areas, you can disable this
+option and leave only Capsule Indirect Shadow enabled for these areas that are indirectly lit. 
+
+In the final step, there are some challenges you can attempt on your own, along with reminders for creating a Shadow Physics Asset. 
+
+[Nav]
\ No newline at end of file
diff --git a/Engine/Documentation/Source/Engine/Rendering/LightingAndShadows/CapsuleShadows/QuickStart/5/CSQS_5.INT.udn b/Engine/Documentation/Source/Engine/Rendering/LightingAndShadows/CapsuleShadows/QuickStart/5/CSQS_5.INT.udn
new file mode 100644
--- /dev/null
+++ b/Engine/Documentation/Source/Engine/Rendering/LightingAndShadows/CapsuleShadows/QuickStart/5/CSQS_5.INT.udn
@@ -0,0 +1,29 @@
+Availability:Docs
+Title:5 - On Your Own!
+Crumbs:
+Description: This step offers some additional things to try with your sample scene and links to more documentation.
+SkillLevel:Intermediate
+Version:4.15
+checkpoint: capsuleshadowsqs
+parent:Engine/Rendering/LightingAndShadows/CapsuleShadows/QuickStart
+type:quick start
+order:5
+tags:Lighting
+ 
+[Nav]
+
+Now that you have created a sample Shadow Physics Asset for Capsule Shadowing, you can try some of the following on your own: 
+
+* Try using different lighting angles, or time-of-day scenarios, for mid-day, dusk, or evening. 
+* Adjust the **Light Source Angle** of your Directional Light to soften your Capsule Shadow in areas that are directly lit so that your character blends in with surrounding shadows. 
+* Add some closed off areas with doorways and windows (or something similar) to a cave opening, to get contrasting light that gives directionality for the light source when **Capsule Indirect Shadow** is enabled.
+* Use the setting for **Capsule Indirect Shadow Min Visibility** in your Character's Details panel to blend the soft shadow's intensity with your environment.
+* Add some additional capsule bodies for the arms of your character for shadowing when it's near wall surfaces, like with a cover system, or if your character goes prone.
+
+For more information on topics covered in this Quick Start guide, see the [Unreal Enditor Manual](https://docs.unrealengine.com/latest/INT/Engine/index.html).
+
+* For information covering Capsule Shadows, see [](Engine/Rendering/LightingAndShadows/CapsuleShadows/Overview).
+* For information covering the Physics Asset Tool, see [Physics Asset Tool (PhAT)](https://docs.unrealengine.com/latest/INT/Engine/Physics/PhAT/).
+* For information covering Lighting, see [Lighting the Environment](https://docs.unrealengine.com/latest/INT/Engine/Rendering/LightingAndShadows/).
+
+[Nav]
diff --git a/Engine/Documentation/Source/Engine/Rendering/LightingAndShadows/CapsuleShadows/QuickStart/CSQuickStart.INT.udn b/Engine/Documentation/Source/Engine/Rendering/LightingAndShadows/CapsuleShadows/QuickStart/CSQuickStart.INT.udn
new file mode 100644
--- /dev/null
+++ b/Engine/Documentation/Source/Engine/Rendering/LightingAndShadows/CapsuleShadows/QuickStart/CSQuickStart.INT.udn
@@ -0,0 +1,49 @@
+Availability:Docs
+Title: Capsule Shadows Quick Start
+Crumbs:
+Description:A guide to quickly get up and running with the Sequencer Editor in Unreal Engine 4. 
+parent:Engine/Rendering/LightingAndShadows/CapsuleShadows
+order:1
+Version: 4.15
+type:quick start
+checkpoint:capsuleshadowsqs
+tags: Lighting
+topic-image:QS_topic_image.png
+SkillLevel:Intermediate
+
+[REGION:imagetable]
+![](QS_CapsuleShadow.png)(w:700)  
+At the end of this guide, you'll have a scene similar to the one above.
+[/REGION]
+
+## Goals
+
+The Capsule Shadows Quick Start intends to quickly get you up and running with the basics of creating a capsule representation of your character with a Physics Asset and enabling Capsule Shadows for character in 
+your games differently lit environments. 
+
+## Objectives
+
+After going through this tutorial, developers will know the following: 
+
+* How to create a Physics Asset for shadows.
+* How to assign the Shadow Physics Asset to your Skeletal Mesh.
+* How to enable Capsule Shadows in your level for your characters.
+
+
+[OBJECT:TopicList]
+	[PARAM:icon]
+		![](%ROOT%/build_icon.png)(convert:false)
+	[/PARAM]
+	[PARAM:title]
+		Steps
+	[/PARAM]
+	[PARAM:description]
+	[/PARAM]
+	[PARAM:links]
+	[DIR(output:"checkpoint" parent:"Engine/Rendering/LightingAndShadows/CapsuleShadows/QuickStart" org:"hierarchy")]
+	[/PARAM]
+[/OBJECT]
+
+[REGION:call_to_action]
+	[Click to Start](Engine/Rendering/LightingAndShadows/CapsuleShadows/QuickStart/1 "%Engine/Rendering/LightingAndShadows/CapsuleShadows/QuickStart/1:title%")
+[/REGION]
diff --git a/Engine/Documentation/Source/Engine/Rendering/LightingAndShadows/ContactShadows/ContactShadows.INT.udn b/Engine/Documentation/Source/Engine/Rendering/LightingAndShadows/ContactShadows/ContactShadows.INT.udn
--- a/Engine/Documentation/Source/Engine/Rendering/LightingAndShadows/ContactShadows/ContactShadows.INT.udn
+++ b/Engine/Documentation/Source/Engine/Rendering/LightingAndShadows/ContactShadows/ContactShadows.INT.udn
@@ -1,4 +1,4 @@
-Availability: Docs
+Availability: Public
 Title: Contact Shadows
 Crumbs: 
 Description: This page provides an overview of Contact Shadows in UE4.
@@ -9,21 +9,7 @@ Tags: Rendering
 Type: Reference
 reference-image: ContactShadow.gif
 Related: Engine/Rendering/LightingAndShadows/LightTypes/Point
-
-[VAR:Params]
-[PARAM:image]
-	![%Engine/Rendering/LightingAndShadows/ContactShadows:title%](ContactShadow_Topic.png)
-[/PARAM]
-[PARAM:icon]
-	![](%ROOT%/tutorial_icon.png)(convert:false)
-[/PARAM]
-[PARAM:title]
-	%Engine/Rendering/LightingAndShadows/ContactShadows:title%
-[/PARAM]
-[PARAM:path]
-	[RELATIVE:Engine/Rendering/LightingAndShadows/ContactShadows]
-[/PARAM]
-[/VAR]
+topic-image:ContactShadow_Topic.png
 
 [TOC(start:1 end:4)]
 
diff --git a/Engine/Documentation/Source/Engine/Rendering/LightingAndShadows/DistanceFields/DistanceFieldsLanding.INT.udn b/Engine/Documentation/Source/Engine/Rendering/LightingAndShadows/DistanceFields/DistanceFieldsLanding.INT.udn
new file mode 100644
--- /dev/null
+++ b/Engine/Documentation/Source/Engine/Rendering/LightingAndShadows/DistanceFields/DistanceFieldsLanding.INT.udn
@@ -0,0 +1,70 @@
+Availability: Docs
+Title: Distance Fields
+Crumbs:
+Description: Developing your games using Mesh Distance Field features.
+Version: 4.15
+type:landing
+parent:Engine/Rendering/LightingAndShadows
+order:
+tags:Lighting
+
+
+[REGION:fullwidth raw]
+![](DistanceFieldsBanner.png)
+[/REGION]
+
+
+**Distance Fields** are a powerful part of Unreal Engine 4 and are the basis for major features, including features heavily used in Epic's own games, such as dynamic Ambient Occlusion and Shadows. Distance Fields have also been extended for use with particle collision, dynamic indirect area shadows, and Materials, enabling you to create dynamic effects for soft bodies or dynamic flow maps for water.
+
+Before you can get started using Distance Fields features, you'll need to enable it for your project.
+
+## Enabling Mesh Distance Fields for your Project
+
+1. Navigate to the menu and select **Edit** > **Project Settings** to open the project options window.
+
+	![](ProjectSettingsMenu.png)(w:250)
+	
+1. In the **Project Settings** window, locate and select the **Rendering** section and then, under the **Lighting** category, toggle the option for **Generate Mesh Distance Fields**.
+
+	[REGION:lightbox]
+	[![](ProjectSettingsWindow.png)(w:750)](ProjectSettingsWindow.png) 
+	[/REGION]
+	[REGION:caption]
+	Click image for full size view.
+	[/REGION]
+	Once you toggle this setting, you will be prompted to Restart the Project by a pop-up in the bottom corner of the Editor window. Click the **Restart Now** button for the changes to take effect.
+
+	![](RestartEditorButton.png)(w:400)
+
+1. Once the Editor restarts, you will see a pop-up message indicating that Distance Fields are being built.	
+	
+	![](BuildingMeshDFs.png)(w:400)
+
+	[REGION:note]
+	Mesh Distance Fields are only built if a Static Mesh exists in the current loaded level, when a Static Mesh is opened the Static Mesh Editor for the first time, or if the Distance Field Resolution is changed via the Build Settings in the Static Mesh Editor. 
+	[/REGION]
+
+1. Once Distance Fields are built, you can use the visualization mode to see their representation by going to **Show** > **Visualize** > **Mesh Distance Fields** in the main viewport.  
+
+	![](VisualizeMeshDF.png)(w:450)
+
+
+## Reference
+
+%Engine/Rendering/LightingAndShadows/DistanceFieldAmbientOcclusion:topic%
+%Engine/Rendering/LightingAndShadows/RayTracedDistanceFieldShadowing:topic%
+
+
+## Essentials
+
+%Engine/Rendering/LightingAndShadows/DistanceField/HowTo:topic%
+
+<!---
+%Engine/Rendering/LightingAndShadows/DistanceField/DFHT_1:topic%
+%Engine/Rendering/LightingAndShadows/DistanceField/DFHT_2:topic%
+%Engine/Rendering/LightingAndShadows/DistanceField/DFHT_3:topic%
+%Engine/Rendering/LightingAndShadows/DistanceField/DFHT_4:topic%
+%Engine/Rendering/LightingAndShadows/DistanceField/DFHT_5:topic%
+--->
+
+
diff --git a/Engine/Documentation/Source/Engine/Rendering/LightingAndShadows/DistanceFields/HowTo/DFHT_1/DFHT_1.INT.udn b/Engine/Documentation/Source/Engine/Rendering/LightingAndShadows/DistanceFields/HowTo/DFHT_1/DFHT_1.INT.udn
new file mode 100644
--- /dev/null
+++ b/Engine/Documentation/Source/Engine/Rendering/LightingAndShadows/DistanceFields/HowTo/DFHT_1/DFHT_1.INT.udn
@@ -0,0 +1,199 @@
+Availability:Docs
+Crumbs:
+Title:Using RayTraced Distance Field Soft Shadows
+Description:How to set up and use RayTraced Distance Field Soft Shadows.
+Related:Engine/Rendering/LightingandShadows/DistanceFieldAmbientOcclusion
+Related:Engine/Rendering/LightingandShadows/RayTracedDistanceFieldShadowing
+SkillLevel: Intermediate
+Type:how-to
+Parent: Engine/Rendering/LightingAndShadows/DistanceFields/HowTo
+Order:
+Tags: Lighting
+Version: 4.15
+
+
+[REGION:fullwidth raw]
+![](RTDFBanner.png)
+[/REGION]
+
+[TOC(start:2 end:2)] 
+
+
+
+When you're developing your game, you may find that there are situations where dynamic lighting is the best choice for your scenes, like with large open worlds where precomputed lighting would be inefficient or too demanding on the target system's resources. **Ray Traced Distance Field Soft Shadows**, also referred to simply as **Distance Field Shadows**, enables you to have shadowing at farther distances than traditional **Cascaded Shadow Maps** (CSM) with a **Directional Light** would. For open world games, or long view distances in general, Distance Field Shadows can efficiently shadow beyond the CSM distance. This method was used extensively in the open world demo "A Boy and His Kite", which is available for download from the Epic Games Launcher via the Learn tab.
+
+In this guide, we'll take a look at how to enable Distance Field Shadows for your light types and then explore the various settings that you can use with them.
+
+
+## Steps
+
+[REGION:note]
+This feature requires that **Generate Mesh Distance Fields** be enabled in your **Project Settings** in the **Rendering** section.
+[/REGION]
+
+1. Start by navigating to the **Modes** window, then in the **Lights** section, select and drag a **Directional**, **Spot**, or **Point** Light into the main viewport. 
+
+	![](ModesPanelLights.png)(w:400)
+	
+	[REGION:note]
+	For the purpose of enabling Distance Field Shadows on any Light Actor, the process is the same. The additional sections of this guide will cover specific properties for these light types.
+	[/REGION]
+
+1. With the Light Actor selected, navigate over to its **Details** panel and set its **Mobility** to **Movable**. 
+	
+	![](TransformMobility.png)(w:450)
+
+1. Now, scroll down in the **Details** panel and under the **Distance Field Shadows** section, toggle the option for **RayTraced Distance Field Shadows**.
+
+	[Region:imagetable]
+	|![](DirectionalLightDFSettings.png)(w:450)|![](PointSpotLightSettingsDF.png)(w:450)|
+	|---|---|
+	| Directional Light | Spot/Point Light |
+	[/Region]
+	
+	[REGION:note]
+	If this option if grayed out, make sure that the light is set to **Movable** or that you've first enabled the option for **Generate Mesh Distance Fields** in your Project's Settings.
+	[/REGION]
+	
+1. 	Now, Distance Field Shadows will be enabled for your lights that have this setting toggled on. Try dragging in a few Static Meshes from the Content Browser to see Distance Field Shadows working in your game.
+
+	[REGION:tip]
+	When using a Directional Light, it can be helpful to set the **Dynamic Shadow Distance Moveable Light** slider to **0**. This will disable CSM shadows so that you will only see the Distance Field Shadow of a mesh. This is also useful for testing your scene and diagnosing any potential Distance Field issues in addition to using the visualization mode for Mesh Distance Fields that you can enable in the main viewport by going to **Show** > **Visualize** > **Mesh Distance Fields**.
+	[/REGION]
+
+	
+## Light Source Settings
+
+When you decide on the type of light you want to use, whether it's a Directional, Spot, or Point Light, you will have different options that are avaialble. Below you'll find these settings that you can adjust and how to you use them for your Distance Field Shadows.
+
+### Directional Light
+
+Here you'll find the various settings that are available for Directional Lights that have Distance Field Shadows enabled.
+
+![Distance Field Shadow Settings](DirectionalLightSettings.png)
+
+[REGION:simpletable]
+| Property | Description |
+|----------|-------------|
+| [REGION:tablesection]Light[/REGION] ||
+| Light Source Angle | This is the light source's angle in degrees and is used to support soft area shadows for dynamic shadowing methods using Distance Fields or Capsule Shadows. | 
+| [REGION:tablesection]Distance Field Shadows[/REGION] ||
+| DistanceField Shadow Distance | This is the farthest distance that will have Distance Field Shadowing. Distance Field Shadows will also cover the distance between this value and the **Dynamic Shadow Distance Movable Light** for CSM shadows. | 
+| RayTraced DistanceField Shadows | This will toggle Distance Field Shadows on for the light source. | 
+| Distance Field Trace Distance  | This sets a distance in world units for how far shadows can cast from their shadow caster. When using larger values you will increase the shadowing cost for the scene. |
+| Ray Start Offset Depth Scale | This controls how large an offset ray traced shadows have from the receiving surface as the camera gets further away. This can be used to hide self-shadowing artifacts from low-resolution distance fields on huge static meshes. | 
+[/REGION]
+
+#### Cascaded Shadow Maps VS Distance Field Shadows
+
+Cascaded Shadow Maps provide higher quality shadowing but doesn't work well for long view distances, whereas Distance Fields can shadow over much farther distance more effeciently but rely on the resolution of their
+Mesh Distance Field for quality. It's a good to use both in combination with one another to take advantage of the higher resolution CSM shadows at closer distances to your camera view and then use Distance Fields for 
+the distant shadowing. 
+
+[OBJECT:ComparisonSlider]
+	[PARAM:before]
+	![Cascaded Shadow Maps](Directional_CSM.png)(w:600)
+	[/PARAM]
+	[PARAM:after]
+	![Distance Field Shadows](Directional_DF.png)(w:600)
+	[/PARAM]
+[/OBJECT]
+
+#### Light Source Angle
+
+The **Light Source Angle** of a Directional Light will soften or sharpen the shadows based on the rotational angle of the light and the value entered. This setting only works with Distance Fields and Capsule Shadows and will not work with traditional shadow mapping where the shadow cannot easily be blurred.
+
+The Light Source Angle has been set to 0, which you'll see that the shadows become realy sharp and crisp, even at father distances, while when using a higher value, the shadows will soften and spread out as they extend along the receiving surface.
+
+[OBJECT:ComparisonSlider]
+	[PARAM:before]
+	![Light Source Angle: 0](DF_LSA0.png)(w:600)
+	[/PARAM]
+	[PARAM:after]
+	![Light Source Angle: 1](DF_LSA1.png)(w:600)
+	[/PARAM]
+[/OBJECT]
+
+
+#### Distance Field Trace Distance
+
+The **Distance Field Trace Distance** controls how far a ray trace for shadowing can extend for any shadow caster. If you lower the value for the trace distance, you lower the cost of shadow casting for the scene 
+because shadows will not draw beyond that point from their source. Keep in mind that the trace distance does not affect the shadow distance from the camera's view, which is instead controlled by **Distance Field Shadow Distance**.
+
+The tree, in this example, is casting a shadow along the ground using its Mesh Distance Field. A trace distance of 1500 units (the default value is 10000), the tree's shadow doesn't capture the upper most branches and leaves since the shadows that would be cast from there are farther than 1500 units from the receiving surface.
+
+[OBJECT:ComparisonSlider]
+	[PARAM:before]
+	![Distance Field Trace Distance: 10000](RayStart1.png)(w:600)
+	[/PARAM]
+	[PARAM:after]
+	![Distance Field Trace Distance: 1500](RayStart2.png)(w:600)
+	[/PARAM]
+[/OBJECT]
+
+### Point/Spot Lights
+ 
+Here you'll find the various settings that are available for your **Spot** and **Point** Lights that have Distance Field Shadows enabled.
+ 
+![](PointSpotLightSettings.png)
+ 
+[REGION:simpletable]
+| Property | Description |
+|----------|-------------| 
+| [REGION:tablesection]Light[/REGION] ||
+| Source Radius | This is used as the size of the light source sphere. When you use larger values you get a larger penumbra, but it will cost more performance. |
+| [REGION:tablesection]Distance Field Shadows[/REGION] ||
+| RayTraced DistanceField Shadows | Whether to use ray traced distance field area shadows. | 
+| Ray Start Offset Depth Scale | Controls how large of an offset ray traced shadows have from the receiving surface as the camera gets further away. This can be useful to hide self-shadowing artifacts from low-resolution distance fields on huge static meshes. | 
+[/REGION]
+
+#### Traditional Shadow Maps VS Distance Field Shadows
+
+Similarly to CSM shadows, the shadow maps for Point and Spot lights have very detailed high quality shadows versus Distance Field Shadows, which can use the **Source Radius** of a light to control how soft or sharp the shadows are.
+
+[OBJECT:ComparisonSlider]
+	[PARAM:before]
+	![Default Shadow Map](PointLight_CSM.png)(w:600)
+	[/PARAM]
+	[PARAM:after]
+	![Distance Field Shadow](PointLight_DF.png)(w:600)
+	[/PARAM]
+[/OBJECT]
+
+### Source Radius
+
+The **Source Radius**, similar to the Light Source Angle in a Directional Light, is used to soften or sharpen the area shadows by adjusting the size of the light source itself. 
+
+When using a lower value the shadows will become sharp throughout the light's attenuation radius, whereas high values will soften the shadows as they extend out from the shadow caster.
+
+[OBJECT:ComparisonSlider]
+	[PARAM:before]
+	![Source Radius: 0.1](DF_LightSource1.png)(w:600)
+	[/PARAM]
+	[PARAM:after]
+	![Source Radius: 20](DF_LightSource2.png)(w:600)
+	[/PARAM]
+[/OBJECT]
+
+### Ray Start Offset Depth Scale
+
+In some cases, low resolution Distance Fields or more complex geometry will produce self-shadowing artifacts. This can sometimes be offset by adjusting the **Distance Field Resolution** of your mesh, but you may not want to pay 
+the higher cost of the generated volume texture that would be required. Instead, you can use the **Ray Start Offset Depth Scale** to control where the ray starts for shadows from the receiving surface as the camera moves further away. 
+
+[OBJECT:ComparisonSlider]
+	[PARAM:before]
+	![Ray Start Offset Depth Scale: 0](DF_Depth0.png)(w:600)
+	[/PARAM]
+	[PARAM:after]
+	![Ray Start Offset Depth Scale: 0.01](DF_Depth1.png)(w:600)
+	[/PARAM]
+[/OBJECT]
+
+
+## End Result
+
+Now that you've enabled Distance Fields Shadows for your different light types in your level and seen how to adjust the various settings for them, you should be able to apply these concepts to your own levels and games
+using Distance Fields. 
+
+
+
diff --git a/Engine/Documentation/Source/Engine/Rendering/LightingAndShadows/DistanceFields/HowTo/DFHT_2/DFHT_2.INT.udn b/Engine/Documentation/Source/Engine/Rendering/LightingAndShadows/DistanceFields/HowTo/DFHT_2/DFHT_2.INT.udn
new file mode 100644
--- /dev/null
+++ b/Engine/Documentation/Source/Engine/Rendering/LightingAndShadows/DistanceFields/HowTo/DFHT_2/DFHT_2.INT.udn
@@ -0,0 +1,156 @@
+Availability:Docs
+Crumbs:
+Title:Using Distance Field Ambient Occlusion
+Description:How to set up and use Distance Field Ambient Occlusion.
+Related:Engine/Rendering/LightingandShadows/DistanceFieldAmbientOcclusion
+Related:Engine/Rendering/LightingandShadows/RayTracedDistanceFieldShadowing
+type: how-to
+Parent: Engine/Rendering/LightingAndShadows/DistanceFields/HowTo
+order:
+SkillLevel:Intermediate
+Version: 4.15
+tags: Lighting
+
+[REGION:fullwidth raw]
+![](DFAOBanner.png)
+[/REGION]
+
+[TOC(start:2 end:2)] 
+
+
+In your games, you may rely on screen space techniques to provide **Ambient Occlusion** (AO) or even precomputed lighting to make objects feel grounded in the world. These techniques are very useful, but have their
+own limitations. **Screen Space Ambient Occlusion** (SSAO) is limited to using scene depth, which is limited to what's visible on screen, and precomputed methods only work for static objects in the world.
+**Distance Field Ambient Occlusion** (DFAO) is a fully dynamic AO method using Mesh Distance Fields for your Static Mesh Actors and is not limited to only being used in dynamically lit worlds; it can easily be 
+used with precomputed lighting as well. For additional information, you can read about [Distance Field Ambient Occlusion](Engine/Rendering/LightingandShadows/DistanceFieldAmbientOcclusion) in more detail.
+
+In this guide, we'll take a look at how to enable DFAO with a Movable Sky Light Actor and then explore the various settings that you can use.
+
+
+## Steps
+
+[REGION:note]
+This feature requires that **Generate Mesh Distance Fields** be enabled in your **Project Settings** in the **Rendering** section.
+[/REGION]
+
+1. Start by navigating to the **Modes** window, then in the **Lights** section, select and drag a **Sky Light** Actor into your Level Viewport.
+	
+	![](AddSkylight.png)(w:750)
+	
+1. With the Sky Light selected, navigate over to its **Details** panel and set its **Mobility** to **Movable**.
+
+	![](TransformMobility.png)(w:450)
+	
+1. Now, Distance Field Ambient Occlusion will automatically be enabled for your scene. You can find the settings for **Distance Field Ambient Occlusion** in the Skylight's **Details** panel.  
+
+	[OBJECT:ComparisonSlider]
+		[PARAM:before]
+		![No Distance Field Ambient Occlusion](1_NoDFAO1.png)(w:700)
+		[/PARAM]
+		[PARAM:after]
+		![Distance Field Ambient Occlusion](1_DFAO1.png)(w:700)
+		[/PARAM]
+	[/OBJECT]
+
+ 
+## Visualizing Distance Field Ambient Occlusion
+
+To help you better see DFAO in action in your scene, and to help you troubleshoot or diagnose issues, you can use the visualization view mode by going to the main viewport and selecting **Show** > **Visualize** > **Distance Field Ambient Occlusion**.
+Like other view modes, your viewport will change so that you can see the DFAO representation of your level.
+
+[OBJECT:ComparisonSlider]
+	[PARAM:before]
+	![Distance Field Ambient Occlusion Scene View](1_DFAO1.png)(w:700)
+	[/PARAM]
+	[PARAM:after]
+	![Distance Field Ambient Occlusion Visualization](1_DFAOVis1.png)(w:700)
+	[/PARAM]
+[/OBJECT]
+
+## Distance Field Ambient Occlusion Settings
+
+In the Sky Light Details panel, under **Distance Field Ambient Occlusion**, you'll find the settings available for artistic control of DFAO in your level.
+ 
+ ![](DFAOProperties.png)(w:450)
+ 
+| Property | Description |
+|----------|-------------|
+| Occlusion Max Distance | This is the maximum distance at which the occlusion of one point will affect another. By using higher values, you increase the cost of DFAO exponentially, but, in return, you get occlusion at a farther distance. |
+| Contrast | This applies a contrast S-curve to the computed AO. A value of 0 means no contrast increase, whereas 1 is a significant contrast increase. |
+| Min Occlusion | This controls the darkest that a fully occluded area can get. |
+| Occlusion Tint | Artists can use this to control the Tint color for occluded areas. |
+
+### Occlusion Max Distance
+
+The **Occlusion Max Distance** controls the maximum distance between two points where occlusion of one can affect the other. When you adjust this distance, you're adjusting the accuracy that DFAO will have with other objects in the scene, which in turn means you're increasing the cost of it as well.
+
+Using the two bottom-end and top-end values for the Occlusion Max Distance, a value of 200 units doesn't provide much detailed ambient occlusion, where a value of 1500 units is has a lot of ambient occlusion and looks more 
+accurate for the scene but comes at a higher cost. 
+
+[OBJECT:ComparisonSlider]
+	[PARAM:before]
+	![Occlusion Max Distance - 200](1_MaxOccl200.png)(w:700)
+	[/PARAM]
+	[PARAM:after]
+	![Occlusion Max Distance - 1500](1_MaxOccl1500.png)(w:700)
+	[/PARAM]
+[/OBJECT]
+
+### Contrast
+
+The **Contrast** controls the difference in brightness between the objects affected by DFAO. The contrast value is clamped to be between 0 and 1, so you can use the slider to adjust the amount of contrast you will have in your scene.
+
+Increasing the value for Contrast darkens a lot of the crevices and areas that are heavily occluded in this scene.
+
+[OBJECT:ComparisonSlider]
+	[PARAM:before]
+	![Contrast - 0](1_Contrast_0.png)(w:700)
+	[/PARAM]
+	[PARAM:after]
+	![Contrast - 1](1_Contrast_1.png)(w:700)
+	[/PARAM]
+[/OBJECT]
+
+### Min Occlusion
+
+The **Min Occlusion** controls how dark any fully occluded areas can get in your scene. This setting gives extra control to Artists to lighten DFAO in the occluded areas by using a higher value between 0 and 1.
+
+Here the difference between a value of 0 and 1 is significant demonstrating how the occlusion is lightened for the scene, however, this setting works really well when also used with Contrast where you can darken 
+occluded areas and then use the Min Occlusion to lighten how dark these areas can really get.
+
+[OBJECT:ComparisonSlider]
+	[PARAM:before]
+	![Min Occlusion - 0](1_MinOccl_0.png)(w:700)
+	[/PARAM]
+	[PARAM:after]
+	![Min Occlusion - 1](1_MinOccl_1.png)(w:700)
+	[/PARAM]
+[/OBJECT]
+
+
+### Occlusion Tint
+
+The **Occlusion Tint** affords Artists the most control over the look of DFAO by enabling the tint color for occluded areas to be adjusted. 
+
+Occluded areas will be tinted based on these color values, which you can see here the difference between an Occlusion Tint that is black (the default color value) and red.
+
+[OBJECT:ComparisonSlider]
+	[PARAM:before]
+	![Tint Color - Black](1_OccTint_0.png)(w:700)
+	[/PARAM]
+	[PARAM:after]
+	![Tint Color - Red](1_OccTint_1.png)(w:700)
+	[/PARAM]
+[/OBJECT]
+
+
+## End Result
+
+Now that you've enabled Distance Field Ambient Occlusion for your level and seen how to adjust the various settings for it in your Sky Light, you should be able to control and tweak DFAO to fit your own 
+games and levels for your stylistic needs. 
+
+
+
+
+
+
+
diff --git a/Engine/Documentation/Source/Engine/Rendering/LightingAndShadows/DistanceFields/HowTo/DFHT_3/DFHT_3.INT.udn b/Engine/Documentation/Source/Engine/Rendering/LightingAndShadows/DistanceFields/HowTo/DFHT_3/DFHT_3.INT.udn
new file mode 100644
--- /dev/null
+++ b/Engine/Documentation/Source/Engine/Rendering/LightingAndShadows/DistanceFields/HowTo/DFHT_3/DFHT_3.INT.udn
@@ -0,0 +1,91 @@
+Availability: Docs
+Crumbs:
+Title: Using Particle Collision Mode for Distance Fields
+Description: A how to guide to enable your Particle Systems to use Distance Fields for scene collision.
+Related: Engine/Rendering/LightingAndShadows/RayTracedDistanceFieldShadowing
+Related: Engine/Rendering/LightingAndShadows/DistanceFieldAmbientOcclusion
+SkillLevel:Intermediate
+Version: 4.15
+reference-image:SceneDepthCollision.gif
+reference-image:DistanceFieldCollision.gif
+Parent: Engine/Rendering/LightingAndShadows/DistanceFields/HowTo
+type: how-to
+Order:
+Tags: Lighting
+
+
+[REGION:fullwidth raw]
+![](DFParticleBanner.png)
+[/REGION]
+
+Particle collision in your games can be limited when relying on Scene Depth, which can only use what is in the camera view. At some angles, collision isn't very accurate for particles that 
+have longer lifetimes or that are used over a large area. You can now use **Distance Fields** for the particle's **Collision Mode**, which enables your visual effects (VFX) to use whole-scene collision 
+using your Static Mesh's Distance Field representation. This collision mode offers more reliable collision than Scene Depth, however, the 
+[Global Distance Field](https://docs.unrealengine.com/latest/INT/Engine/Rendering/LightingAndShadows/DistanceFieldAmbientOcclusion/#technicaldetails) is lower resolution, so you may find that 
+particles pass through thin objects, and sharp corners are rounded. In return for this lower resolution, collision with Distance Fields runs very efficiently; its performance costs about the 
+same as using Scene Depth.
+
+In this guide, we'll look at how we can enable the **Collision Mode** in a Particle System to use **Distance Fields** over Scene Depth.
+
+## Steps
+
+[REGION:note]
+This feature requires that **Generate Mesh Distance Fields** be enabled in your **Project Settings** in the **Rendering** section.
+[/REGION]
+
+1. In your **Content Browser**, locate and open any **Particle System** or create a new one if you do not have one readily available. 
+	
+	![](ContentBrowser.png)(w:350)
+
+	[REGION:tip]
+	Content Examples provides a number of ready-made Particle Systems that would useful. The **P_GPU**Particles_Collision2** asset from [Content Examples](Resources/ContentExamples) is a good example to use as a starting point with minimal setup.
+	[/REGION]	
+	
+1. Once **Cascade** (the Particle System Editor) opens, right-click in the empty space and from the module list, select **Type Data** > **New GPU Sprites**.
+	
+	[REGION:raw]
+	![](AddGPUModule.png)(w:700)
+	[/REGION]
+	
+1. Now, right-click, again, in the empty space and from the module list, select **Collision** > **Collision(Scene Depth)**. 
+
+	[REGION:raw]
+	![](AddCollisionModule.png)(w:700)
+	[/REGION]
+	
+1. With the **Collision** module selected, navigate over to the **Details** panel and next to **Collision Mode** use the dropdown to select **Distance Field**.
+
+	![](DetailsPanelCollisionMode.png)(w:350)
+	
+	[REGION:note]
+	Distance Field particle collision is a Shader Model 5 (SM5) feature, so on hardware that doesn't support SM5, the particles will fall back to use Scene Depth Collision instead.
+	[/REGION]
+
+## End Result
+
+Now that you've seen how to enable your Particle System(s) to use the Collision Mode for Distance Fields, your particles will be able to collide with the entire scene without relying on Scene Depth. This affords you the flexibility in the types of effects and lifetimes of those effects that you create, whether they cover a large area or are localized to a small one.
+
+
+[REGION:imagetable]
+| [INCLUDE:#SceneDepth] | [INCLUDE:#DistanceField] |
+| --- | --- |
+| Scene Depth Collision | Distance Field Collision|
+[/REGION]
+
+In these examples, we'll take a look at how **Scene Depth** collision can have varying results depending on your camera view compared to **Distance Field** collision, which doesn't rely on the camera view. This is especially noticeable when moving under an enclosed area, like the tree in this example, where the parts of the tree that are out of view no longer provide collision for the area particle.
+
+
+
+<!---
+[EXCERPT:SceneDepth]
+[REGION:asyncgif]
+![](SceneDepthCollision.gif)(w:600)
+[/REGION]
+[/EXCERPT:SceneDepth]
+
+[EXCERPT:DistanceField]
+[REGION:asyncgif]
+![](DistanceFieldCollision.gif)(w:600)
+[/REGION]
+[/EXCERPT:DistanceField]
+--->
\ No newline at end of file
diff --git a/Engine/Documentation/Source/Engine/Rendering/LightingAndShadows/DistanceFields/HowTo/DFHT_4/DFHT_4.INT.udn b/Engine/Documentation/Source/Engine/Rendering/LightingAndShadows/DistanceFields/HowTo/DFHT_4/DFHT_4.INT.udn
new file mode 100644
--- /dev/null
+++ b/Engine/Documentation/Source/Engine/Rendering/LightingAndShadows/DistanceFields/HowTo/DFHT_4/DFHT_4.INT.udn
@@ -0,0 +1,92 @@
+Availability: Docs
+Crumbs:
+Title:Using Distance Field Indirect Shadows
+Description: How to setup and use Distance Field Indirect Shadows. 
+Related: Engine/Rendering/LightingAndShadows/RayTracedDistanceFieldShadowing
+Related: Engine/Rendering/LightingAndShadows/DistanceFieldAmbientOcclusion
+Parent: Engine/Rendering/LightingAndShadows/DistanceFields/HowTo
+Order:
+type: how-to
+SkillLevel:Intermediate
+Version: 4.15
+reference-image:DFIndirect.gif
+Tags: Lighting
+
+[REGION:fullwidth raw]
+![](DFISBanner.png)
+[/REGION]
+
+
+
+When your games rely on precomputed lighting with indirectly lit areas, it can be a challenge to have nice soft shadowing for your dynamic objects without faking it by using decals or complex materials. 
+To solve this problem, **Distance Field Indirect Shadows** (DFIS) enables your dynamic Static Meshes to have soft area shadows in these areas where traditional shadow mapping techniques 
+wouldn't work. Distance Field Indirect Shadows, similarly to how [Capsule Shadows](Engine/Rendering/LightingAndShadows/CapsuleShadows/Overview) for Skeletal Meshes samples the 
+[indirect lighting cache](Engine/Rendering/LightingAndShadows/IndirectLightingCache/) to get the directionality and intensity of the indirect lighting to cast a soft area shadow using the Mesh Distance Field.
+
+In this guide, we'll take a look at how to enable Distance Fields for meshes that can then be used in your levels for dynamic indirect shadowing.
+
+## Steps
+
+[REGION:note]
+Unlike other Distance Field features, like those for [Ray Traced Distance Field Shadows](Engine/Rendering/LightingAndShadows/RayTracedDistanceFieldShadowing) and [Ambient Occlusion](Engine/Rendering/LightingAndShadows/DistanceFieldAmbientOcclusion), this feature can be enabled per-mesh instead of project-wide, which is detailed in the steps below.
+[/REGION]
+
+1. In the **Content Browser**, start by selecting any **Static Mesh** asset and double-clicking to open the Static Mesh Editor. 
+
+	![](ContentBrowser.png)(w:350)
+	
+	[REGION:note]
+	For the purpose of this example, the Static Mesh **SM_Statue** from Sun Temple project (available on the Epic Games Launcher) will be used for the remainder of this demonstration.
+	[/REGION]
+1. In the Static Mesh Editor, navigate to the **Details** panel. In the **Static Mesh Settings** section, toggle the option for **Generate Mesh Distance Fields**. After enabling this, you can save and close the Editor window.
+
+	[REGION:raw]
+	![](EnableDistanceField.png)(w:1050)
+	[/REGION]
+1. Now, go back to the **Content Browser**, where you selected the asset, and drag it into the Level Viewport. 
+
+	![](DragToWorld.png)(w:450)
+1. With the Actor selected in your level, go to the **Details** panel and set the **Mobility** to **Movable**. 
+
+	![](TransformMobility.png)(w:450)
+	
+	Then, under the **Lighting** section, toggle the option for **Distance Field Indirect Shadow**.
+	
+	![](EnableDFIS.png)(w:450)
+	
+1.	Once you've toggled the **Distance Field Indirect Shadows** option, your Static Mesh Actor will cast a soft area shadow when in precomputed, indirectly lit areas. With your own Static Mesh, you should see results similar to the example below.
+	
+	[REGION:note]
+	If you're not seeing a soft shadow or drastically differing results, first make sure that you've already built lighting for the level and that you're testing in an area that gets a good amount of bounce light. In very brightly lit
+	areas or areas directly lit by the light source, the indirect shadows will be practically non-existent. 
+	[/REGION]
+	
+	![](DFISMinVisibility_1.png)(w:450)		
+
+	
+### Distance Field Indirect Shadow Min Visibility
+
+![](DFISMinVisibilitySetting.png)(w:450)
+
+The setting for **Distance Field Indirect Shadow Min Visibility** controls how dark the movable Actor's indirect shadow can be. This settings gives artists the most control over how dark these shadows can be, enabling 
+them to blend with surrounding precomputed shadows. 
+
+[OBJECT:ComparisonSlider]
+	[PARAM:before]
+	![Min Visibility: 0.1 ](DFISMinVisibility_1.png)(w:550)
+	[/PARAM]
+	[PARAM:after]
+	![Min Visibility: 0.6 ](DFISMinVisibility_2.png)(w:550)
+	[/PARAM]
+[/OBJECT]
+
+	
+	
+## End Result
+
+Now that you've seen how to use Distance Field Indirect Shadows, you should be able to enable your own Static Meshes to use dynamic indirect shadowing in your levels and individual actors, along with being able to adjust the shadows to blend nicely with your surrounding baked lighting in indirectly lit areas.
+
+[REGION:asyncgif]
+![](DFIndirect.gif)(w:600)
+[/Region]
+
