Copying //UE4/Dev-Documentation to //UE4/Samples-Main (Source: //UE4/Dev-Documentation @ 3372893)

diff --git a/Engine/Documentation/Source/Engine/Rendering/LightingAndShadows/DistanceFields/HowTo/DFHT_5/DFHT_5.INT.udn b/Engine/Documentation/Source/Engine/Rendering/LightingAndShadows/DistanceFields/HowTo/DFHT_5/DFHT_5.INT.udn
new file mode 100644
--- /dev/null
+++ b/Engine/Documentation/Source/Engine/Rendering/LightingAndShadows/DistanceFields/HowTo/DFHT_5/DFHT_5.INT.udn
@@ -0,0 +1,126 @@
+Availability: Docs
+Crumbs:
+Title:Using Static Mesh Build Settings for Distance Fields
+Description: An overview of the settings in the Static Mesh Editor's Build Settings for Mesh Distance Fields.
+Related: Engine/Rendering/LightingAndShadows/RayTracedDistanceFieldShadowing
+Related: Engine/Rendering/LightingAndShadows/DistanceFieldAmbientOcclusion
+Parent: Engine/Rendering/LightingAndShadows/DistanceFields/HowTo
+type: how-to
+SkillLevel:Intermediate
+Order:
+Version: 4.15
+Tags: Lighting
+
+
+[REGION:fullwdith raw]
+![](BuildSettingsDF_Banner.png)
+[/REGION]
+
+[TOC(start:2 end:3)] 
+
+When you open Static Meshes in the Editor, you'll find the **Build Settings** rollout in the Details panel under each **LOD** section. These properties serve to change aspects of your Static Meshes for your game, whether that is to generate a new lightmap all the way to changing the scale of the mesh without having to reimport or manually change it inside of your levels. You may have noticed that some of these settings are available in the Import Options window when you import your mesh, but some are more specialized, like the ones for Distance Fields that we'll be taking a look at in this guide.
+
+Continue reading below to explore how the **Build Settings** options for Distance Fields can be used in your games that have them enabled. 
+
+<!---![](StaticMeshBuildSettings.png)--->
+
+## Distance Field Resolution Scale
+
+The **Distance Field Resolution Scale** is used to increase the accuracy of the Mesh Distance Field representation, although this will significantly increase the memory cost of the volume texture used. In the Static Mesh Editor's viewport, you can use the stats in the corner to see what the volume texture's memory usage and size is.
+
+![](VolumeTextureSize.png)(w:300)
+[REGION:caption]
+Distance Field Stat for Volume Texture Size and Memory Usage
+[/REGION]
+
+[OBJECT:ComparisonSlider]
+	[PARAM:before]
+	![Distance Field Resolution Scale: 1 ](MeshView_1.png)(w:650)
+	[/PARAM]
+	[PARAM:after]
+	![Distance Field Resolution Scale: 5](MeshView_2.png)(w:650)
+	[/PARAM]
+[/OBJECT]
+
+The statue's Distance Field Resolution Scale has been increased from 1 to 5 and as a result you get a significant increase in resolution and visual quality, but at an increased cost; memory usage increases from 0.03 MB to 1.01 MB. 
+
+## Generate Distance Field As If Two Sided
+
+The setting for **Generate Distance Field as if Two-Sided** enables use with things like foliage, where you often have many sheets that make up the leaves of the plant. Since the Mesh Distance Fields feature operates by 
+using surfaces this means that the ambient occlusion (AO) or shadows never really soften. When you toggle this option, the shadows and AO will soften which reduces the artifacts that can appear from an opaque Distance Field,
+but it will cost more since the shadows are never fully opaque. 
+
+[OBJECT:ComparisonSlider]
+	[PARAM:before]
+	![Generate Distance Field as if Two-Sided: On](DF_TS2.png)(w:650)
+	[/PARAM]
+	[PARAM:after]
+	![Generate Distance Field as if Two-Sided: Off](DF_TS1.png)(w:650)
+	[/PARAM]
+[/OBJECT]
+
+## Distance Field Bias
+
+The setting for **Distance Field Bias** which enables you to reduce self-shadowing artifacts on ambient animations where the surface may move inside of the Mesh Distance Field or on complex geometry that has some minor 
+self-shadowing.
+
+[OBJECT:ComparisonSlider]
+	[PARAM:before]
+	![Distance Field Bias: 0 ](DF_Bias0.png)(w:650)
+	[/PARAM]
+	[PARAM:after]
+	![Distance Field Bias: 1 ](DF_Bias1.png)(w:650)
+	[/PARAM]
+[/OBJECT]
+
+The rock's crevices have some self-shadowing artifacts that appear, which can more clearly be seen with a simple base color applied. Increasing the Distance Field Bias shrinks the Mesh Distance Field inward so that 
+parts of the mesh are no longer inside of it.
+
+## Distance Field Replacement Mesh
+
+The selection dropdown for **Distance Field Replacement Mesh** enables you to select any Static Mesh in your project to use as a replacement for the current meshes Distance Field representation. 
+
+[OBJECT:ComparisonSlider]
+	[PARAM:before]
+	![Distance Field Replacement Mesh none selected ](DFReplace1.png)(w:650)
+	[/PARAM]
+	[PARAM:after]
+	![Distance Field Replacement Mesh selected ](DFReplace2.png)(w:650)
+	[/PARAM]
+[/OBJECT]
+
+The Distance Field representation can be swapped for another mesh when needed and in this demonstration a noticeably different mesh has been dilberately chosen, although you would likely use a mesh that is 
+closely resembles the original mesh, but without as much complexity in geometry to keep the volume texture size lower and still retain the desired shape. 
+
+## End Result
+
+Now that you've explored the Build Settings specific to Distance Fields, you should have a better understanding of what these can do to help you optimize and improve quality for your projects that 
+use them. 
+
+
+
+
+
+
+
+
+
+
+
+
+
+
+
+
+
+
+
+
+
+
+
+
+
+
+
+
diff --git a/Engine/Documentation/Source/Engine/Rendering/LightingAndShadows/DistanceFields/HowTo/DFHowTo.INT.udn b/Engine/Documentation/Source/Engine/Rendering/LightingAndShadows/DistanceFields/HowTo/DFHowTo.INT.udn
new file mode 100644
--- /dev/null
+++ b/Engine/Documentation/Source/Engine/Rendering/LightingAndShadows/DistanceFields/HowTo/DFHowTo.INT.udn
@@ -0,0 +1,30 @@
+Availability: Docs
+Title: Distance Fields How-To's
+Crumbs:
+Description: A how to Guide for using different aspects of Mesh Distance Fields in Unreal Engine 4.
+Version: 4.15
+type:landing
+parent:Engine/Rendering/LightingAndShadows/DistanceFields
+order:
+topic-image:DFHowTo_topic.png
+tags:Lighting
+
+
+The following step-by-step guides walk you through a variety of topics designed to help you better use Distance Fields in your projects. 
+
+
+## Beginner
+
+[DIR(output:"fancy" parent:"Engine/Rendering/LightingAndShadows/DistanceFields/HowTo" skill_level:"Beginner")]
+
+## Intermediate
+
+[DIR(output:"fancy" parent:"Engine/Rendering/LightingAndShadows/DistanceFields/HowTo" skill_level:"Intermediate")]
+
+
+<!---
+## Advanced
+
+[DIR(output:"fancy" parent:"Engine/Rendering/Materials/HowTo" skill_level:"Advanced")]
+
+--->
\ No newline at end of file
diff --git a/Engine/Documentation/Source/Engine/Rendering/LightingAndShadows/LightingAndShadows.INT.udn b/Engine/Documentation/Source/Engine/Rendering/LightingAndShadows/LightingAndShadows.INT.udn
--- a/Engine/Documentation/Source/Engine/Rendering/LightingAndShadows/LightingAndShadows.INT.udn
+++ b/Engine/Documentation/Source/Engine/Rendering/LightingAndShadows/LightingAndShadows.INT.udn
@@ -67,7 +67,7 @@ type:landing
 %Engine/Rendering/LightingAndShadows/IndirectLightingCache:Topic%
 %Engine/Rendering/LightingAndShadows/LightMobility/StaticLights:Topic%
 %Engine/Rendering/LightingAndShadows/LightMobility/StationaryLights:Topic%
-%Engine/Rendering/LightingAndShadows/PrecomputedLightingScenarios:Topic%
+<!--%Engine/Rendering/LightingAndShadows/PrecomputedLightingScenarios:Topic%-->
 [/REGION]
 
 
@@ -79,6 +79,7 @@ type:landing
 %Engine/Rendering/LightingAndShadows/LightPropagationVolumes:Topic%
 %Engine/Rendering/LightingAndShadows/DistanceFieldAmbientOcclusion:Topic%
 %Engine/Rendering/LightingAndShadows/RayTracedDistanceFieldShadowing:Topic%
+%Engine/Rendering/LightingAndShadows/ContactShadows:Topic%
 [/REGION]
 
 ## General
diff --git a/Engine/Documentation/Source/Engine/Rendering/Materials/CustomizedUVs/CustomizedUVs.INT.udn b/Engine/Documentation/Source/Engine/Rendering/Materials/CustomizedUVs/CustomizedUVs.INT.udn
--- a/Engine/Documentation/Source/Engine/Rendering/Materials/CustomizedUVs/CustomizedUVs.INT.udn
+++ b/Engine/Documentation/Source/Engine/Rendering/Materials/CustomizedUVs/CustomizedUVs.INT.udn
@@ -7,6 +7,7 @@ parent:Engine/Rendering/Materials
 order:12
 topic-image:CUVs_topic.png
 type:overview
+tags:Materials
 
 
 [TOC(start:2)]
diff --git a/Engine/Documentation/Source/Engine/Rendering/Materials/Editor/Interface/MaterialEditorInterface.INT.udn b/Engine/Documentation/Source/Engine/Rendering/Materials/Editor/Interface/MaterialEditorInterface.INT.udn
--- a/Engine/Documentation/Source/Engine/Rendering/Materials/Editor/Interface/MaterialEditorInterface.INT.udn
+++ b/Engine/Documentation/Source/Engine/Rendering/Materials/Editor/Interface/MaterialEditorInterface.INT.udn
@@ -3,6 +3,7 @@ Title:Material Editor UI
 Crumbs: %ROOT%, Engine, Engine/Rendering/Materials, Engine/Rendering/Materials/Editor
 Description:Guide to using the Material Editor UI.
 Version: 4.9
+tags:Materials
 
 
 [TOC(start:2)]
diff --git a/Engine/Documentation/Source/Engine/Rendering/Materials/ExpressionReference/Math/MathExpressions.INT.udn b/Engine/Documentation/Source/Engine/Rendering/Materials/ExpressionReference/Math/MathExpressions.INT.udn
--- a/Engine/Documentation/Source/Engine/Rendering/Materials/ExpressionReference/Math/MathExpressions.INT.udn
+++ b/Engine/Documentation/Source/Engine/Rendering/Materials/ExpressionReference/Math/MathExpressions.INT.udn
@@ -3,6 +3,7 @@ Title:Math Expressions
 Crumbs: %ROOT%, Engine, Engine/Rendering/Materials, Engine/Rendering/Materials/ExpressionReference
 Description:Expressions that perform mathematical operations on one or more inputs.
 Version: 4.9
+tags:Materials
 
 [EXCERPT:TOC]
 [TOC(start:2)]
@@ -67,7 +68,7 @@ The **Clamp** expression takes in value(s) and constrains them to a specified ra
 | **Max** | Takes in the value to use as the maximum when clamping. |
 
 
-**Examples:** Clamping 0.3 with Min 0.0 and Max 1.0 yields 0.3; Clamping 1.3 with Min 0.0 and Max 1.0 yields 1.0.
+**Examples:** Clamping an input range of (0.0) to (1.0) 0.3 with Min 0.0 and Max 1.0 yields 0.3; Clamping 1.3 with Min 0.0 and Max 1.0 yields 1.0.
 
 ![ClampExample.png](ClampExample.png)(w:900)
 
@@ -77,7 +78,7 @@ The **Clamp** expression takes in value(s) and constrains them to a specified ra
 
 ## Cosine
 
-The **Cosine** expression outputs the value of a sine wave over the range of [0, 1], repeating. Most commonly, this is used to output a continuous oscillating waveform by connecting a [Time](Engine/Rendering/Materials/ExpressionReference/Constant#Time) expression to its input. The output value will cycle back and forth between -1 and 1. A visual representation of the wave is shown below:
+The **Cosine** expression outputs the value of a cosine wave over the input range of [0, 1] and the output range of [-1, 1], both repeating. Most commonly, this is used to output a continuous oscillating waveform by connecting a [Time](Engine/Rendering/Materials/ExpressionReference/Constant#Time) expression to its input, but it can also be used to create ripples in worldspace or screenspace, or any other application where a continuous, smooth cycle is needed. A visual representation of the wave is shown below, scaled to the [0, 1] output range:
 
 ![CosineWave.png](CosineWave.png)
 
@@ -115,7 +116,7 @@ The **CrossProduct** expression computes the cross product of two three-channel
 
 ## Divide
 
-The **Divide** expression takes two inputs and outputs the result of the first divided by the second.  The division happens per channel, meaning that the R channel of the first is divided by the second, the G channel of the first is divided by the second, and so on. Both inputs must have the same number of values unless the divisor is a single float value. _Never divide by zero._
+The **Divide** expression takes two inputs and outputs the result of the first divided by the second. The division happens per channel, meaning that the R channel of the first is divided by the second, the G channel of the first is divided by the second, and so on. Both inputs must have the same number of values unless the divisor is a single float value. Division by zero in any channel will result in that channel containing the local hardware's "infinity" implementation, which is usually a very large, positive value.
 
 | Item | Description |
 | -------- | ----------- |
@@ -127,7 +128,7 @@ The **Divide** expression takes two inputs and outputs the result of the first d
 | **B** | Takes in the value(s) to divide by, the divisor. |
 
 
-**Examples:** The output of a Divide with A=(1.0,0.5,-0.4) and B=(2.0,2.0,4.0) is (0.5,0.25,-0.1)
+**Examples:** Using Divide with A=(1.0) and B=(5.0) outputs (0.2), which appears as a dark gray.
 
 ![DivideExample.png](DivideExample.png)
 
@@ -135,13 +136,13 @@ The **Divide** expression takes two inputs and outputs the result of the first d
 [EXCERPT:DotProduct]
 ## DotProduct
 
-The **DotProduct** expression computes the dot product, or the length of one vector projected onto the other. This calculation is used by many techniques for computing falloff.  DotProduct requires both vector inputs to have the same number of channels.
+The **DotProduct** expression computes the dot product, which can be described as the length of one vector projected onto the other, or as the cosine between the two vectors multiplied by their magnitudes. This calculation is used by many techniques for computing falloff.  DotProduct requires both vector inputs to have the same number of channels.
 
 | Item | Description |
 | -------- | ----------- |
 | [REGION:tablesection]Inputs[/REGION] ||
 | **A** | Takes in a value or vector of any length. |
-| **B** | Takes in a value or vector of any length. |
+| **B** | Takes in a value or vector of the same length as **A**. |
 
 
 ![DotProductExample.png](DotProductExample.png)(w:900)
@@ -159,17 +160,17 @@ The **Floor** expression takes in value(s), rounds them **down** to the previous
 
 ## FMod
 
-The **FMod** expression returns the floating-point remainder of the division operation of the two inputs.
+The **FMod** expression returns the floating-point remainder of the division operation of the two inputs. The Dividend (input "A") can be any value, but negative Dividends will result in negative results. The Divisor (second input) should not be zero as this implies a division by zero, but whether the Divisor is negative or positive will not affect the result. A common use case is to make a material that brightens up to a maximum value, then immediately drops back to the minimum value on the next frame, only to begin climbing toward the maximum again.
 
 ![FModExample.png](FModExample.png)
 
-In this example, FMod is returning a repeating value that ramps up from 0 to 1 every second. Once 1 is reached, the value snaps back to 0 and repeats indefinitely.
+In this example, FMod takes a zero-to-one UV progression and converts it into a repeating cycle every 0.2 UV units on the X axis in the green channel.
 
 ## Frac
 
-The **Frac** expression takes in value(s) and outputs the decimal part of those values.  See also [Ceil](#Ceil) and [Floor](#Floor).
+The **Frac** expression takes in values and outputs the fractional portion of those values. In other words, for an input value "X", the result is "X minus Floor of X". The output value will range from zero to one, inclusive on the low end, but not the high end. See also [Ceil](#Ceil) and [Floor](#Floor).
 
-**Examples:** Frac of 0.2 is 0.2; Frac of (0.0,1.6) is (0.0, 0.6).
+**Examples:** Frac of (0.2) is (0.2). Frac of (-0.2) is (0.8). Frac of (0.0,1.6,1.0) is (0.0,0.6,0.0).
 
 ![FracExample.png](FracExample.png)(w:900)
 
@@ -202,7 +203,7 @@ In this example, the If expression takes in a heightmap and outputs either red o
 
 ## Max
 
-The **Max** expression takes in two inputs and then outputs the maximum of the two. 
+The **Max** expression takes in two inputs and outputs the higher of the two.
 
 This node is similar to Photoshop's Lighten. 
 
@@ -222,7 +223,7 @@ In the example above, A is "0" and B is "1"; therefore, "1" (white) is the resul
 
 ## Min
 
-The **Min** expression takes in two inputs then outputs the minimum of the two. 
+The **Min** expression takes in two inputs and outputs the lower of the two.
 
 This node is similar to Photoshop's Darken. 
 
@@ -244,7 +245,7 @@ In the example above, A is "0" and B is "1"; therefore, "0" (black) is the resul
 
 ## Multiply
 
-The **Multiply** expression takes two inputs, multiplies them together, and outputs the result. Similar to Photoshop's multiply layer blend.  Multiplication happens per channel, meaning that the R channel of the first is multiplied by the second; the G channel of the first is multiplied by the second, and so on. Both inputs must have the same number of values unless one of the values is a single float value.
+The **Multiply** expression takes two inputs, multiplies them together, and outputs the result. Similar to Photoshop's multiply layer blend.  Multiplication happens per channel, meaning that the R channel of the first is multiplied by the R channel of the second; the G chanWnel of the first is multiplied by the G channel of the second, and so on. Both inputs must have the same number of values unless one of the values is a single float value.
 
 | Item | Description |
 | -------- | ----------- |
@@ -268,9 +269,18 @@ Do not forget that materials in UE4 are not limited to [0,1].  If colors/values
 [EXCERPT:Normalize]
 ## Normalize
 
-The **Normalize** expression calculates and outputs the normalized value of its input. This means each component of the input is divided by the L-2 norm (length) of the vector.
+The **Normalize** expression calculates and outputs the normalized value of its input. Normalized vectors (also called "unit vectors") have an overall length of 1.0. This means each component of the input is divided by the total magnitude (length) of the vector.
 
-![NormalizeExample.png](NormalizeExample.png)(w:900)
+**Examples:** Passing either (0,2,0) or (0,0.2,0) through Normalize will output (0,1,0). Passing (0,1,-1) through Normalize will output (0, 0.707, -0.707). The only special case is an all-zero vector, which will be unchanged.
+
+[OBJECT:ComparisonSlider]
+	[PARAM:before]
+	![Normalized Input Vector](NormalizeExample.png)(w:900)
+	[/PARAM]
+	[PARAM:after]
+	![Non-Normalized Input Vector](NormalizeExample2.png)(w:900)
+	[/PARAM]
+[/OBJECT]
 
 [REGION:note]
 It is not necessary to normalize an expression that plugs into the Normal material output.
@@ -280,7 +290,7 @@ It is not necessary to normalize an expression that plugs into the Normal materi
 
 ## OneMinus
 
-The **OneMinus** expression takes an input value and outputs 'one-minus-that-value'. This operation is performed per channel.
+The **OneMinus** expression takes an input value "X" and outputs "One Minus X". This operation is performed per channel.
 
 **Examples:** OneMinus of 0.4 is 0.6; OneMinus of (0.2,0.5,1.0) is (0.8,0.5,0.0); OneMinus of (0.0,-0.4,1.6) is (1.0,1.4,-0.6).
 
@@ -304,14 +314,14 @@ The **Power** expression takes two inputs, raises Base to the Exp power, and out
 
 **Examples:** Power of Base 0.5, Exp 2.0 is 0.25.
 
-**Example Usage:** If the colors you pass to Power are in [0,1], Power can behave as a contrast adjustment, where only brighter values remain.
+**Example Usage:** If the colors you pass to Power are in [0,1], Power can behave as a sort of contrast adjustment, where very bright values are dimmed slightly, but darker values are decreased drastically.
 
 ![PowerExample.png](PowerExample.png)
 
 
 ## Sine
 
-The **Sine** expression outputs the value of a sine wave over the range of [0, 1], repeating. Most commonly, this is used to output a continuous oscillating waveform by connecting a [Time](Engine/Rendering/Materials/ExpressionReference/Constant#Time) expression to its input. The output value will cycle back and forth between -1 and 1. The difference between this and the output of the [Cosine](#Cosine) expression is the output waveform is offset half the period. A visual representation of the wave is shown below:
+The **Sine** expression outputs the value of a Sine wave over the input range of [0, 1] and the output range of [-1, 1], both repeating. The difference between this and the output of the [Cosine](#Cosine) expression is the output waveform is offset by one-quarter of the period, meaning that `Cos(X)` is equal to `Sin(X + 0.25)`. Most commonly, this is used to output a continuous oscillating waveform by connecting a [Time](Engine/Rendering/Materials/ExpressionReference/Constant#Time) expression to its input, but it can also be used to create ripples in worldspace or screenspace, or any other application where a continuous, smooth cycle is needed. A visual representation of the wave is shown below, scaled to the [0, 1] output range:
 
 ![SineWave.png](SineWave.png)
 
@@ -328,7 +338,7 @@ The **Sine** expression outputs the value of a sine wave over the range of [0, 1
 
 ## SquareRoot
 
-The **SquareRoot** expression outputs the square root of the input value.  SquareRoot can operate only on a single float input value.
+The **SquareRoot** expression outputs the square root of the input value. If applied to a vector, each component will be handled individually.
 
 ![SqrtExample.png](SqrtExample.png)(w:900)
 
diff --git a/Engine/Documentation/Source/Engine/Rendering/Materials/ExpressionReference/VectorOps/VectorOpExpressions.INT.udn b/Engine/Documentation/Source/Engine/Rendering/Materials/ExpressionReference/VectorOps/VectorOpExpressions.INT.udn
--- a/Engine/Documentation/Source/Engine/Rendering/Materials/ExpressionReference/VectorOps/VectorOpExpressions.INT.udn
+++ b/Engine/Documentation/Source/Engine/Rendering/Materials/ExpressionReference/VectorOps/VectorOpExpressions.INT.udn
@@ -3,6 +3,7 @@ Title:Vector Operation Expressions
 Crumbs: %ROOT%, Engine, Engine/Rendering/Materials, Engine/Rendering/Materials/ExpressionReference
 Description:Expressions that perform operations on vector input values.
 Version: 4.9
+tags:Materials
 
 [EXCERPT:TOC]
 [TOC(start:2)]
@@ -20,7 +21,7 @@ The **AppendVector** expression allows you to combine channels together to creat
 | **B** | Takes in the value(s) to be appended. |
 
 
-**Examples:** Append of 0.2 and 0.4 is (0.2,0.4); Append of (0.2,0.4) and (1.6) is (0.2,0.4,1.6).
+**Examples:** Append of 0.2 and 0.4 is (0.2,0.4); Append of (0.2,0.4) and (1.0) is (0.2,0.4,1.0).
 
 ![AppendVectorExample.png](AppendVectorExample.png)
 
@@ -41,7 +42,7 @@ The **ComponentMask** expression allows you to select a specific subset of chann
 | **A** | If checked, the alpha (fourth) channel of the input value will be passed through to the output. |
 
 
-**Examples:** ComponentMask with an input of (0.2,0.8,1.4) and the R and B channels checked will output (0.2,1.4).
+**Examples:** ComponentMask with an input of (0.2,0.4,1.0) and the G channel will produce an output of (0.4), which appears as a 40% bright grayscale value when used as a color vector.
 
 [REGION:fullwidth]
 ![ComponentMaskExample.png](ComponentMaskExample.png)
@@ -55,7 +56,7 @@ The **ComponentMask** expression allows you to select a specific subset of chann
 
 ## DeriveNormalZ
 
-The **DeriveNormalZ** expression derives the Z component of a tangent space normal given the X and Y components and outputs the resulting three-channel tangent space normal. Z is calculated as Z = sqrt(1 - (x * x + y * y));
+The **DeriveNormalZ** expression derives the Z component of a tangent space normal given the X and Y components and outputs the resulting three-channel tangent space normal. Z is calculated as Z = sqrt(1 - (x \* x + y \* y));
 
 | Item | Description |
 | -------- | ----------- |
diff --git a/Engine/Documentation/Source/Engine/Rendering/Materials/Functions/Reference/FunctionReference.INT.udn b/Engine/Documentation/Source/Engine/Rendering/Materials/Functions/Reference/FunctionReference.INT.udn
--- a/Engine/Documentation/Source/Engine/Rendering/Materials/Functions/Reference/FunctionReference.INT.udn
+++ b/Engine/Documentation/Source/Engine/Rendering/Materials/Functions/Reference/FunctionReference.INT.udn
@@ -3,7 +3,7 @@ title:Material Function Reference
 Crumbs: %ROOT%, Engine, Engine/Rendering/Materials
 Description:Reference page for each of the material functions, separated by category.
 Version: 4.9
-
+tags:Materials
 
 Below is breakdown of the various **Material Functions** included within Unreal Editor by default, sorted by category.
 
diff --git a/Engine/Documentation/Source/Engine/Rendering/Materials/HowTo/Making_Functions/Making_Functions.INT.udn b/Engine/Documentation/Source/Engine/Rendering/Materials/HowTo/Making_Functions/Making_Functions.INT.udn
--- a/Engine/Documentation/Source/Engine/Rendering/Materials/HowTo/Making_Functions/Making_Functions.INT.udn
+++ b/Engine/Documentation/Source/Engine/Rendering/Materials/HowTo/Making_Functions/Making_Functions.INT.udn
@@ -10,7 +10,7 @@ Related: Engine/Rendering/Materials/ExpressionReference
 Related: Engine/Rendering/Materials
 Version: 4.9
 SkillLevel: Intermediate
-
+tags:Materials
 
 [TOC(start:2)]
 
diff --git a/Engine/Documentation/Source/Engine/Rendering/Materials/HowTo/MeshDecals/MeshDecals.INT.udn b/Engine/Documentation/Source/Engine/Rendering/Materials/HowTo/MeshDecals/MeshDecals.INT.udn
new file mode 100644
--- /dev/null
+++ b/Engine/Documentation/Source/Engine/Rendering/Materials/HowTo/MeshDecals/MeshDecals.INT.udn
@@ -0,0 +1,166 @@
+Availability:Docs
+Title: Using Mesh Decals
+Description: This how to will cover setting up and using Mesh Decals.
+Crumbs:
+parent: Engine/Rendering/Materials/HowTo
+type: reference
+tags:Materials
+version:4.15
+skilllevel:intermediate
+
+[TOC (start:2 end:4)]
+
+
+Unreal Engine 4 now has support for Mesh Decals which enables you to use the properties of Deferred Decals on separate surface geometry 
+for added detail to your Static and Skeletal Meshes. Because Deferred Decals rely on projection, you are limited to mostly planar surface 
+details that shear and distort when not aligned with the surface it’s projecting onto. Mesh Decals afford you decals that do not follow 
+a simple projection and instead can be used with geometry that wraps around edges, used with spline meshes, and enhance the look of your characters. 
+
+
+## Mesh Decal VS Masked Material
+
+Mesh Decals are a mix of translucent blend mode and deferred decals in that they do not render to depth but do update the GBuffer or 
+DBuffer after the opaque geometry has been rendered. In contrast to using a masked material, there is no cost for EarlyZ pass, you don't get shadows or 
+proper occlusion, but with the trade-off, you get the soft transitions in the material. 
+
+Mesh Decals provide several differences to Deferred Decal Actors that you should be aware of: 
+
+* There are fewer draw calls since large Deferred Decals often do front and back facing draw calls.
+* It's quite faster since there are fewer pixels covered and flat back facing decals cover 0 pixels. 
+* You have the ability to do more complex projections since you can use a custom UV. 
+* You have the ability to have properly normal mapped decals wrapping around surfaces or stretching along splines.
+
+[OBJECT:ComparisonSlider]
+	[PARAM:before]
+	![ Mesh Decal ](MeshDecal.png)(w:400)
+	[/PARAM]
+	[PARAM:after]
+	![ Masked Material ](MaskedMaterial.png)(w:400)
+	[/PARAM]
+[/OBJECT]
+
+
+## Authoring Content
+
+Authoring content for use as Mesh Decal geometry is in-line with any model creation. Since this geometry will not rely on projection like a Deferred Decal Actor, the geometry
+needs only to stay in front of the surface you want to affect. With this in mind, the decal geometry should "hug" the underlying surface, but it does not have to match it
+very closely because there is already a small depth bias that is hard-coded to account for this small offset. Also, using a bit of geometry that feathers out from the decal can be good
+to create the soft transitions for the decal that you cannot attain with a masked material.
+
+<div align= "left">
+[REGION:imagetable]
+| ![](MeshDecalAuthoring.png)(h:450) | ![](MeshDecalAuthoringComposite.png)(h:450) |
+| --- | --- |
+| Base Mesh and Decal Geometry (Isolated) | Composited Mesh |
+[/REGION]
+</div>
+
+Another thing to keep in mind when developing your content is that Mesh Decals can be a challenge with Levels of Detail (LOD) or long view distances when there is limited depth buffer
+precision because geometry will intersect or not match up as intended. However, you can either change the mesh to account for this in most cases or you can use World Position Offset in 
+your material to adjust the offset without having to return to your modeling application. 
+
+![](Offset.png)(w:500)
+
+[OBJECT:ComparisonSlider]
+	[PARAM:before]
+	![ Offset Value: 0 ](2_Offset.png)(w:500)
+	[/PARAM]
+	[PARAM:after]
+	![ Offset Value: -0.6 ](1_Offset.png)(w:500)
+	[/PARAM]
+[/OBJECT]
+
+In this example, the decal geometry is close enough to the base mesh that the depth cannot account the offset. Using a small negative offset value pushes the geometry out just enough that it 
+doesn't intersect with the underlying geometry.
+
+## Limitations
+
+* There is no artist specified sort order.
+* Tangent Space support is not yet implemented.
+* Shader Overdraw/Overdraw functionality is missing.
+* There is no easily adjustable DepthBias. Currently, you have to use an offset in the mesh from the surface it's above, or control with World Position Offset in the Material.
+* Sort by depth to avoid artifacts when there are many layers overlapping.
+* Material Editor Preview is not visible.
+* Material Editor does not currently show the correct instructions count since it assumes base pass usage.
+* DBuffer and non-DBuffer use the same container.
+* Parallel Rendering isn't yet hooked up. There are some CPU savings if this feature is used extensively. 
+
+## Using Mesh Decal Materials
+
+In the following steps, we will enable Dbuffer Decals for our project so that we can take advantage of lighting features for our decal material and we'll create a basic 
+Deferred Decal Material that can be applied to our Static Mesh.  
+
+### Required Files
+
+In order to follow along with this How-To, you will need to download, extract, and then import to UE4 the following FBX and Texture files. If you are unfamiliar with how to do this please
+check out [How to Import Static Meshes Guide](Engine/Content/ImportingContent/ImportingMeshes) and [How to Import Texture Guide](Engine/Content/ImportingContent/ImportingTextures) for more 
+information about how to go about doing this. 
+
+**[Required Textures Download](MeshDecalAssets.zip)** (Right - Click Save as)
+
+
+### Enabling DBuffer Decals for your Project
+
+1. In the **Edit** menu, click on **Project Settings** to view the Unreal Editor options. 
+
+	![](ProjectSettingsMenu.png)(w:250)
+1. Under the **Engine** heading, select the **Rendering** section and locate the **Lighting** category that contains various Lighting options.
+
+	[REGION:lightbox]
+	[![](ProjectSettings.png)(w:700)](ProjectSettings.png)
+	[/REGION]
+1. Now, you can toggle the **DBuffer Decals** option to use lighting features with your Deferred Decals.
+	
+	![](EnableDBufferDecal.png)(w:350)
+	
+1. Before you'll be able to use this feature, you'll first need to **Restart the Editor**.
+	
+	![](RestartTheEditor.png)(w:550)
+
+	
+### Creating a Basic Deferred Decal
+
+1. In the **Content Browser**, select the **Add New** button and choose the option for **Material** to create a new Material. Make sure you give your Material a name that will be easy to locate later.
+For the purpose of this how-to, I have named mine "M_MeshDecal."
+	
+	![](AddNewMaterial.png)(w:300)
+
+1. Now, select your **Material** and double-click it to open it. Once the Material Editor opens, start by setting the following attributes in the **Details** panel so that it can be used as a Deferred Decal.
+
+	* **Material Domain:** Deferred Decal
+	* **Blend Mode:** Translucent
+	* **Decal Blend Mode:** DBuffer Translucent Color, Normal, Roughness 
+
+1. Next, you'll need to setup your material using the textures in the **Required Files** section of this page. There are three textures in the .zip file that we will need for this step; a diffuse, mask, and normal texture.
+Once you've imported these textures, select them from the **Content Browser** and drag them into the Material Editor Graph. 
+
+	![](MaterialGraph1.png)(w:400)
+
+	Now, plug in the outputs of the Texture Sample nodes into their corresponding inputs on the Main Material node. For the Mask texture,"T_MeshDecalDamage_M", be sure to use the **Blue Channel** output when
+	plugging into the Opacity Mask input. This will make sure the blue values in the texture are used as the maskfor what should be visible. 
+		
+	![](MaterialGraph2.png)(w:400)
+1. Right-click in the Material Graph and type in "Constant" or hold down the "1" key and click in the graph to add a **Constant** value node. Plug this into your **Roughness** input and give it a default value of **0.7**. 
+
+	![](MaterialGraph3.png)(w:400)
+1. Before finishing up, we'll need a way to control the offset of the decal geometry from the base mesh to prevent or lessen any artifacts that can occur due to depth precision. Add the following nodes to your Material Graph
+and plug the output of the **Multiply** node into the **World Position Offset** input of the Main Material node. 
+
+	![](MaterialGraph4.png)(w:400)
+
+	Once you've done this, you're Material Graph should look similar to this: 
+
+	[REGION:lightbox]
+	[![](FinalMaterialGraph.png)(w:400)](FinalMaterialGraph.png)
+	[/REGION]
+1. With your Deferred Decal Material complete, place the Static Mesh "SM_MeshDecal", which is included in the Required Files .zip, into your level and apply the Material **M_MeshDecal** to Element 0
+of the Static Mesh's Material slots. For the Base Mesh's Material you can assign any material or even just a basic color to Element 1 for the purposes of this demonstration. You should now have something
+that looks similar to our example below. 
+
+	![](MeshDecal.png)(w:400)
+
+	
+## End Result
+
+Now that you've set up and seen how to create your own Materials to use with a Mesh Decal, you should be confident in authoring your own assets in your modeling applications and creating your own Materials that take 
+full advantage of lighting for use with Mesh Decals applied. 
diff --git a/Engine/Documentation/Source/Engine/Rendering/Materials/IntroductionToMaterials/IntroductionToMaterials.INT.udn b/Engine/Documentation/Source/Engine/Rendering/Materials/IntroductionToMaterials/IntroductionToMaterials.INT.udn
--- a/Engine/Documentation/Source/Engine/Rendering/Materials/IntroductionToMaterials/IntroductionToMaterials.INT.udn
+++ b/Engine/Documentation/Source/Engine/Rendering/Materials/IntroductionToMaterials/IntroductionToMaterials.INT.udn
@@ -7,26 +7,7 @@ SkillLevel: Beginner
 topic-image:Engine/Rendering/Materials/EssentialImage.png
 order:1
 type:overview
-
-[VAR:playlistnav]
-[OBJECT:PlaylistNavigation]
-	[PARAMLITERAL:id]
-		LevelDesign_Beginner
-	[/PARAMLITERAL]
-	[PARAM:previous]
-		[Previous Topic](Engine/Actors/Triggers "%Engine/Actors/Triggers:title%")
-	[/PARAM]
-	[PARAM:next]
-		[Next Topic](Engine/Rendering/LightingAndShadows/Basics "%Engine/Rendering/LightingAndShadows/Basics:title%")
-	[/PARAM]
-	[PARAM:home]
-		[](Tracks/LevelDesign/Beginner)
-		%Tracks/LevelDesign/Beginner:menu%
-	[/PARAM]
-[/OBJECT]
-[/VAR]
-
-%playlistnav%
+tags:Materials
 
 
 [REGION:fullwidth]
@@ -155,5 +136,3 @@ Applying a Material to a Brush surface is a simple process:
 
 	![](AddingMaterialToBox.png)![](RedOnBox.png)(w:700)
 
-
-%playlistnav%
diff --git a/Engine/Documentation/Source/Engine/Rendering/Materials/LightingModels/SubSurfaceProfile/SSSProfileLightingModel.INT.udn b/Engine/Documentation/Source/Engine/Rendering/Materials/LightingModels/SubSurfaceProfile/SSSProfileLightingModel.INT.udn
--- a/Engine/Documentation/Source/Engine/Rendering/Materials/LightingModels/SubSurfaceProfile/SSSProfileLightingModel.INT.udn
+++ b/Engine/Documentation/Source/Engine/Rendering/Materials/LightingModels/SubSurfaceProfile/SSSProfileLightingModel.INT.udn
@@ -7,6 +7,8 @@ parent:Engine/Rendering/Materials/LightingModels
 order:1
 topic-image:sss_profile_topic.png
 type:overview
+tags:Materials
+
 
 [TOC(start:2)]
 
diff --git a/Engine/Documentation/Source/Engine/Rendering/Materials/MaterialInputs/MaterialInputs.INT.udn b/Engine/Documentation/Source/Engine/Rendering/Materials/MaterialInputs/MaterialInputs.INT.udn
--- a/Engine/Documentation/Source/Engine/Rendering/Materials/MaterialInputs/MaterialInputs.INT.udn
+++ b/Engine/Documentation/Source/Engine/Rendering/Materials/MaterialInputs/MaterialInputs.INT.udn
@@ -7,6 +7,8 @@ parent:Engine/Rendering/Materials
 order:5
 topic-image:Engine/Rendering/Materials/MaterialInputsTopic.png
 type:reference
+tags:Materials
+
 
 Here we take a look at the available inputs when making a Material. By feeding values (by way of constants, parameters, and textures) into these inputs, you can define just about any physical surface you can imagine.
 
diff --git a/Engine/Documentation/Source/Engine/Rendering/Materials/PixelNormalOffset/RefractionPixelNormalOffset.INT.udn b/Engine/Documentation/Source/Engine/Rendering/Materials/PixelNormalOffset/RefractionPixelNormalOffset.INT.udn
new file mode 100644
--- /dev/null
+++ b/Engine/Documentation/Source/Engine/Rendering/Materials/PixelNormalOffset/RefractionPixelNormalOffset.INT.udn
@@ -0,0 +1,99 @@
+Availability: Public
+Title:Refraction Mode: Pixel Normal Offset
+Crumbs:
+Description:Overview of the Refraction mode for Pixel Normal Offset in your Materials.
+Related: Engine/Rendering/Materials
+Related: Engine/Rendering/Materials/HowTo/Refraction
+Version: 4.14
+SkillLevel: Intermediate
+Parent:Engine/Rendering/Materials
+Order:13
+type:overview
+tags:Materials
+topic-image:PNO_Topic.png
+
+[TOC(start:2 end:2)]
+
+In Unreal Engine 4, we now make use of a non-physical model of refraction, **Pixel Normal Offset**. The built-in code, using the physical model of refraction for **Index of Refraction**, is based on how
+light rays refract as they transfer between mediums, which can cause a lot of artifacts as the scene color is being read from off screen. This can be useful for small objects, like glass jars, where 
+you won't notice the effect as much, but is problematic for flat surfaces. **Pixel Normal Offset** enables refraction for these large flat surfaces, like water, where you do not want this constant offset
+that is reading from off screen so much. It uses the vertex normal as a reference and then computes the refraction offset from how different the per-pixel normal is from the vertex normal, allowing these
+flat surfaces not to shift.
+
+## Physical vs Non-Physical Model of Refraction
+
+In the comparisons below, the Refraction Modes for the physical model of **Index of Refraction** and the non-physical model for **Pixel Normal Offset** demonstrate the differences in how the normal is
+read in the material. 
+
+[REGION:note]
+For additional information about the physical model of refraction, and using it with your materials, you can read about it on our page for [Using Refraction](Engine/Rendering/Materials/HowTo/Refraction). 
+[/REGION]
+
+[REGION:imagetable]
+| [INCLUDE:#IORWithoutNormals] | [INCLUDE:#PNOWithoutNormals]  |
+| ![](RefractionSettingIOR.png)(w:400) | ![](RefractionSettingPNO.png)(w:400) |
+|---|---|
+| Physical Model of Refraction: Index of Refraction without Normal Map| Non-Physical Model of Refraction: Pixel Normal Offset without Normal Map |
+[/REGION]
+
+Here you will notice that the image is shifted when using the **Index of Refraction** mode compared to the **Pixel Normal Offset** mode where you do not read from off screen so much. **Index of 
+Refraction** will work without a Normal map plugged into the material, whereas with **Pixel Normal Offset**, if there is no Normal map you will not get any refraction.
+
+
+[REGION:imagetable]
+| [INCLUDE:#IORWithNormals] | [INCLUDE:#PNOWithNormals] |
+| ![](RefractionSettingIOR.png)(w:400) | ![](RefractionSettingPNO.png)(w:400) |
+|---|---|
+| Physical Model with Normal Map | Non-Physical Model with Normal Map |
+[/REGION] 
+
+When you add a Normal map to your material, with refraction and a value greater than 1 plugged into your Refraction input, the normal will translate along the surface when using **Pixel Normal Offset**.
+However, you will notice that with **Index of Refraction**, you will still have an offset that reads from off screen, which is not a desirable effect for these flat surfaces using refraction.
+
+## Final Results
+
+In this example, the refraction amount is adjusted between a value of 1.0, which is no refraction at all, to a value of 2.0, for some refraction along the surface without shifting the image while using
+**Pixel Normal Offset**. 
+
+[OBJECT:EmbeddedVideo]
+	[PARAMLITERAL:width]
+	640
+	[/PARAMLITERAL]
+	[PARAMLITERAL:height]
+	360
+	[/PARAMLITERAL]
+	[PARAMLITERAL:videoid]
+	H8FGf4iZdn0
+	[/PARAMLITERAL]
+[/OBJECT]
+
+
+
+
+<!---
+
+[EXCERPT:IORWithoutNormals]
+[REGION:lightbox]
+[![](IOR_WithoutNormals.png)(w:550)](IOR_WithoutNormals.png)
+[/REGION]
+[/EXCERPT:IORWithoutNormals]
+
+[EXCERPT:PNOWithoutNormals]
+[REGION:lightbox]
+[![](PNO_WithoutNormals.png)(w:550)](PNO_WithoutNormals.png)
+[/REGION]
+[/EXCERPT:PNOWithoutNormals]
+
+[EXCERPT:IORWithNormals]
+[REGION:lightbox]
+[![](IOR_WithNormals.png)(w:550)](IOR_WithNormals.png)
+[/REGION]
+[/EXCERPT:IORWithNormals]
+
+[EXCERPT:PNOWithNormals]
+[REGION:lightbox]
+[![](PNO_WithNormals.png)(w:550)](PNO_WithNormals.png)
+[/REGION]
+[/EXCERPT:PNOWithNormals]
+
+--->
diff --git a/Engine/Documentation/Source/Engine/Rendering/ParticleSystems/Cascade/CascadeEditor.INT.udn b/Engine/Documentation/Source/Engine/Rendering/ParticleSystems/Cascade/CascadeEditor.INT.udn
--- a/Engine/Documentation/Source/Engine/Rendering/ParticleSystems/Cascade/CascadeEditor.INT.udn
+++ b/Engine/Documentation/Source/Engine/Rendering/ParticleSystems/Cascade/CascadeEditor.INT.udn
@@ -3,6 +3,7 @@ Title:Cascade Particle Editor Reference
 Crumbs: %ROOT%, Engine, Engine/Rendering/ParticleSystems
 Description:An overview of the Cascade Editor's UI, breaking down each key area.
 Version: 4.9
+tags:Particles
 
 [VAR:Topic]
 [OBJECT:Topic]
diff --git a/Engine/Documentation/Source/Engine/Rendering/ParticleSystems/LODs/LODs.INT.udn b/Engine/Documentation/Source/Engine/Rendering/ParticleSystems/LODs/LODs.INT.udn
--- a/Engine/Documentation/Source/Engine/Rendering/ParticleSystems/LODs/LODs.INT.udn
+++ b/Engine/Documentation/Source/Engine/Rendering/ParticleSystems/LODs/LODs.INT.udn
@@ -3,6 +3,7 @@ Title: Particle System Level of Detail (LOD)
 Crumbs: %ROOT%, Engine, Engine/Rendering/ParticleSystems
 Description:GPU-simulation of particles that allows for hundreds of thousands of particles at any time.
 Version: 4.9
+tags:Particles
 
 [TOC(start:2 end:3)]
 
diff --git a/Engine/Documentation/Source/Engine/Rendering/ParticleSystems/Optimization/Optimization.INT.udn b/Engine/Documentation/Source/Engine/Rendering/ParticleSystems/Optimization/Optimization.INT.udn
--- a/Engine/Documentation/Source/Engine/Rendering/ParticleSystems/Optimization/Optimization.INT.udn
+++ b/Engine/Documentation/Source/Engine/Rendering/ParticleSystems/Optimization/Optimization.INT.udn
@@ -3,7 +3,7 @@ Title:VFX Optimization Guide
 Crumbs: %ROOT%, Engine, Engine/Rendering/ParticleSystems
 Description:Solving performance issues by optimizing Particle Systems.
 Version: 4.9
-
+tags:Particles
 
 [TOC (start:2 end:3)]
 
diff --git a/Engine/Documentation/Source/Engine/Rendering/ParticleSystems/Reference/Modules/Camera/CameraModules.INT.udn b/Engine/Documentation/Source/Engine/Rendering/ParticleSystems/Reference/Modules/Camera/CameraModules.INT.udn
--- a/Engine/Documentation/Source/Engine/Rendering/ParticleSystems/Reference/Modules/Camera/CameraModules.INT.udn
+++ b/Engine/Documentation/Source/Engine/Rendering/ParticleSystems/Reference/Modules/Camera/CameraModules.INT.udn
@@ -3,7 +3,7 @@ Title:Camera Modules
 Crumbs: %ROOT%, Engine, Engine/Rendering/ParticleSystems, Engine/Rendering/ParticleSystems/Reference
 Description:These modules modify the emitter behavior with respect to the camera.
 Version: 4.9
-
+tags:Particles
 
 
 These modules modify the emitter behavior with respect to the camera.
diff --git a/Engine/Documentation/Source/Engine/Rendering/RenderTargets/BlueprintRenderTargets/BestPractices/BlueprintRenderTargetsBestPractices.INT.udn b/Engine/Documentation/Source/Engine/Rendering/RenderTargets/BlueprintRenderTargets/BestPractices/BlueprintRenderTargetsBestPractices.INT.udn
--- a/Engine/Documentation/Source/Engine/Rendering/RenderTargets/BlueprintRenderTargets/BestPractices/BlueprintRenderTargetsBestPractices.INT.udn
+++ b/Engine/Documentation/Source/Engine/Rendering/RenderTargets/BlueprintRenderTargets/BestPractices/BlueprintRenderTargetsBestPractices.INT.udn
@@ -1,9 +1,9 @@
-availability:docs
+availability:Docs
 title:Blueprint and Render Target Best Practices
 description: A collection of Best Practices to follow when using Blueprints and Render Targets.
 crumbs:
 version:4.13
-Type:topic
+Type:overview
 SkillLevel:Intermediate
 parent:Engine
 order:2
diff --git a/Engine/Documentation/Source/Engine/Rendering/RenderTargets/BlueprintRenderTargets/BlueprintRenderTargets.INT.udn b/Engine/Documentation/Source/Engine/Rendering/RenderTargets/BlueprintRenderTargets/BlueprintRenderTargets.INT.udn
--- a/Engine/Documentation/Source/Engine/Rendering/RenderTargets/BlueprintRenderTargets/BlueprintRenderTargets.INT.udn
+++ b/Engine/Documentation/Source/Engine/Rendering/RenderTargets/BlueprintRenderTargets/BlueprintRenderTargets.INT.udn
@@ -1,4 +1,4 @@
-availability:Docs
+availability:Public
 title: Blueprints and Render Targets
 description: Using Render Targets with Blueprints in Unreal Engine 4.
 crumbs:%ROOT%, Engine, Engine/Rendering, Engine/Rendering/RenderTargets
@@ -8,7 +8,7 @@ parent:Engine/Rendering/RenderTargets
 order:1
 Navigation:topic
 tags: Render Targets
-tags: Blueprint
+tags: Blueprints
 topic-image:Engine/Rendering/RenderTargets/Render_Targets_Overview.png
 
 ![](Engine\Rendering\RenderTargets\BlueprintRenderTarget.png)
diff --git a/Engine/Documentation/Source/Engine/Rendering/RenderTargets/BlueprintRenderTargets/HowTo/BlueprintRenderTargetsHowTos.INT.udn b/Engine/Documentation/Source/Engine/Rendering/RenderTargets/BlueprintRenderTargets/HowTo/BlueprintRenderTargetsHowTos.INT.udn
--- a/Engine/Documentation/Source/Engine/Rendering/RenderTargets/BlueprintRenderTargets/HowTo/BlueprintRenderTargetsHowTos.INT.udn
+++ b/Engine/Documentation/Source/Engine/Rendering/RenderTargets/BlueprintRenderTargets/HowTo/BlueprintRenderTargetsHowTos.INT.udn
@@ -1,4 +1,4 @@
-availability:Docs
+availability:Public
 title: Blueprint and Render Targets How To's
 description: How to guides for using Blueprints to manipulate render targets. 
 crumbs:%ROOT%, Engine, Engine\Rendering\RenderTargets
@@ -7,7 +7,7 @@ Type:how-to
 parent:Engine/Rendering/RenderTargets/BlueprintRenderTargets
 Navigation:landing
 tags: Render Targets
-tags: Blueprint
+tags: Blueprints
 topic-image:Engine/Rendering/RenderTargets/Render_Targets_Best_Practies.png
 
 The following **Blueprint and Render Targets - How To** pages provide step-by-step examples on how to use Blueprints and render targets together in your Unreal Engine 4 (UE4) projects. 
diff --git a/Engine/Documentation/Source/Engine/Rendering/RenderTargets/BlueprintRenderTargets/HowTo/CreatingTextures/1/CTHT_1.INT.udn b/Engine/Documentation/Source/Engine/Rendering/RenderTargets/BlueprintRenderTargets/HowTo/CreatingTextures/1/CTHT_1.INT.udn
--- a/Engine/Documentation/Source/Engine/Rendering/RenderTargets/BlueprintRenderTargets/HowTo/CreatingTextures/1/CTHT_1.INT.udn
+++ b/Engine/Documentation/Source/Engine/Rendering/RenderTargets/BlueprintRenderTargets/HowTo/CreatingTextures/1/CTHT_1.INT.udn
@@ -1,4 +1,4 @@
-availability:docs
+availability:Public
 title:1. Project Setup
 description:Going over the Project Setup for creating Textures using Blueprints and Render Targets.
 crumbs:%ROOT%, Engine, Engine\Rendering\RenderTargets
diff --git a/Engine/Documentation/Source/Engine/Rendering/RenderTargets/BlueprintRenderTargets/HowTo/CreatingTextures/2/CTHT_2.INT.udn b/Engine/Documentation/Source/Engine/Rendering/RenderTargets/BlueprintRenderTargets/HowTo/CreatingTextures/2/CTHT_2.INT.udn
--- a/Engine/Documentation/Source/Engine/Rendering/RenderTargets/BlueprintRenderTargets/HowTo/CreatingTextures/2/CTHT_2.INT.udn
+++ b/Engine/Documentation/Source/Engine/Rendering/RenderTargets/BlueprintRenderTargets/HowTo/CreatingTextures/2/CTHT_2.INT.udn
@@ -1,4 +1,4 @@
-availability:docs
+availability:Public
 title:2. Required Asset Creation & Setup
 description:Going over creating and setting up the required assets.
 crumbs:%ROOT%, Engine, Engine\Rendering\RenderTargets
diff --git a/Engine/Documentation/Source/Engine/Rendering/RenderTargets/BlueprintRenderTargets/HowTo/FluidSurface/FluidSurface.INT.udn b/Engine/Documentation/Source/Engine/Rendering/RenderTargets/BlueprintRenderTargets/HowTo/FluidSurface/FluidSurface.INT.udn
--- a/Engine/Documentation/Source/Engine/Rendering/RenderTargets/BlueprintRenderTargets/HowTo/FluidSurface/FluidSurface.INT.udn
+++ b/Engine/Documentation/Source/Engine/Rendering/RenderTargets/BlueprintRenderTargets/HowTo/FluidSurface/FluidSurface.INT.udn
@@ -1,4 +1,4 @@
-availability:docs
+availability:Public
 title:Creating a Fluid surface with Blueprints and Render Targets
 description:How-To use Blueprints and Render Targets to create realtime fluid surface Actor.
 crumbs:%ROOT%, Engine, Engine\Rendering\RenderTargets
@@ -13,7 +13,7 @@ Related:Engine/Blueprints/
 Related:Engine/Rendering/Materials/
 Related:Engine/Blueprints/UserGuide/UserConstructionScript/
 tags:Render Targets
-tags:Blueprint
+tags:Blueprints
 checkpoint:FSB
 
 ![](FS_Header_Image.png)
