Copying //UE4/Dev-Documentation to //UE4/Samples-Main (Source: //UE4/Dev-Documentation @ 3372893)

diff --git a/Engine/Documentation/Source/Engine/Rendering/RenderTargets/BlueprintRenderTargets/HowTo/HeightFieldPainter/3/HFP3.INT.udn b/Engine/Documentation/Source/Engine/Rendering/RenderTargets/BlueprintRenderTargets/HowTo/HeightFieldPainter/3/HFP3.INT.udn
--- a/Engine/Documentation/Source/Engine/Rendering/RenderTargets/BlueprintRenderTargets/HowTo/HeightFieldPainter/3/HFP3.INT.udn
+++ b/Engine/Documentation/Source/Engine/Rendering/RenderTargets/BlueprintRenderTargets/HowTo/HeightFieldPainter/3/HFP3.INT.udn
@@ -1,8 +1,8 @@
-availability:docs
+availability:Public
 title:3. Blueprint Setup
 description:Going over creating and setting up the height map painter Blueprint.
 crumbs:%ROOT%, Engine, Engine\Rendering\RenderTargets
-version:4.13
+version:4.14
 type:multi-step
 SkillLevel:Intermediate
 parent:Engine/Rendering/RenderTargets/BlueprintRenderTargets/HowTo/HeightFieldPainter
@@ -60,6 +60,45 @@ The Construction script is used to create and assign a 2D render target in addit
 [/REGION]
 
 ## Trace from Camera Function Setup
+
+There are two ways you can implement this function. The first method requires a bit of math and works primarily on planes or other flat surfaces. The second method uses the **Find Collisoin UV** node which requires more memory to use, but will work across any **Static Mesh Component**, regardless of shape (UV seams will still cause artifcats). You can choose your implementation method below.
+
+%Globals:MethodSelect%
+
+[OBJECT:ToggleButtonContent]
+	[PARAMLITERAL:category]
+	Method
+	[/PARAMLITERAL]
+	[PARAMLITERAL:id]
+	method01
+	[/PARAMLITERAL]
+	[PARAMLITERAL:active]
+	active_button_content
+	[/PARAMLITERAL]
+	[PARAM:content]
+	[INCLUDE:#m1]
+	[/PARAM]
+[/OBJECT]
+[OBJECT:ToggleButtonContent]
+	[PARAMLITERAL:category]
+	Method
+	[/PARAMLITERAL]
+	[PARAMLITERAL:id]
+	method02
+	[/PARAMLITERAL]
+	[PARAMLITERAL:active]
+	[/PARAMLITERAL]
+	[PARAM:content]
+	[INCLUDE:#m2]
+	[/PARAM]
+[/OBJECT]
+
+[NAV]
+
+<!--
+Method Split Here
+
+[EXCERPT:m1]
 The Trace from Camera function will be used to determine the location on the Static Mesh that should have its height increased. It determines this location by casting a ray from the center of the player's camera into the world. When that ray intersects with a Static Mesh, the location that it intersects will be the location that is raised. To set up the `Trace from Camera` function, you will need to do the following.
 
 1. Click the small (white) plus sign on the **Functions** tab to create a new function and name it **TraceFromCamera**.
@@ -75,10 +114,10 @@ The Trace from Camera function will be used to determine the location on the Sta
 	[region:note]
 	Do not forget to Compile and Save your Blueprint so that it can be used.
 	[/region]
-
+	
 ## Event Graph Setup
-Now that all of the need parts have been created and setup, it is now time to connect all of the various parts together in the Event Graph. You can copy the completed Blueprint code into your Blueprint by hovering your mouse over the image below and then clicking on the  **Copy Node Graph** option, copying-and-pasting the provided Blueprint code into your Height Field Painter Blueprint.
 
+Now that all of the need parts have been created and setup, it is now time to connect all of the various parts together in the Event Graph. You can copy the completed Blueprint code into your Blueprint by hovering your mouse over the image below and then clicking on the  **Copy Node Graph** option, copying-and-pasting the provided Blueprint code into your Height Field Painter Blueprint.
 
 [REGION:fullwidth]
 [INCLUDE:Engine\Rendering\RenderTargets\BlueprintRenderTargets\HowTo\HeightFieldPainter\3\#EG]
@@ -92,9 +131,47 @@ When you copy-paste the Blueprint code, the **Event Begin Play** node will not b
 
 
 Now that the required Blueprints have been setup, in the next section we will go over how to use the height map painter inside UE4.
+[/EXCERPT]
 
-[NAV]
+[EXCERPT:m2]
+
+[REGION:note]
+You have to enable **Support UV from Hit Result** feature in the editor to use this method. See [](Engine/Physics/Tracing/Overview#EnablingUVCoordinatesfromTrace) for information on how to do this.
+[/REGION]
+
+1. Click the small (white) plus sign on the **Functions** tab to create a new function and name it **TraceFromCamera**.
+
+	![](Engine\Rendering\RenderTargets\BlueprintRenderTargets\HowTo\HeightFieldPainter\HFP_BP_Setup_04.png)
 
+1. Open the TraceFromCamera function and copy the following Blueprint code into your Blueprint by hovering your mouse over the image below and then clicking on the **Copy Expression Graph** option, selecting the code, copying it, and then pasting it into the TraceFromCamera function.
+
+	[REGION:fullwidth]
+	[INCLUDE:Engine\Rendering\RenderTargets\BlueprintRenderTargets\HowTo\HeightFieldPainter\3\#FunctionUV]
+	[/REGION]
+
+	[region:note]
+	Do not forget to Compile and Save your Blueprint so that it can be used.
+	[/region]
+
+		
+## Event Graph Setup
+
+Now that all of the need parts have been created and setup, it is now time to connect all of the various parts together in the Event Graph. You can copy the completed Blueprint code into your Blueprint by hovering your mouse over the image below and then clicking on the  **Copy Node Graph** option, copying-and-pasting the provided Blueprint code into your Height Field Painter Blueprint.
+
+[REGION:fullwidth]
+[INCLUDE:Engine\Rendering\RenderTargets\BlueprintRenderTargets\HowTo\HeightFieldPainter\3\#FunctionUV02]
+[/REGION]
+
+[region:note]
+When you copy-paste the Blueprint code, the **Event Begin Play** node will not be added to the graph. Make sure that you add the Event Begin Play node and connect it to the **Enable Input** input. 
+
+![](Engine\Rendering\RenderTargets\BlueprintRenderTargets\HowTo\HeightFieldPainter\HFP_BP_Setup_08.png)
+[/region]
+
+
+Now that the required Blueprints have been setup, in the next section we will go over how to use the height map painter inside UE4.
+[/EXCERPT]
+-->
 
 <!--
 [EXCERPT:Function]	
@@ -845,4 +922,572 @@ Now that the required Blueprints have been setup, in the next section we will go
 	[/PARAM]
 [/OBJECT]
 [/EXCERPT:EG]
+
+
+[EXCERPT:FunctionUV]	
+[OBJECT:Snippet]
+	[PARAMLITERAL:id]
+		TraceFromCameraUV
+	[/PARAMLITERAL]
+	[PARAMLITERAL:type]
+		blueprint
+	[/PARAMLITERAL]
+	[PARAM:image]
+		![](Engine\Rendering\RenderTargets\BlueprintRenderTargets\HowTo\HeightFieldPainter\HFP_VP_Setup_UV_01.png)
+	[/PARAM]
+	[PARAM:snippet]
+		Begin Object Class=K2Node_FunctionEntry Name="K2Node_FunctionEntry_48"
+		   ExtraFlags=201457664
+		   SignatureName="TraceFromCamera"
+		   bIsEditable=True
+		   NodePosX=-592
+		   NodePosY=-160
+		   NodeGuid=8314002942029CD4D7ECFC970D017D4D
+		   CustomProperties Pin (PinId=25C5CAE14F04A75F98A7739DD80E5B65,PinName="then",Direction="EGPD_Output",PinType.PinCategory="exec",PinType.PinSubCategory="",PinType.PinSubCategoryObject=None,PinType.PinSubCategoryMemberReference=(),PinType.PinValueType=(),PinType.bIsMap=False,PinType.bIsSet=False,PinType.bIsArray=False,PinType.bIsReference=False,PinType.bIsConst=False,PinType.bIsWeakPointer=False,LinkedTo=(K2Node_CallFunction_65598 B3AD098C4E9F315318BE2C9A781CFD16,),PersistentGuid=00000000000000000000000000000000,bHidden=False,bNotConnectable=False,bDefaultValueIsReadOnly=False,bDefaultValueIsIgnored=False,bAdvancedView=False,)
+		End Object
+		Begin Object Class=K2Node_MacroInstance Name="K2Node_MacroInstance_65"
+		   MacroGraphReference=(MacroGraph=EdGraph'/Engine/EditorBlueprintResources/StandardMacros.StandardMacros:IsValid',GraphBlueprint=Blueprint'/Engine/EditorBlueprintResources/StandardMacros.StandardMacros',GraphGuid=64422BCD430703FF5CAEA8B79A32AA65)
+		   NodePosX=1504
+		   NodePosY=-208
+		   ErrorType=4
+		   ErrorMsg="Info The execution path doesn\'t end with a return node.  Is Not Valid "
+		   NodeGuid=C322C376456F56B16A3D8496CE1A0A38
+		   CustomProperties Pin (PinId=54DCDF6A4C18DFFBF71F54AE24DA40C0,PinName="exec",PinType.PinCategory="exec",PinType.PinSubCategory="",PinType.PinSubCategoryObject=None,PinType.PinSubCategoryMemberReference=(),PinType.PinValueType=(),PinType.bIsMap=False,PinType.bIsSet=False,PinType.bIsArray=False,PinType.bIsReference=False,PinType.bIsConst=False,PinType.bIsWeakPointer=False,LinkedTo=(K2Node_IfThenElse_161 8DD9A5BC40D019A1AFD6F284F8C15FFD,),PersistentGuid=00000000000000000000000000000000,bHidden=False,bNotConnectable=False,bDefaultValueIsReadOnly=False,bDefaultValueIsIgnored=False,bAdvancedView=False,)
+		   CustomProperties Pin (PinId=C3FB28AC47AAD2B5998C2AA9D4862C2E,PinName="InputObject",PinType.PinCategory="object",PinType.PinSubCategory="",PinType.PinSubCategoryObject=Class'/Script/CoreUObject.Object',PinType.PinSubCategoryMemberReference=(),PinType.PinValueType=(),PinType.bIsMap=False,PinType.bIsSet=False,PinType.bIsArray=False,PinType.bIsReference=False,PinType.bIsConst=False,PinType.bIsWeakPointer=False,LinkedTo=(K2Node_CallFunction_65759 F1D6985542953219C7EF09A3FBFA8AF0,),PersistentGuid=00000000000000000000000000000000,bHidden=False,bNotConnectable=False,bDefaultValueIsReadOnly=False,bDefaultValueIsIgnored=False,bAdvancedView=False,)
+		   CustomProperties Pin (PinId=E75DA26F4EBE18C63E3CEDA6B1D88DC6,PinName="Is Valid",Direction="EGPD_Output",PinType.PinCategory="exec",PinType.PinSubCategory="",PinType.PinSubCategoryObject=None,PinType.PinSubCategoryMemberReference=(),PinType.PinValueType=(),PinType.bIsMap=False,PinType.bIsSet=False,PinType.bIsArray=False,PinType.bIsReference=False,PinType.bIsConst=False,PinType.bIsWeakPointer=False,LinkedTo=(K2Node_CallFunction_311 F45D897646865C99837DAD9886C0860F,),PersistentGuid=00000000000000000000000000000000,bHidden=False,bNotConnectable=False,bDefaultValueIsReadOnly=False,bDefaultValueIsIgnored=False,bAdvancedView=False,)
+		   CustomProperties Pin (PinId=7426FEC54F3C31FE985AE088ACC9F56E,PinName="Is Not Valid",Direction="EGPD_Output",PinType.PinCategory="exec",PinType.PinSubCategory="",PinType.PinSubCategoryObject=None,PinType.PinSubCategoryMemberReference=(),PinType.PinValueType=(),PinType.bIsMap=False,PinType.bIsSet=False,PinType.bIsArray=False,PinType.bIsReference=False,PinType.bIsConst=False,PinType.bIsWeakPointer=False,PersistentGuid=00000000000000000000000000000000,bHidden=False,bNotConnectable=False,bDefaultValueIsReadOnly=False,bDefaultValueIsIgnored=False,bAdvancedView=False,)
+		End Object
+		Begin Object Class=K2Node_CallFunction Name="K2Node_CallFunction_64853"
+		   bIsPureFunc=True
+		   FunctionReference=(MemberParent=Class'/Script/Engine.GameplayStatics',MemberName="GetPlayerCameraManager")
+		   NodePosX=-592
+		   NodePosY=48
+		   NodeGuid=B2E54DA44E93145C2F65BE9CB3553A6B
+		   CustomProperties Pin (PinId=73A2796F477586504925ACB63B1C3373,PinName="self",PinFriendlyName=NSLOCTEXT("K2Node", "Target", "Target"),PinToolTip="Target\nGameplay Statics Reference",PinType.PinCategory="object",PinType.PinSubCategory="",PinType.PinSubCategoryObject=Class'/Script/Engine.GameplayStatics',PinType.PinSubCategoryMemberReference=(),PinType.PinValueType=(),PinType.bIsMap=False,PinType.bIsSet=False,PinType.bIsArray=False,PinType.bIsReference=False,PinType.bIsConst=False,PinType.bIsWeakPointer=False,DefaultObject="/Script/Engine.Default__GameplayStatics",PersistentGuid=00000000000000000000000000000000,bHidden=True,bNotConnectable=False,bDefaultValueIsReadOnly=False,bDefaultValueIsIgnored=False,bAdvancedView=False,)
+		   CustomProperties Pin (PinId=A29AA4884E79B1B2CBFCAE93A185BFA1,PinName="WorldContextObject",PinToolTip="World Context Object\nObject Reference",PinType.PinCategory="object",PinType.PinSubCategory="",PinType.PinSubCategoryObject=Class'/Script/CoreUObject.Object',PinType.PinSubCategoryMemberReference=(),PinType.PinValueType=(),PinType.bIsMap=False,PinType.bIsSet=False,PinType.bIsArray=False,PinType.bIsReference=False,PinType.bIsConst=True,PinType.bIsWeakPointer=False,DefaultValue="0",PersistentGuid=00000000000000000000000000000000,bHidden=True,bNotConnectable=False,bDefaultValueIsReadOnly=False,bDefaultValueIsIgnored=False,bAdvancedView=False,)
+		   CustomProperties Pin (PinId=F04D1A5446F837775E768082313D9559,PinName="PlayerIndex",PinToolTip="Player Index\nInteger",PinType.PinCategory="int",PinType.PinSubCategory="",PinType.PinSubCategoryObject=None,PinType.PinSubCategoryMemberReference=(),PinType.PinValueType=(),PinType.bIsMap=False,PinType.bIsSet=False,PinType.bIsArray=False,PinType.bIsReference=False,PinType.bIsConst=False,PinType.bIsWeakPointer=False,DefaultValue="0",AutogeneratedDefaultValue="0",PersistentGuid=00000000000000000000000000000000,bHidden=False,bNotConnectable=False,bDefaultValueIsReadOnly=False,bDefaultValueIsIgnored=False,bAdvancedView=False,)
+		   CustomProperties Pin (PinId=88B4C49D414D557C8728F39A3967490F,PinName="ReturnValue",PinToolTip="Return Value\nPlayer Camera Manager Reference",Direction="EGPD_Output",PinType.PinCategory="object",PinType.PinSubCategory="",PinType.PinSubCategoryObject=Class'/Script/Engine.PlayerCameraManager',PinType.PinSubCategoryMemberReference=(),PinType.PinValueType=(),PinType.bIsMap=False,PinType.bIsSet=False,PinType.bIsArray=False,PinType.bIsReference=False,PinType.bIsConst=False,PinType.bIsWeakPointer=False,LinkedTo=(K2Node_CallFunction_64919 136B71094B407B9186881DB137E50CB0,K2Node_CallFunction_65094 451585FC4960B7F55857FDB8A5A37392,),PersistentGuid=00000000000000000000000000000000,bHidden=False,bNotConnectable=False,bDefaultValueIsReadOnly=False,bDefaultValueIsIgnored=False,bAdvancedView=False,)
+		End Object
+		Begin Object Class=K2Node_CallFunction Name="K2Node_CallFunction_64919"
+		   bIsPureFunc=True
+		   bIsConstFunc=True
+		   FunctionReference=(MemberName="K2_GetActorLocation",bSelfContext=True)
+		   NodePosX=80
+		   NodePosY=32
+		   NodeGuid=5FB89855416B660DAD5F4F8FC39DB92E
+		   CustomProperties Pin (PinId=136B71094B407B9186881DB137E50CB0,PinName="self",PinFriendlyName=NSLOCTEXT("K2Node", "Target", "Target"),PinToolTip="Target\nActor Reference",PinType.PinCategory="object",PinType.PinSubCategory="",PinType.PinSubCategoryObject=Class'/Script/Engine.Actor',PinType.PinSubCategoryMemberReference=(),PinType.PinValueType=(),PinType.bIsMap=False,PinType.bIsSet=False,PinType.bIsArray=False,PinType.bIsReference=False,PinType.bIsConst=False,PinType.bIsWeakPointer=False,LinkedTo=(K2Node_CallFunction_64853 88B4C49D414D557C8728F39A3967490F,),PersistentGuid=00000000000000000000000000000000,bHidden=False,bNotConnectable=False,bDefaultValueIsReadOnly=False,bDefaultValueIsIgnored=False,bAdvancedView=False,)
+		   CustomProperties Pin (PinId=7C3AB90A4E8F92F9FAE0709612FDD604,PinName="ReturnValue",PinToolTip="Return Value\nVector ",Direction="EGPD_Output",PinType.PinCategory="struct",PinType.PinSubCategory="",PinType.PinSubCategoryObject=ScriptStruct'/Script/CoreUObject.Vector',PinType.PinSubCategoryMemberReference=(),PinType.PinValueType=(),PinType.bIsMap=False,PinType.bIsSet=False,PinType.bIsArray=False,PinType.bIsReference=False,PinType.bIsConst=False,PinType.bIsWeakPointer=False,DefaultValue="0, 0, 0",AutogeneratedDefaultValue="0, 0, 0",LinkedTo=(K2Node_CallFunction_65598 85DD87AD4631A8DC79F17588BD6ADD07,K2Node_CommutativeAssociativeBinaryOperator_1357 C2F85EED4A8F818C3CE355A3C56691F5,),PersistentGuid=00000000000000000000000000000000,bHidden=False,bNotConnectable=False,bDefaultValueIsReadOnly=False,bDefaultValueIsIgnored=False,bAdvancedView=False,)
+		End Object
+		Begin Object Class=K2Node_CallFunction Name="K2Node_CallFunction_65094"
+		   bIsPureFunc=True
+		   bIsConstFunc=True
+		   FunctionReference=(MemberName="K2_GetActorRotation",bSelfContext=True)
+		   NodePosX=-320
+		   NodePosY=144
+		   NodeGuid=764278254E138B733D149D9BA7A38113
+		   CustomProperties Pin (PinId=451585FC4960B7F55857FDB8A5A37392,PinName="self",PinFriendlyName=NSLOCTEXT("K2Node", "Target", "Target"),PinToolTip="Target\nActor Reference",PinType.PinCategory="object",PinType.PinSubCategory="",PinType.PinSubCategoryObject=Class'/Script/Engine.Actor',PinType.PinSubCategoryMemberReference=(),PinType.PinValueType=(),PinType.bIsMap=False,PinType.bIsSet=False,PinType.bIsArray=False,PinType.bIsReference=False,PinType.bIsConst=False,PinType.bIsWeakPointer=False,LinkedTo=(K2Node_CallFunction_64853 88B4C49D414D557C8728F39A3967490F,),PersistentGuid=00000000000000000000000000000000,bHidden=False,bNotConnectable=False,bDefaultValueIsReadOnly=False,bDefaultValueIsIgnored=False,bAdvancedView=False,)
+		   CustomProperties Pin (PinId=3048B8A64F99A2395A2A78B01F534663,PinName="ReturnValue",PinToolTip="Return Value\nRotator ",Direction="EGPD_Output",PinType.PinCategory="struct",PinType.PinSubCategory="",PinType.PinSubCategoryObject=ScriptStruct'/Script/CoreUObject.Rotator',PinType.PinSubCategoryMemberReference=(),PinType.PinValueType=(),PinType.bIsMap=False,PinType.bIsSet=False,PinType.bIsArray=False,PinType.bIsReference=False,PinType.bIsConst=False,PinType.bIsWeakPointer=False,DefaultValue="0, 0, 0",AutogeneratedDefaultValue="0, 0, 0",LinkedTo=(K2Node_CallFunction_65130 44D474914E7F590449BEE0850F8024CE,),PersistentGuid=00000000000000000000000000000000,bHidden=False,bNotConnectable=False,bDefaultValueIsReadOnly=False,bDefaultValueIsIgnored=False,bAdvancedView=False,)
+		End Object
+		Begin Object Class=K2Node_CallFunction Name="K2Node_CallFunction_65130"
+		   bIsPureFunc=True
+		   FunctionReference=(MemberParent=Class'/Script/Engine.KismetMathLibrary',MemberName="Conv_RotatorToVector")
+		   NodePosX=-96
+		   NodePosY=144
+		   NodeGuid=2F75275144CB740CE8315DA1FE224709
+		   CustomProperties Pin (PinId=06AE3D444E2D6A5652B7C291690B1A65,PinName="self",PinFriendlyName=NSLOCTEXT("K2Node", "Target", "Target"),PinToolTip="Target\nKismet Math Library Reference",PinType.PinCategory="object",PinType.PinSubCategory="",PinType.PinSubCategoryObject=Class'/Script/Engine.KismetMathLibrary',PinType.PinSubCategoryMemberReference=(),PinType.PinValueType=(),PinType.bIsMap=False,PinType.bIsSet=False,PinType.bIsArray=False,PinType.bIsReference=False,PinType.bIsConst=False,PinType.bIsWeakPointer=False,DefaultObject="/Script/Engine.Default__KismetMathLibrary",PersistentGuid=00000000000000000000000000000000,bHidden=True,bNotConnectable=False,bDefaultValueIsReadOnly=False,bDefaultValueIsIgnored=False,bAdvancedView=False,)
+		   CustomProperties Pin (PinId=44D474914E7F590449BEE0850F8024CE,PinName="InRot",PinToolTip="In Rot\nRotator ",PinType.PinCategory="struct",PinType.PinSubCategory="",PinType.PinSubCategoryObject=ScriptStruct'/Script/CoreUObject.Rotator',PinType.PinSubCategoryMemberReference=(),PinType.PinValueType=(),PinType.bIsMap=False,PinType.bIsSet=False,PinType.bIsArray=False,PinType.bIsReference=False,PinType.bIsConst=False,PinType.bIsWeakPointer=False,DefaultValue="0, 0, 0",AutogeneratedDefaultValue="0, 0, 0",LinkedTo=(K2Node_CallFunction_65094 3048B8A64F99A2395A2A78B01F534663,),PersistentGuid=00000000000000000000000000000000,bHidden=False,bNotConnectable=False,bDefaultValueIsReadOnly=False,bDefaultValueIsIgnored=False,bAdvancedView=False,)
+		   CustomProperties Pin (PinId=6A4437574CCDB781A7E425AF26B29860,PinName="ReturnValue",PinToolTip="Return Value\nVector ",Direction="EGPD_Output",PinType.PinCategory="struct",PinType.PinSubCategory="",PinType.PinSubCategoryObject=ScriptStruct'/Script/CoreUObject.Vector',PinType.PinSubCategoryMemberReference=(),PinType.PinValueType=(),PinType.bIsMap=False,PinType.bIsSet=False,PinType.bIsArray=False,PinType.bIsReference=False,PinType.bIsConst=False,PinType.bIsWeakPointer=False,DefaultValue="0, 0, 0",AutogeneratedDefaultValue="0, 0, 0",LinkedTo=(K2Node_CallFunction_65459 CE2B40EF47E2476D00764EA442CB2F0D,),PersistentGuid=00000000000000000000000000000000,bHidden=False,bNotConnectable=False,bDefaultValueIsReadOnly=False,bDefaultValueIsIgnored=False,bAdvancedView=False,)
+		End Object
+		Begin Object Class=K2Node_CallFunction Name="K2Node_CallFunction_65349"
+		   bIsPureFunc=True
+		   FunctionReference=(MemberParent=Class'/Script/Engine.KismetMathLibrary',MemberName="Multiply_VectorFloat")
+		   NodePosX=304
+		   NodePosY=192
+		   NodeGuid=EDA3065247E4DB6DB6A514B1FE898420
+		   CustomProperties Pin (PinId=ABB067854B9FDB401D7AF2A71639E8A9,PinName="self",PinFriendlyName=NSLOCTEXT("K2Node", "Target", "Target"),PinType.PinCategory="object",PinType.PinSubCategory="",PinType.PinSubCategoryObject=Class'/Script/Engine.KismetMathLibrary',PinType.PinSubCategoryMemberReference=(),PinType.PinValueType=(),PinType.bIsMap=False,PinType.bIsSet=False,PinType.bIsArray=False,PinType.bIsReference=False,PinType.bIsConst=False,PinType.bIsWeakPointer=False,DefaultObject="/Script/Engine.Default__KismetMathLibrary",PersistentGuid=00000000000000000000000000000000,bHidden=True,bNotConnectable=False,bDefaultValueIsReadOnly=False,bDefaultValueIsIgnored=False,bAdvancedView=False,)
+		   CustomProperties Pin (PinId=AE6D33C2476153AC065E9095CB766952,PinName="A",PinType.PinCategory="struct",PinType.PinSubCategory="",PinType.PinSubCategoryObject=ScriptStruct'/Script/CoreUObject.Vector',PinType.PinSubCategoryMemberReference=(),PinType.PinValueType=(),PinType.bIsMap=False,PinType.bIsSet=False,PinType.bIsArray=False,PinType.bIsReference=False,PinType.bIsConst=False,PinType.bIsWeakPointer=False,DefaultValue="0, 0, 0",AutogeneratedDefaultValue="0, 0, 0",LinkedTo=(K2Node_CallFunction_65459 182E27F347849D96EFDA44BC8BAF228B,),PersistentGuid=00000000000000000000000000000000,bHidden=False,bNotConnectable=False,bDefaultValueIsReadOnly=False,bDefaultValueIsIgnored=False,bAdvancedView=False,)
+		   CustomProperties Pin (PinId=2C5718584516424BB52CF6A79E702026,PinName="B",PinType.PinCategory="float",PinType.PinSubCategory="",PinType.PinSubCategoryObject=None,PinType.PinSubCategoryMemberReference=(),PinType.PinValueType=(),PinType.bIsMap=False,PinType.bIsSet=False,PinType.bIsArray=False,PinType.bIsReference=False,PinType.bIsConst=False,PinType.bIsWeakPointer=False,DefaultValue="0.0",AutogeneratedDefaultValue="0.0",LinkedTo=(K2Node_CommutativeAssociativeBinaryOperator_4227 0FAF44A34397D0F87EB02B88E8A653E3,),PersistentGuid=00000000000000000000000000000000,bHidden=False,bNotConnectable=False,bDefaultValueIsReadOnly=False,bDefaultValueIsIgnored=False,bAdvancedView=False,)
+		   CustomProperties Pin (PinId=40870E1F404EA24CCC3295A98D8A13B6,PinName="ReturnValue",Direction="EGPD_Output",PinType.PinCategory="struct",PinType.PinSubCategory="",PinType.PinSubCategoryObject=ScriptStruct'/Script/CoreUObject.Vector',PinType.PinSubCategoryMemberReference=(),PinType.PinValueType=(),PinType.bIsMap=False,PinType.bIsSet=False,PinType.bIsArray=False,PinType.bIsReference=False,PinType.bIsConst=False,PinType.bIsWeakPointer=False,DefaultValue="0, 0, 0",AutogeneratedDefaultValue="0, 0, 0",LinkedTo=(K2Node_CommutativeAssociativeBinaryOperator_1357 4410C89D42A350AA1AF9EE830D6A174E,),PersistentGuid=00000000000000000000000000000000,bHidden=False,bNotConnectable=False,bDefaultValueIsReadOnly=False,bDefaultValueIsIgnored=False,bAdvancedView=False,)
+		End Object
+		Begin Object Class=K2Node_CallFunction Name="K2Node_CallFunction_65459"
+		   bIsPureFunc=True
+		   FunctionReference=(MemberParent=Class'/Script/Engine.KismetMathLibrary',MemberName="Normal")
+		   NodePosX=112
+		   NodePosY=144
+		   NodeGuid=81E6DA3B42012ED0D2A389A6817EAED3
+		   CustomProperties Pin (PinId=E124F3A24616014EE593778FBFC7A74F,PinName="self",PinFriendlyName=NSLOCTEXT("K2Node", "Target", "Target"),PinType.PinCategory="object",PinType.PinSubCategory="",PinType.PinSubCategoryObject=Class'/Script/Engine.KismetMathLibrary',PinType.PinSubCategoryMemberReference=(),PinType.PinValueType=(),PinType.bIsMap=False,PinType.bIsSet=False,PinType.bIsArray=False,PinType.bIsReference=False,PinType.bIsConst=False,PinType.bIsWeakPointer=False,DefaultObject="/Script/Engine.Default__KismetMathLibrary",PersistentGuid=00000000000000000000000000000000,bHidden=True,bNotConnectable=False,bDefaultValueIsReadOnly=False,bDefaultValueIsIgnored=False,bAdvancedView=False,)
+		   CustomProperties Pin (PinId=CE2B40EF47E2476D00764EA442CB2F0D,PinName="A",PinType.PinCategory="struct",PinType.PinSubCategory="",PinType.PinSubCategoryObject=ScriptStruct'/Script/CoreUObject.Vector',PinType.PinSubCategoryMemberReference=(),PinType.PinValueType=(),PinType.bIsMap=False,PinType.bIsSet=False,PinType.bIsArray=False,PinType.bIsReference=False,PinType.bIsConst=False,PinType.bIsWeakPointer=False,DefaultValue="0, 0, 0",AutogeneratedDefaultValue="0, 0, 0",LinkedTo=(K2Node_CallFunction_65130 6A4437574CCDB781A7E425AF26B29860,),PersistentGuid=00000000000000000000000000000000,bHidden=False,bNotConnectable=False,bDefaultValueIsReadOnly=False,bDefaultValueIsIgnored=False,bAdvancedView=False,)
+		   CustomProperties Pin (PinId=182E27F347849D96EFDA44BC8BAF228B,PinName="ReturnValue",Direction="EGPD_Output",PinType.PinCategory="struct",PinType.PinSubCategory="",PinType.PinSubCategoryObject=ScriptStruct'/Script/CoreUObject.Vector',PinType.PinSubCategoryMemberReference=(),PinType.PinValueType=(),PinType.bIsMap=False,PinType.bIsSet=False,PinType.bIsArray=False,PinType.bIsReference=False,PinType.bIsConst=False,PinType.bIsWeakPointer=False,DefaultValue="0, 0, 0",AutogeneratedDefaultValue="0, 0, 0",LinkedTo=(K2Node_CallFunction_65349 AE6D33C2476153AC065E9095CB766952,),PersistentGuid=00000000000000000000000000000000,bHidden=False,bNotConnectable=False,bDefaultValueIsReadOnly=False,bDefaultValueIsIgnored=False,bAdvancedView=False,)
+		End Object
+		Begin Object Class=K2Node_CallFunction Name="K2Node_CallFunction_65598"
+		   FunctionReference=(MemberParent=Class'/Script/Engine.KismetSystemLibrary',MemberName="LineTraceSingle_NEW")
+		   NodePosX=624
+		   NodePosY=-160
+		   NodeGuid=AD7993D14CDD0F724891519FCC1C8292
+		   AdvancedPinDisplay=Hidden
+		   CustomProperties Pin (PinId=B3AD098C4E9F315318BE2C9A781CFD16,PinName="execute",PinToolTip="\nExec",PinType.PinCategory="exec",PinType.PinSubCategory="",PinType.PinSubCategoryObject=None,PinType.PinSubCategoryMemberReference=(),PinType.PinValueType=(),PinType.bIsMap=False,PinType.bIsSet=False,PinType.bIsArray=False,PinType.bIsReference=False,PinType.bIsConst=False,PinType.bIsWeakPointer=False,LinkedTo=(K2Node_FunctionEntry_48 25C5CAE14F04A75F98A7739DD80E5B65,),PersistentGuid=00000000000000000000000000000000,bHidden=False,bNotConnectable=False,bDefaultValueIsReadOnly=False,bDefaultValueIsIgnored=False,bAdvancedView=False,)
+		   CustomProperties Pin (PinId=E001912C48BB5BEC656460843FE151FA,PinName="then",PinToolTip="\nExec",Direction="EGPD_Output",PinType.PinCategory="exec",PinType.PinSubCategory="",PinType.PinSubCategoryObject=None,PinType.PinSubCategoryMemberReference=(),PinType.PinValueType=(),PinType.bIsMap=False,PinType.bIsSet=False,PinType.bIsArray=False,PinType.bIsReference=False,PinType.bIsConst=False,PinType.bIsWeakPointer=False,LinkedTo=(K2Node_IfThenElse_161 66D340DC4A52DC8D727C43ABA6E818C4,),PersistentGuid=00000000000000000000000000000000,bHidden=False,bNotConnectable=False,bDefaultValueIsReadOnly=False,bDefaultValueIsIgnored=False,bAdvancedView=False,)
+		   CustomProperties Pin (PinId=941FA0A3430F4462DB5E30B8A5A81DF4,PinName="self",PinFriendlyName=NSLOCTEXT("K2Node", "Target", "Target"),PinToolTip="Target\nKismet System Library Reference",PinType.PinCategory="object",PinType.PinSubCategory="",PinType.PinSubCategoryObject=Class'/Script/Engine.KismetSystemLibrary',PinType.PinSubCategoryMemberReference=(),PinType.PinValueType=(),PinType.bIsMap=False,PinType.bIsSet=False,PinType.bIsArray=False,PinType.bIsReference=False,PinType.bIsConst=False,PinType.bIsWeakPointer=False,DefaultObject="/Script/Engine.Default__KismetSystemLibrary",PersistentGuid=00000000000000000000000000000000,bHidden=True,bNotConnectable=False,bDefaultValueIsReadOnly=False,bDefaultValueIsIgnored=False,bAdvancedView=False,)
+		   CustomProperties Pin (PinId=6F242C374EDACE5FC62C50B118086553,PinName="WorldContextObject",PinToolTip="World Context Object\nObject Reference",PinType.PinCategory="object",PinType.PinSubCategory="",PinType.PinSubCategoryObject=Class'/Script/CoreUObject.Object',PinType.PinSubCategoryMemberReference=(),PinType.PinValueType=(),PinType.bIsMap=False,PinType.bIsSet=False,PinType.bIsArray=False,PinType.bIsReference=False,PinType.bIsConst=False,PinType.bIsWeakPointer=False,DefaultValue="0",PersistentGuid=00000000000000000000000000000000,bHidden=True,bNotConnectable=False,bDefaultValueIsReadOnly=False,bDefaultValueIsIgnored=False,bAdvancedView=False,)
+		   CustomProperties Pin (PinId=85DD87AD4631A8DC79F17588BD6ADD07,PinName="Start",PinToolTip="Start\nVector \n\nStart of line segment.",PinType.PinCategory="struct",PinType.PinSubCategory="",PinType.PinSubCategoryObject=ScriptStruct'/Script/CoreUObject.Vector',PinType.PinSubCategoryMemberReference=(),PinType.PinValueType=(),PinType.bIsMap=False,PinType.bIsSet=False,PinType.bIsArray=False,PinType.bIsReference=False,PinType.bIsConst=True,PinType.bIsWeakPointer=False,DefaultValue="0, 0, 0",AutogeneratedDefaultValue="0, 0, 0",LinkedTo=(K2Node_CallFunction_64919 7C3AB90A4E8F92F9FAE0709612FDD604,),PersistentGuid=00000000000000000000000000000000,bHidden=False,bNotConnectable=False,bDefaultValueIsReadOnly=False,bDefaultValueIsIgnored=False,bAdvancedView=False,)
+		   CustomProperties Pin (PinId=54C3C6D04942C8BC62F2D696DFF0A12C,PinName="End",PinToolTip="End\nVector \n\nEnd of line segment.",PinType.PinCategory="struct",PinType.PinSubCategory="",PinType.PinSubCategoryObject=ScriptStruct'/Script/CoreUObject.Vector',PinType.PinSubCategoryMemberReference=(),PinType.PinValueType=(),PinType.bIsMap=False,PinType.bIsSet=False,PinType.bIsArray=False,PinType.bIsReference=False,PinType.bIsConst=True,PinType.bIsWeakPointer=False,DefaultValue="0, 0, 0",AutogeneratedDefaultValue="0, 0, 0",LinkedTo=(K2Node_CommutativeAssociativeBinaryOperator_1357 1778AF2B4812E5B635D0DC915DFAA735,),PersistentGuid=00000000000000000000000000000000,bHidden=False,bNotConnectable=False,bDefaultValueIsReadOnly=False,bDefaultValueIsIgnored=False,bAdvancedView=False,)
+		   CustomProperties Pin (PinId=FAD144344A968D69191168B056747AB6,PinName="TraceChannel",PinToolTip="Trace Channel\nETraceTypeQuery Enum",PinType.PinCategory="byte",PinType.PinSubCategory="",PinType.PinSubCategoryObject=Enum'/Script/Engine.ETraceTypeQuery',PinType.PinSubCategoryMemberReference=(),PinType.PinValueType=(),PinType.bIsMap=False,PinType.bIsSet=False,PinType.bIsArray=False,PinType.bIsReference=False,PinType.bIsConst=False,PinType.bIsWeakPointer=False,DefaultValue="TraceTypeQuery1",PersistentGuid=00000000000000000000000000000000,bHidden=False,bNotConnectable=False,bDefaultValueIsReadOnly=False,bDefaultValueIsIgnored=False,bAdvancedView=False,)
+		   CustomProperties Pin (PinId=8B576AA34608DE6F0B60328E9D9C1166,PinName="bTraceComplex",PinToolTip="Trace Complex\nBoolean\n\nTrue to test against complex collision, false to test against simplified collision.",PinType.PinCategory="bool",PinType.PinSubCategory="",PinType.PinSubCategoryObject=None,PinType.PinSubCategoryMemberReference=(),PinType.PinValueType=(),PinType.bIsMap=False,PinType.bIsSet=False,PinType.bIsArray=False,PinType.bIsReference=False,PinType.bIsConst=False,PinType.bIsWeakPointer=False,DefaultValue="false",AutogeneratedDefaultValue="false",PersistentGuid=00000000000000000000000000000000,bHidden=False,bNotConnectable=False,bDefaultValueIsReadOnly=False,bDefaultValueIsIgnored=False,bAdvancedView=False,)
+		   CustomProperties Pin (PinId=E4B43C3146C4E49A9F72ECBAF170A48B,PinName="ActorsToIgnore",PinToolTip="Actors to Ignore\nArray of Actor References",PinType.PinCategory="object",PinType.PinSubCategory="",PinType.PinSubCategoryObject=Class'/Script/Engine.Actor',PinType.PinSubCategoryMemberReference=(),PinType.PinValueType=(),PinType.bIsMap=False,PinType.bIsSet=False,PinType.bIsArray=True,PinType.bIsReference=True,PinType.bIsConst=True,PinType.bIsWeakPointer=False,PersistentGuid=00000000000000000000000000000000,bHidden=False,bNotConnectable=False,bDefaultValueIsReadOnly=False,bDefaultValueIsIgnored=True,bAdvancedView=False,)
+		   CustomProperties Pin (PinId=6DC8450447A07CB69FEA6ABBD6C12BED,PinName="DrawDebugType",PinToolTip="Draw Debug Type\nEDrawDebugTrace Enum",PinType.PinCategory="byte",PinType.PinSubCategory="",PinType.PinSubCategoryObject=Enum'/Script/Engine.EDrawDebugTrace',PinType.PinSubCategoryMemberReference=(),PinType.PinValueType=(),PinType.bIsMap=False,PinType.bIsSet=False,PinType.bIsArray=False,PinType.bIsReference=False,PinType.bIsConst=False,PinType.bIsWeakPointer=False,DefaultValue="None",PersistentGuid=00000000000000000000000000000000,bHidden=False,bNotConnectable=False,bDefaultValueIsReadOnly=False,bDefaultValueIsIgnored=False,bAdvancedView=False,)
+		   CustomProperties Pin (PinId=43A3FD1943054F5E9C27A298298CE721,PinName="OutHit",PinToolTip="Out Hit\nHit Result Structure\n\nProperties of the trace hit.",Direction="EGPD_Output",PinType.PinCategory="struct",PinType.PinSubCategory="",PinType.PinSubCategoryObject=ScriptStruct'/Script/Engine.HitResult',PinType.PinSubCategoryMemberReference=(),PinType.PinValueType=(),PinType.bIsMap=False,PinType.bIsSet=False,PinType.bIsArray=False,PinType.bIsReference=False,PinType.bIsConst=False,PinType.bIsWeakPointer=False,LinkedTo=(K2Node_CallFunction_65759 9BE54D1545016821EA539F83438E9077,K2Node_CallFunction_296 72651D1B46947C47D63F4AB9879B0012,),PersistentGuid=00000000000000000000000000000000,bHidden=False,bNotConnectable=False,bDefaultValueIsReadOnly=False,bDefaultValueIsIgnored=False,bAdvancedView=False,)
+		   CustomProperties Pin (PinId=D4D0981940265BE6F98E8B9F2CB3C064,PinName="bIgnoreSelf",PinToolTip="Ignore Self\nBoolean",PinType.PinCategory="bool",PinType.PinSubCategory="",PinType.PinSubCategoryObject=None,PinType.PinSubCategoryMemberReference=(),PinType.PinValueType=(),PinType.bIsMap=False,PinType.bIsSet=False,PinType.bIsArray=False,PinType.bIsReference=False,PinType.bIsConst=False,PinType.bIsWeakPointer=False,DefaultValue="false",AutogeneratedDefaultValue="true",PersistentGuid=00000000000000000000000000000000,bHidden=False,bNotConnectable=False,bDefaultValueIsReadOnly=False,bDefaultValueIsIgnored=False,bAdvancedView=False,)
+		   CustomProperties Pin (PinId=17643C41432C7B66D6FF7F9A6A4BA258,PinName="TraceColor",PinToolTip="Trace Color\nLinear Color Structure",PinType.PinCategory="struct",PinType.PinSubCategory="",PinType.PinSubCategoryObject=ScriptStruct'/Script/CoreUObject.LinearColor',PinType.PinSubCategoryMemberReference=(),PinType.PinValueType=(),PinType.bIsMap=False,PinType.bIsSet=False,PinType.bIsArray=False,PinType.bIsReference=False,PinType.bIsConst=False,PinType.bIsWeakPointer=False,DefaultValue="(R=1.000000,G=0.000000,B=0.000000,A=1.000000)",AutogeneratedDefaultValue="(R=1.000000,G=0.000000,B=0.000000,A=1.000000)",PersistentGuid=00000000000000000000000000000000,bHidden=False,bNotConnectable=False,bDefaultValueIsReadOnly=False,bDefaultValueIsIgnored=False,bAdvancedView=True,)
+		   CustomProperties Pin (PinId=F62FE7C840096CE5353C2C98BF138F9C,PinName="TraceHitColor",PinToolTip="Trace Hit Color\nLinear Color Structure",PinType.PinCategory="struct",PinType.PinSubCategory="",PinType.PinSubCategoryObject=ScriptStruct'/Script/CoreUObject.LinearColor',PinType.PinSubCategoryMemberReference=(),PinType.PinValueType=(),PinType.bIsMap=False,PinType.bIsSet=False,PinType.bIsArray=False,PinType.bIsReference=False,PinType.bIsConst=False,PinType.bIsWeakPointer=False,DefaultValue="(R=0.000000,G=1.000000,B=0.000000,A=1.000000)",AutogeneratedDefaultValue="(R=0.000000,G=1.000000,B=0.000000,A=1.000000)",PersistentGuid=00000000000000000000000000000000,bHidden=False,bNotConnectable=False,bDefaultValueIsReadOnly=False,bDefaultValueIsIgnored=False,bAdvancedView=True,)
+		   CustomProperties Pin (PinId=9A93C66A42BC103DD20019BCA763A17F,PinName="DrawTime",PinToolTip="Draw Time\nFloat",PinType.PinCategory="float",PinType.PinSubCategory="",PinType.PinSubCategoryObject=None,PinType.PinSubCategoryMemberReference=(),PinType.PinValueType=(),PinType.bIsMap=False,PinType.bIsSet=False,PinType.bIsArray=False,PinType.bIsReference=False,PinType.bIsConst=False,PinType.bIsWeakPointer=False,DefaultValue="5.000000",AutogeneratedDefaultValue="5.000000",PersistentGuid=00000000000000000000000000000000,bHidden=False,bNotConnectable=False,bDefaultValueIsReadOnly=False,bDefaultValueIsIgnored=False,bAdvancedView=True,)
+		   CustomProperties Pin (PinId=3A64487743C81996C17E56B543685516,PinName="ReturnValue",PinToolTip="Return Value\nBoolean\n\nTrue if there was a hit, false otherwise.",Direction="EGPD_Output",PinType.PinCategory="bool",PinType.PinSubCategory="",PinType.PinSubCategoryObject=None,PinType.PinSubCategoryMemberReference=(),PinType.PinValueType=(),PinType.bIsMap=False,PinType.bIsSet=False,PinType.bIsArray=False,PinType.bIsReference=False,PinType.bIsConst=False,PinType.bIsWeakPointer=False,DefaultValue="false",AutogeneratedDefaultValue="false",LinkedTo=(K2Node_IfThenElse_161 1144399140FB3DC67E1A66B5FEB4A1E1,),PersistentGuid=00000000000000000000000000000000,bHidden=False,bNotConnectable=False,bDefaultValueIsReadOnly=False,bDefaultValueIsIgnored=False,bAdvancedView=False,)
+		End Object
+		Begin Object Class=K2Node_CommutativeAssociativeBinaryOperator Name="K2Node_CommutativeAssociativeBinaryOperator_1357"
+		   bIsPureFunc=True
+		   FunctionReference=(MemberParent=Class'/Script/Engine.KismetMathLibrary',MemberName="Add_VectorVector")
+		   NodePosX=432
+		   NodePosY=112
+		   NodeGuid=6FCAAA3643F1C5AC5E5155899B299110
+		   CustomProperties Pin (PinId=B729CFBC424DA01AC5E021B25FE0F7CD,PinName="self",PinFriendlyName=NSLOCTEXT("K2Node", "Target", "Target"),PinType.PinCategory="object",PinType.PinSubCategory="",PinType.PinSubCategoryObject=Class'/Script/Engine.KismetMathLibrary',PinType.PinSubCategoryMemberReference=(),PinType.PinValueType=(),PinType.bIsMap=False,PinType.bIsSet=False,PinType.bIsArray=False,PinType.bIsReference=False,PinType.bIsConst=False,PinType.bIsWeakPointer=False,DefaultObject="/Script/Engine.Default__KismetMathLibrary",PersistentGuid=00000000000000000000000000000000,bHidden=True,bNotConnectable=False,bDefaultValueIsReadOnly=False,bDefaultValueIsIgnored=False,bAdvancedView=False,)
+		   CustomProperties Pin (PinId=4410C89D42A350AA1AF9EE830D6A174E,PinName="A",PinType.PinCategory="struct",PinType.PinSubCategory="",PinType.PinSubCategoryObject=ScriptStruct'/Script/CoreUObject.Vector',PinType.PinSubCategoryMemberReference=(),PinType.PinValueType=(),PinType.bIsMap=False,PinType.bIsSet=False,PinType.bIsArray=False,PinType.bIsReference=False,PinType.bIsConst=False,PinType.bIsWeakPointer=False,DefaultValue="0, 0, 0",AutogeneratedDefaultValue="0, 0, 0",LinkedTo=(K2Node_CallFunction_65349 40870E1F404EA24CCC3295A98D8A13B6,),PersistentGuid=00000000000000000000000000000000,bHidden=False,bNotConnectable=False,bDefaultValueIsReadOnly=False,bDefaultValueIsIgnored=False,bAdvancedView=False,)
+		   CustomProperties Pin (PinId=C2F85EED4A8F818C3CE355A3C56691F5,PinName="B",PinType.PinCategory="struct",PinType.PinSubCategory="",PinType.PinSubCategoryObject=ScriptStruct'/Script/CoreUObject.Vector',PinType.PinSubCategoryMemberReference=(),PinType.PinValueType=(),PinType.bIsMap=False,PinType.bIsSet=False,PinType.bIsArray=False,PinType.bIsReference=False,PinType.bIsConst=False,PinType.bIsWeakPointer=False,DefaultValue="0, 0, 0",AutogeneratedDefaultValue="0, 0, 0",LinkedTo=(K2Node_CallFunction_64919 7C3AB90A4E8F92F9FAE0709612FDD604,),PersistentGuid=00000000000000000000000000000000,bHidden=False,bNotConnectable=False,bDefaultValueIsReadOnly=False,bDefaultValueIsIgnored=False,bAdvancedView=False,)
+		   CustomProperties Pin (PinId=1778AF2B4812E5B635D0DC915DFAA735,PinName="ReturnValue",Direction="EGPD_Output",PinType.PinCategory="struct",PinType.PinSubCategory="",PinType.PinSubCategoryObject=ScriptStruct'/Script/CoreUObject.Vector',PinType.PinSubCategoryMemberReference=(),PinType.PinValueType=(),PinType.bIsMap=False,PinType.bIsSet=False,PinType.bIsArray=False,PinType.bIsReference=False,PinType.bIsConst=False,PinType.bIsWeakPointer=False,AutogeneratedDefaultValue="0, 0, 0",LinkedTo=(K2Node_CallFunction_65598 54C3C6D04942C8BC62F2D696DFF0A12C,),PersistentGuid=00000000000000000000000000000000,bHidden=False,bNotConnectable=False,bDefaultValueIsReadOnly=False,bDefaultValueIsIgnored=False,bAdvancedView=False,)
+		End Object
+		Begin Object Class=K2Node_CallFunction Name="K2Node_CallFunction_65759"
+		   bIsPureFunc=True
+		   FunctionReference=(MemberParent=Class'/Script/Engine.GameplayStatics',MemberName="BreakHitResult")
+		   NodePosX=1088
+		   NodePosY=-64
+		   NodeGuid=13D0A7214A86CFAEA5D280843894F2C4
+		   CustomProperties Pin (PinId=CF78D28A455F553D39502496BC89B70E,PinName="self",PinFriendlyName=NSLOCTEXT("K2Node", "Target", "Target"),PinToolTip="Target\nGameplay Statics Reference",PinType.PinCategory="object",PinType.PinSubCategory="",PinType.PinSubCategoryObject=Class'/Script/Engine.GameplayStatics',PinType.PinSubCategoryMemberReference=(),PinType.PinValueType=(),PinType.bIsMap=False,PinType.bIsSet=False,PinType.bIsArray=False,PinType.bIsReference=False,PinType.bIsConst=False,PinType.bIsWeakPointer=False,DefaultObject="/Script/Engine.Default__GameplayStatics",PersistentGuid=00000000000000000000000000000000,bHidden=True,bNotConnectable=False,bDefaultValueIsReadOnly=False,bDefaultValueIsIgnored=False,bAdvancedView=False,)
+		   CustomProperties Pin (PinId=9BE54D1545016821EA539F83438E9077,PinName="Hit",PinToolTip="Hit\nHit Result Structure (by ref)\n\nThe source HitResult.",PinType.PinCategory="struct",PinType.PinSubCategory="",PinType.PinSubCategoryObject=ScriptStruct'/Script/Engine.HitResult',PinType.PinSubCategoryMemberReference=(),PinType.PinValueType=(),PinType.bIsMap=False,PinType.bIsSet=False,PinType.bIsArray=False,PinType.bIsReference=True,PinType.bIsConst=True,PinType.bIsWeakPointer=False,LinkedTo=(K2Node_CallFunction_65598 43A3FD1943054F5E9C27A298298CE721,),PersistentGuid=00000000000000000000000000000000,bHidden=False,bNotConnectable=False,bDefaultValueIsReadOnly=False,bDefaultValueIsIgnored=True,bAdvancedView=False,)
+		   CustomProperties Pin (PinId=42FAE549473540FC07A70799D95A29B3,PinName="bBlockingHit",PinToolTip="Blocking Hit\nBoolean\n\nTrue if there was a blocking hit, false otherwise.",Direction="EGPD_Output",PinType.PinCategory="bool",PinType.PinSubCategory="",PinType.PinSubCategoryObject=None,PinType.PinSubCategoryMemberReference=(),PinType.PinValueType=(),PinType.bIsMap=False,PinType.bIsSet=False,PinType.bIsArray=False,PinType.bIsReference=False,PinType.bIsConst=False,PinType.bIsWeakPointer=False,DefaultValue="false",AutogeneratedDefaultValue="false",PersistentGuid=00000000000000000000000000000000,bHidden=False,bNotConnectable=False,bDefaultValueIsReadOnly=False,bDefaultValueIsIgnored=False,bAdvancedView=False,)
+		   CustomProperties Pin (PinId=31399B19487BF09D8EF1088E8ED56118,PinName="bInitialOverlap",PinToolTip="Initial Overlap\nBoolean\n\nTrue if the hit started in an initial overlap. In this case some other values should be interpreted differently. Time will be 0, ImpactPoint will equal Location, and normals will be equal and indicate a depenetration vector.",Direction="EGPD_Output",PinType.PinCategory="bool",PinType.PinSubCategory="",PinType.PinSubCategoryObject=None,PinType.PinSubCategoryMemberReference=(),PinType.PinValueType=(),PinType.bIsMap=False,PinType.bIsSet=False,PinType.bIsArray=False,PinType.bIsReference=False,PinType.bIsConst=False,PinType.bIsWeakPointer=False,DefaultValue="false",AutogeneratedDefaultValue="false",PersistentGuid=00000000000000000000000000000000,bHidden=False,bNotConnectable=False,bDefaultValueIsReadOnly=False,bDefaultValueIsIgnored=False,bAdvancedView=False,)
+		   CustomProperties Pin (PinId=4312416249D19F19294CADA26CDE64E5,PinName="Time",PinToolTip="Time\nFloat\n\n\'Time\' of impact along trace direction ranging from [0.0 to 1.0) if there is a hit, indicating time between start and end. Equals 1.0 if there is no hit.",Direction="EGPD_Output",PinType.PinCategory="float",PinType.PinSubCategory="",PinType.PinSubCategoryObject=None,PinType.PinSubCategoryMemberReference=(),PinType.PinValueType=(),PinType.bIsMap=False,PinType.bIsSet=False,PinType.bIsArray=False,PinType.bIsReference=False,PinType.bIsConst=False,PinType.bIsWeakPointer=False,DefaultValue="0.0",AutogeneratedDefaultValue="0.0",PersistentGuid=00000000000000000000000000000000,bHidden=False,bNotConnectable=False,bDefaultValueIsReadOnly=False,bDefaultValueIsIgnored=False,bAdvancedView=False,)
+		   CustomProperties Pin (PinId=C3E1DE9D4B514E4FCDA1159FEA7E7832,PinName="Location",PinToolTip="Location\nVector \n\nLocation of the hit in world space. If this was a swept shape test, this is the location where we can place the shape in the world where it will not penetrate.",Direction="EGPD_Output",PinType.PinCategory="struct",PinType.PinSubCategory="",PinType.PinSubCategoryObject=ScriptStruct'/Script/CoreUObject.Vector',PinType.PinSubCategoryMemberReference=(),PinType.PinValueType=(),PinType.bIsMap=False,PinType.bIsSet=False,PinType.bIsArray=False,PinType.bIsReference=False,PinType.bIsConst=False,PinType.bIsWeakPointer=False,AutogeneratedDefaultValue="0, 0, 0",PersistentGuid=00000000000000000000000000000000,bHidden=False,bNotConnectable=False,bDefaultValueIsReadOnly=False,bDefaultValueIsIgnored=False,bAdvancedView=False,)
+		   CustomProperties Pin (PinId=EB9F6AED44EDCC7D8F313FB63A3B4139,PinName="ImpactPoint",PinToolTip="Impact Point\nVector \n\nLocation of the actual contact point of the trace shape with the surface of the hit object. Equal to Location in the case of an initial overlap.",Direction="EGPD_Output",PinType.PinCategory="struct",PinType.PinSubCategory="",PinType.PinSubCategoryObject=ScriptStruct'/Script/CoreUObject.Vector',PinType.PinSubCategoryMemberReference=(),PinType.PinValueType=(),PinType.bIsMap=False,PinType.bIsSet=False,PinType.bIsArray=False,PinType.bIsReference=False,PinType.bIsConst=False,PinType.bIsWeakPointer=False,AutogeneratedDefaultValue="0, 0, 0",PersistentGuid=00000000000000000000000000000000,bHidden=False,bNotConnectable=False,bDefaultValueIsReadOnly=False,bDefaultValueIsIgnored=False,bAdvancedView=False,)
+		   CustomProperties Pin (PinId=AC37F9C247A40EF82D1755B03A6EF0AF,PinName="Normal",PinToolTip="Normal\nVector \n\nNormal of the hit in world space, for the object that was swept (e.g. for a sphere trace this points towards the sphere\'s center). Equal to ImpactNormal for line tests.",Direction="EGPD_Output",PinType.PinCategory="struct",PinType.PinSubCategory="",PinType.PinSubCategoryObject=ScriptStruct'/Script/CoreUObject.Vector',PinType.PinSubCategoryMemberReference=(),PinType.PinValueType=(),PinType.bIsMap=False,PinType.bIsSet=False,PinType.bIsArray=False,PinType.bIsReference=False,PinType.bIsConst=False,PinType.bIsWeakPointer=False,AutogeneratedDefaultValue="0, 0, 0",PersistentGuid=00000000000000000000000000000000,bHidden=False,bNotConnectable=False,bDefaultValueIsReadOnly=False,bDefaultValueIsIgnored=False,bAdvancedView=False,)
+		   CustomProperties Pin (PinId=43DE840740F7EDB5AEEE2BA59BA63807,PinName="ImpactNormal",PinToolTip="Impact Normal\nVector \n\nNormal of the hit in world space, for the object that was hit by the sweep.",Direction="EGPD_Output",PinType.PinCategory="struct",PinType.PinSubCategory="",PinType.PinSubCategoryObject=ScriptStruct'/Script/CoreUObject.Vector',PinType.PinSubCategoryMemberReference=(),PinType.PinValueType=(),PinType.bIsMap=False,PinType.bIsSet=False,PinType.bIsArray=False,PinType.bIsReference=False,PinType.bIsConst=False,PinType.bIsWeakPointer=False,AutogeneratedDefaultValue="0, 0, 0",PersistentGuid=00000000000000000000000000000000,bHidden=False,bNotConnectable=False,bDefaultValueIsReadOnly=False,bDefaultValueIsIgnored=False,bAdvancedView=False,)
+		   CustomProperties Pin (PinId=532D29F44B5DCE5F021B698ED58D27C4,PinName="PhysMat",PinToolTip="Phys Mat\nPhysical Material Reference\n\nPhysical material that was hit. Must set bReturnPhysicalMaterial to true in the query params for this to be returned.",Direction="EGPD_Output",PinType.PinCategory="object",PinType.PinSubCategory="",PinType.PinSubCategoryObject=Class'/Script/Engine.PhysicalMaterial',PinType.PinSubCategoryMemberReference=(),PinType.PinValueType=(),PinType.bIsMap=False,PinType.bIsSet=False,PinType.bIsArray=False,PinType.bIsReference=False,PinType.bIsConst=False,PinType.bIsWeakPointer=False,PersistentGuid=00000000000000000000000000000000,bHidden=False,bNotConnectable=False,bDefaultValueIsReadOnly=False,bDefaultValueIsIgnored=False,bAdvancedView=False,)
+		   CustomProperties Pin (PinId=F1D6985542953219C7EF09A3FBFA8AF0,PinName="HitActor",PinToolTip="Hit Actor\nActor Reference\n\nActor hit by the trace.",Direction="EGPD_Output",PinType.PinCategory="object",PinType.PinSubCategory="",PinType.PinSubCategoryObject=Class'/Script/Engine.Actor',PinType.PinSubCategoryMemberReference=(),PinType.PinValueType=(),PinType.bIsMap=False,PinType.bIsSet=False,PinType.bIsArray=False,PinType.bIsReference=False,PinType.bIsConst=False,PinType.bIsWeakPointer=False,LinkedTo=(K2Node_MacroInstance_65 C3FB28AC47AAD2B5998C2AA9D4862C2E,),PersistentGuid=00000000000000000000000000000000,bHidden=False,bNotConnectable=False,bDefaultValueIsReadOnly=False,bDefaultValueIsIgnored=False,bAdvancedView=False,)
+		   CustomProperties Pin (PinId=FDFC05184366EBC9BDC42AA3A95EF523,PinName="HitComponent",PinToolTip="Hit Component\nPrimitive Component Reference\n\nPrimitiveComponent hit by the trace.",Direction="EGPD_Output",PinType.PinCategory="object",PinType.PinSubCategory="",PinType.PinSubCategoryObject=Class'/Script/Engine.PrimitiveComponent',PinType.PinSubCategoryMemberReference=(),PinType.PinValueType=(),PinType.bIsMap=False,PinType.bIsSet=False,PinType.bIsArray=False,PinType.bIsReference=False,PinType.bIsConst=False,PinType.bIsWeakPointer=False,PersistentGuid=00000000000000000000000000000000,bHidden=False,bNotConnectable=False,bDefaultValueIsReadOnly=False,bDefaultValueIsIgnored=False,bAdvancedView=False,)
+		   CustomProperties Pin (PinId=CE3689BE4446FD1C971B7A9942E7AC24,PinName="HitBoneName",PinToolTip="Hit Bone Name\nName\n\nName of the bone hit (valid only if we hit a skeletal mesh).",Direction="EGPD_Output",PinType.PinCategory="name",PinType.PinSubCategory="",PinType.PinSubCategoryObject=None,PinType.PinSubCategoryMemberReference=(),PinType.PinValueType=(),PinType.bIsMap=False,PinType.bIsSet=False,PinType.bIsArray=False,PinType.bIsReference=False,PinType.bIsConst=False,PinType.bIsWeakPointer=False,DefaultValue="None",AutogeneratedDefaultValue="None",PersistentGuid=00000000000000000000000000000000,bHidden=False,bNotConnectable=False,bDefaultValueIsReadOnly=False,bDefaultValueIsIgnored=False,bAdvancedView=False,)
+		   CustomProperties Pin (PinId=1E9AAE5D433101E315ABFCB6EE0C3F4E,PinName="HitItem",PinToolTip="Hit Item\nInteger\n\nPrimitive-specific data recording which item in the primitive was hit",Direction="EGPD_Output",PinType.PinCategory="int",PinType.PinSubCategory="",PinType.PinSubCategoryObject=None,PinType.PinSubCategoryMemberReference=(),PinType.PinValueType=(),PinType.bIsMap=False,PinType.bIsSet=False,PinType.bIsArray=False,PinType.bIsReference=False,PinType.bIsConst=False,PinType.bIsWeakPointer=False,DefaultValue="0",AutogeneratedDefaultValue="0",PersistentGuid=00000000000000000000000000000000,bHidden=False,bNotConnectable=False,bDefaultValueIsReadOnly=False,bDefaultValueIsIgnored=False,bAdvancedView=False,)
+		   CustomProperties Pin (PinId=FA0E13FF4A7001F230C36C84199334C5,PinName="FaceIndex",PinToolTip="Face Index\nInteger\n\nIf colliding with trimesh or landscape, index of face that was hit.",Direction="EGPD_Output",PinType.PinCategory="int",PinType.PinSubCategory="",PinType.PinSubCategoryObject=None,PinType.PinSubCategoryMemberReference=(),PinType.PinValueType=(),PinType.bIsMap=False,PinType.bIsSet=False,PinType.bIsArray=False,PinType.bIsReference=False,PinType.bIsConst=False,PinType.bIsWeakPointer=False,DefaultValue="0",AutogeneratedDefaultValue="0",PersistentGuid=00000000000000000000000000000000,bHidden=False,bNotConnectable=False,bDefaultValueIsReadOnly=False,bDefaultValueIsIgnored=False,bAdvancedView=False,)
+		   CustomProperties Pin (PinId=6AFDED3441A726D6362809827812A060,PinName="TraceStart",PinToolTip="Trace Start\nVector ",Direction="EGPD_Output",PinType.PinCategory="struct",PinType.PinSubCategory="",PinType.PinSubCategoryObject=ScriptStruct'/Script/CoreUObject.Vector',PinType.PinSubCategoryMemberReference=(),PinType.PinValueType=(),PinType.bIsMap=False,PinType.bIsSet=False,PinType.bIsArray=False,PinType.bIsReference=False,PinType.bIsConst=False,PinType.bIsWeakPointer=False,DefaultValue="0, 0, 0",AutogeneratedDefaultValue="0, 0, 0",PersistentGuid=00000000000000000000000000000000,bHidden=False,bNotConnectable=False,bDefaultValueIsReadOnly=False,bDefaultValueIsIgnored=False,bAdvancedView=False,)
+		   CustomProperties Pin (PinId=96D907364CB400491CEFF4B6B3921F92,PinName="TraceEnd",PinToolTip="Trace End\nVector ",Direction="EGPD_Output",PinType.PinCategory="struct",PinType.PinSubCategory="",PinType.PinSubCategoryObject=ScriptStruct'/Script/CoreUObject.Vector',PinType.PinSubCategoryMemberReference=(),PinType.PinValueType=(),PinType.bIsMap=False,PinType.bIsSet=False,PinType.bIsArray=False,PinType.bIsReference=False,PinType.bIsConst=False,PinType.bIsWeakPointer=False,DefaultValue="0, 0, 0",AutogeneratedDefaultValue="0, 0, 0",PersistentGuid=00000000000000000000000000000000,bHidden=False,bNotConnectable=False,bDefaultValueIsReadOnly=False,bDefaultValueIsIgnored=False,bAdvancedView=False,)
+		End Object
+		Begin Object Class=K2Node_CommutativeAssociativeBinaryOperator Name="K2Node_CommutativeAssociativeBinaryOperator_4227"
+		   bIsPureFunc=True
+		   FunctionReference=(MemberParent=Class'/Script/Engine.KismetMathLibrary',MemberName="Multiply_FloatFloat")
+		   NodePosX=112
+		   NodePosY=224
+		   NodeGuid=FB6F8E114CCFC6E409EAD7B3A2DA5597
+		   CustomProperties Pin (PinId=0D4F3C2044C7E8EA342002AF074BCB55,PinName="self",PinFriendlyName=NSLOCTEXT("K2Node", "Target", "Target"),PinType.PinCategory="object",PinType.PinSubCategory="",PinType.PinSubCategoryObject=Class'/Script/Engine.KismetMathLibrary',PinType.PinSubCategoryMemberReference=(),PinType.PinValueType=(),PinType.bIsMap=False,PinType.bIsSet=False,PinType.bIsArray=False,PinType.bIsReference=False,PinType.bIsConst=False,PinType.bIsWeakPointer=False,DefaultObject="/Script/Engine.Default__KismetMathLibrary",PersistentGuid=00000000000000000000000000000000,bHidden=True,bNotConnectable=False,bDefaultValueIsReadOnly=False,bDefaultValueIsIgnored=False,bAdvancedView=False,)
+		   CustomProperties Pin (PinId=AE00303B4D5F38B1B48C718ABEC37EF0,PinName="A",PinType.PinCategory="float",PinType.PinSubCategory="",PinType.PinSubCategoryObject=None,PinType.PinSubCategoryMemberReference=(),PinType.PinValueType=(),PinType.bIsMap=False,PinType.bIsSet=False,PinType.bIsArray=False,PinType.bIsReference=False,PinType.bIsConst=False,PinType.bIsWeakPointer=False,DefaultValue="0.0",AutogeneratedDefaultValue="0.0",LinkedTo=(K2Node_VariableGet_281 F3DEDC984C80801BC87F67B0097115B8,),PersistentGuid=00000000000000000000000000000000,bHidden=False,bNotConnectable=False,bDefaultValueIsReadOnly=False,bDefaultValueIsIgnored=False,bAdvancedView=False,)
+		   CustomProperties Pin (PinId=290EC66447D9833EAF1FD8914F55B78C,PinName="B",PinType.PinCategory="float",PinType.PinSubCategory="",PinType.PinSubCategoryObject=None,PinType.PinSubCategoryMemberReference=(),PinType.PinValueType=(),PinType.bIsMap=False,PinType.bIsSet=False,PinType.bIsArray=False,PinType.bIsReference=False,PinType.bIsConst=False,PinType.bIsWeakPointer=False,DefaultValue="10",AutogeneratedDefaultValue="0.0",PersistentGuid=00000000000000000000000000000000,bHidden=False,bNotConnectable=False,bDefaultValueIsReadOnly=False,bDefaultValueIsIgnored=False,bAdvancedView=False,)
+		   CustomProperties Pin (PinId=0FAF44A34397D0F87EB02B88E8A653E3,PinName="ReturnValue",Direction="EGPD_Output",PinType.PinCategory="float",PinType.PinSubCategory="",PinType.PinSubCategoryObject=None,PinType.PinSubCategoryMemberReference=(),PinType.PinValueType=(),PinType.bIsMap=False,PinType.bIsSet=False,PinType.bIsArray=False,PinType.bIsReference=False,PinType.bIsConst=False,PinType.bIsWeakPointer=False,DefaultValue="0.0",AutogeneratedDefaultValue="0.0",LinkedTo=(K2Node_CallFunction_65349 2C5718584516424BB52CF6A79E702026,),PersistentGuid=00000000000000000000000000000000,bHidden=False,bNotConnectable=False,bDefaultValueIsReadOnly=False,bDefaultValueIsIgnored=False,bAdvancedView=False,)
+		End Object
+		Begin Object Class=K2Node_IfThenElse Name="K2Node_IfThenElse_161"
+		   NodePosX=1120
+		   NodePosY=-208
+		   ErrorType=4
+		   ErrorMsg="Info The execution path doesn\'t end with a return node.  False "
+		   NodeGuid=34B6B6BC48A15E72761AF2A4794F0316
+		   CustomProperties Pin (PinId=66D340DC4A52DC8D727C43ABA6E818C4,PinName="execute",PinType.PinCategory="exec",PinType.PinSubCategory="",PinType.PinSubCategoryObject=None,PinType.PinSubCategoryMemberReference=(),PinType.PinValueType=(),PinType.bIsMap=False,PinType.bIsSet=False,PinType.bIsArray=False,PinType.bIsReference=False,PinType.bIsConst=False,PinType.bIsWeakPointer=False,LinkedTo=(K2Node_CallFunction_65598 E001912C48BB5BEC656460843FE151FA,),PersistentGuid=00000000000000000000000000000000,bHidden=False,bNotConnectable=False,bDefaultValueIsReadOnly=False,bDefaultValueIsIgnored=False,bAdvancedView=False,)
+		   CustomProperties Pin (PinId=1144399140FB3DC67E1A66B5FEB4A1E1,PinName="Condition",PinType.PinCategory="bool",PinType.PinSubCategory="",PinType.PinSubCategoryObject=None,PinType.PinSubCategoryMemberReference=(),PinType.PinValueType=(),PinType.bIsMap=False,PinType.bIsSet=False,PinType.bIsArray=False,PinType.bIsReference=False,PinType.bIsConst=False,PinType.bIsWeakPointer=False,DefaultValue="true",AutogeneratedDefaultValue="true",LinkedTo=(K2Node_CallFunction_65598 3A64487743C81996C17E56B543685516,),PersistentGuid=00000000000000000000000000000000,bHidden=False,bNotConnectable=False,bDefaultValueIsReadOnly=False,bDefaultValueIsIgnored=False,bAdvancedView=False,)
+		   CustomProperties Pin (PinId=8DD9A5BC40D019A1AFD6F284F8C15FFD,PinName="then",PinFriendlyName=NSLOCTEXT("K2Node", "true", "true"),Direction="EGPD_Output",PinType.PinCategory="exec",PinType.PinSubCategory="",PinType.PinSubCategoryObject=None,PinType.PinSubCategoryMemberReference=(),PinType.PinValueType=(),PinType.bIsMap=False,PinType.bIsSet=False,PinType.bIsArray=False,PinType.bIsReference=False,PinType.bIsConst=False,PinType.bIsWeakPointer=False,LinkedTo=(K2Node_MacroInstance_65 54DCDF6A4C18DFFBF71F54AE24DA40C0,),PersistentGuid=00000000000000000000000000000000,bHidden=False,bNotConnectable=False,bDefaultValueIsReadOnly=False,bDefaultValueIsIgnored=False,bAdvancedView=False,)
+		   CustomProperties Pin (PinId=4CD062D04319B76E8FD9D6B0CDDC7EA9,PinName="else",PinFriendlyName=NSLOCTEXT("K2Node", "false", "false"),Direction="EGPD_Output",PinType.PinCategory="exec",PinType.PinSubCategory="",PinType.PinSubCategoryObject=None,PinType.PinSubCategoryMemberReference=(),PinType.PinValueType=(),PinType.bIsMap=False,PinType.bIsSet=False,PinType.bIsArray=False,PinType.bIsReference=False,PinType.bIsConst=False,PinType.bIsWeakPointer=False,PersistentGuid=00000000000000000000000000000000,bHidden=False,bNotConnectable=False,bDefaultValueIsReadOnly=False,bDefaultValueIsIgnored=False,bAdvancedView=False,)
+		End Object
+		Begin Object Class=EdGraphNode_Comment Name="EdGraphNode_Comment_590"
+		   NodePosX=-640
+		   NodePosY=-304
+		   NodeWidth=1632
+		   NodeHeight=704
+		   NodeComment="Trace from the center of the screen in direction of the camera"
+		   NodeGuid=551D4DF949A91B6DA08A58A45A1CDEA6
+		End Object
+		Begin Object Class=EdGraphNode_Comment Name="EdGraphNode_Comment_596"
+		   NodePosX=1040
+		   NodePosY=-304
+		   NodeWidth=1152
+		   NodeHeight=704
+		   NodeComment="If the Trace hits something send UV\'s to Apply Paint"
+		   NodeGuid=56FA82134685BB13CC95D9B1B8C1E73C
+		End Object
+		Begin Object Class=K2Node_VariableGet Name="K2Node_VariableGet_281"
+		   VariableReference=(MemberName="InteractionDistance",MemberGuid=27EB0D754F93A53064BCD9BE1860B0F8,bSelfContext=True)
+		   NodePosX=-96
+		   NodePosY=272
+		   ErrorType=1
+		   ErrorMsg="Error The property associated with  InteractionDistance  could not be found"
+		   NodeGuid=C195F9E046F31887CDC3F5B991216BDE
+		   CustomProperties Pin (PinId=F3DEDC984C80801BC87F67B0097115B8,PinName="InteractionDistance",Direction="EGPD_Output",PinType.PinCategory="float",PinType.PinSubCategory="",PinType.PinSubCategoryObject=None,PinType.PinSubCategoryMemberReference=(),PinType.PinValueType=(),PinType.bIsMap=False,PinType.bIsSet=False,PinType.bIsArray=False,PinType.bIsReference=False,PinType.bIsConst=False,PinType.bIsWeakPointer=False,AutogeneratedDefaultValue="0.0",LinkedTo=(K2Node_CommutativeAssociativeBinaryOperator_4227 AE00303B4D5F38B1B48C718ABEC37EF0,),PersistentGuid=00000000000000000000000000000000,bHidden=False,bNotConnectable=False,bDefaultValueIsReadOnly=False,bDefaultValueIsIgnored=False,bAdvancedView=False,)
+		   CustomProperties Pin (PinId=6388FEAC4809B345E74AAFAEC02C96EC,PinName="self",PinFriendlyName=NSLOCTEXT("K2Node", "Target", "Target"),PinType.PinCategory="object",PinType.PinSubCategory="",PinType.PinSubCategoryObject=BlueprintGeneratedClass'/Game/Docs/Traces/NewBlueprint.NewBlueprint_C',PinType.PinSubCategoryMemberReference=(),PinType.PinValueType=(),PinType.bIsMap=False,PinType.bIsSet=False,PinType.bIsArray=False,PinType.bIsReference=False,PinType.bIsConst=False,PinType.bIsWeakPointer=False,PersistentGuid=00000000000000000000000000000000,bHidden=True,bNotConnectable=False,bDefaultValueIsReadOnly=False,bDefaultValueIsIgnored=False,bAdvancedView=False,)
+		End Object
+		Begin Object Class=K2Node_CallFunction Name="K2Node_CallFunction_296"
+		   bIsPureFunc=True
+		   FunctionReference=(MemberParent=Class'/Script/Engine.GameplayStatics',MemberName="FindCollisionUV")
+		   NodePosX=1504
+		   NodePosY=-112
+		   NodeGuid=1C0EE1E748A1CB5F4D0D9483DFDCE444
+		   CustomProperties Pin (PinId=73B5292749A0273FEE9A2D8F4FFD1EE5,PinName="self",PinFriendlyName=NSLOCTEXT("K2Node", "Target", "Target"),PinToolTip="Target\nGameplay Statics Reference",PinType.PinCategory="object",PinType.PinSubCategory="",PinType.PinSubCategoryObject=Class'/Script/Engine.GameplayStatics',PinType.PinSubCategoryMemberReference=(),PinType.PinValueType=(),PinType.bIsMap=False,PinType.bIsSet=False,PinType.bIsArray=False,PinType.bIsReference=False,PinType.bIsConst=False,PinType.bIsWeakPointer=False,DefaultObject="/Script/Engine.Default__GameplayStatics",PersistentGuid=00000000000000000000000000000000,bHidden=True,bNotConnectable=False,bDefaultValueIsReadOnly=False,bDefaultValueIsIgnored=False,bAdvancedView=False,)
+		   CustomProperties Pin (PinId=72651D1B46947C47D63F4AB9879B0012,PinName="Hit",PinToolTip="Hit\nHit Result Structure (by ref)",PinType.PinCategory="struct",PinType.PinSubCategory="",PinType.PinSubCategoryObject=ScriptStruct'/Script/Engine.HitResult',PinType.PinSubCategoryMemberReference=(),PinType.PinValueType=(),PinType.bIsMap=False,PinType.bIsSet=False,PinType.bIsArray=False,PinType.bIsReference=True,PinType.bIsConst=True,PinType.bIsWeakPointer=False,LinkedTo=(K2Node_CallFunction_65598 43A3FD1943054F5E9C27A298298CE721,),PersistentGuid=00000000000000000000000000000000,bHidden=False,bNotConnectable=False,bDefaultValueIsReadOnly=False,bDefaultValueIsIgnored=True,bAdvancedView=False,)
+		   CustomProperties Pin (PinId=E43C942E4A0AE2FCC42AC88A60AF24D7,PinName="UVChannel",PinToolTip="UVChannel\nInteger",PinType.PinCategory="int",PinType.PinSubCategory="",PinType.PinSubCategoryObject=None,PinType.PinSubCategoryMemberReference=(),PinType.PinValueType=(),PinType.bIsMap=False,PinType.bIsSet=False,PinType.bIsArray=False,PinType.bIsReference=False,PinType.bIsConst=False,PinType.bIsWeakPointer=False,DefaultValue="0",AutogeneratedDefaultValue="0",PersistentGuid=00000000000000000000000000000000,bHidden=False,bNotConnectable=False,bDefaultValueIsReadOnly=False,bDefaultValueIsIgnored=False,bAdvancedView=False,)
+		   CustomProperties Pin (PinId=B78D125B4929D4D20E4338A6EEB13A7B,PinName="UV",PinToolTip="UV\nVector 2D Structure",Direction="EGPD_Output",PinType.PinCategory="struct",PinType.PinSubCategory="",PinType.PinSubCategoryObject=ScriptStruct'/Script/CoreUObject.Vector2D',PinType.PinSubCategoryMemberReference=(),PinType.PinValueType=(),PinType.bIsMap=False,PinType.bIsSet=False,PinType.bIsArray=False,PinType.bIsReference=False,PinType.bIsConst=False,PinType.bIsWeakPointer=False,LinkedTo=(K2Node_CallFunction_311 80B33882446EDB9EFA5F78B8FF6F8432,),PersistentGuid=00000000000000000000000000000000,bHidden=False,bNotConnectable=False,bDefaultValueIsReadOnly=False,bDefaultValueIsIgnored=False,bAdvancedView=False,)
+		   CustomProperties Pin (PinId=8B769974413EA584886120B57BF3B026,PinName="ReturnValue",PinToolTip="Return Value\nBoolean",Direction="EGPD_Output",PinType.PinCategory="bool",PinType.PinSubCategory="",PinType.PinSubCategoryObject=None,PinType.PinSubCategoryMemberReference=(),PinType.PinValueType=(),PinType.bIsMap=False,PinType.bIsSet=False,PinType.bIsArray=False,PinType.bIsReference=False,PinType.bIsConst=False,PinType.bIsWeakPointer=False,DefaultValue="false",AutogeneratedDefaultValue="false",PersistentGuid=00000000000000000000000000000000,bHidden=False,bNotConnectable=False,bDefaultValueIsReadOnly=False,bDefaultValueIsIgnored=False,bAdvancedView=False,)
+		End Object
+		Begin Object Class=K2Node_CallFunction Name="K2Node_CallFunction_311"
+		   FunctionReference=(MemberName="ApplyPaint",MemberGuid=36419B2D4BC97CF5BDAC0481CBD89E97,bSelfContext=True)
+		   NodePosX=1920
+		   NodePosY=-208
+		   NodeGuid=568091E446E04E2FA33D5FAC6285ADD2
+		   CustomProperties Pin (PinId=F45D897646865C99837DAD9886C0860F,PinName="execute",PinToolTip="\nExec",PinType.PinCategory="exec",PinType.PinSubCategory="",PinType.PinSubCategoryObject=None,PinType.PinSubCategoryMemberReference=(),PinType.PinValueType=(),PinType.bIsMap=False,PinType.bIsSet=False,PinType.bIsArray=False,PinType.bIsReference=False,PinType.bIsConst=False,PinType.bIsWeakPointer=False,LinkedTo=(K2Node_MacroInstance_65 E75DA26F4EBE18C63E3CEDA6B1D88DC6,),PersistentGuid=00000000000000000000000000000000,bHidden=False,bNotConnectable=False,bDefaultValueIsReadOnly=False,bDefaultValueIsIgnored=False,bAdvancedView=False,)
+		   CustomProperties Pin (PinId=073BB5B84A683A2604B2E59A58F6A849,PinName="then",PinToolTip="\nExec",Direction="EGPD_Output",PinType.PinCategory="exec",PinType.PinSubCategory="",PinType.PinSubCategoryObject=None,PinType.PinSubCategoryMemberReference=(),PinType.PinValueType=(),PinType.bIsMap=False,PinType.bIsSet=False,PinType.bIsArray=False,PinType.bIsReference=False,PinType.bIsConst=False,PinType.bIsWeakPointer=False,PersistentGuid=00000000000000000000000000000000,bHidden=False,bNotConnectable=False,bDefaultValueIsReadOnly=False,bDefaultValueIsIgnored=False,bAdvancedView=False,)
+		   CustomProperties Pin (PinId=2BB4B746445D222517430EB468AC7497,PinName="self",PinFriendlyName=NSLOCTEXT("K2Node", "Target", "Target"),PinToolTip="Target\nSelf Reference",PinType.PinCategory="object",PinType.PinSubCategory="self",PinType.PinSubCategoryObject=None,PinType.PinSubCategoryMemberReference=(),PinType.PinValueType=(),PinType.bIsMap=False,PinType.bIsSet=False,PinType.bIsArray=False,PinType.bIsReference=False,PinType.bIsConst=False,PinType.bIsWeakPointer=False,PersistentGuid=00000000000000000000000000000000,bHidden=False,bNotConnectable=False,bDefaultValueIsReadOnly=False,bDefaultValueIsIgnored=False,bAdvancedView=False,)
+		   CustomProperties Pin (PinId=80B33882446EDB9EFA5F78B8FF6F8432,PinName="UV",PinToolTip="UV\nVector 2D Structure",PinType.PinCategory="struct",PinType.PinSubCategory="",PinType.PinSubCategoryObject=ScriptStruct'/Script/CoreUObject.Vector2D',PinType.PinSubCategoryMemberReference=(),PinType.PinValueType=(),PinType.bIsMap=False,PinType.bIsSet=False,PinType.bIsArray=False,PinType.bIsReference=False,PinType.bIsConst=False,PinType.bIsWeakPointer=False,LinkedTo=(K2Node_CallFunction_296 B78D125B4929D4D20E4338A6EEB13A7B,),PersistentGuid=00000000000000000000000000000000,bHidden=False,bNotConnectable=False,bDefaultValueIsReadOnly=False,bDefaultValueIsIgnored=False,bAdvancedView=False,)
+		End Object
+			
+	[/PARAM]
+[/OBJECT]
+[/EXCERPT:FunctionUV]
+
+
+[EXCERPT:FunctionUV02]	
+[OBJECT:Snippet]
+	[PARAMLITERAL:id]
+		TraceFromCameraUV2
+	[/PARAMLITERAL]
+	[PARAMLITERAL:type]
+		blueprint
+	[/PARAMLITERAL]
+	[PARAM:image]
+		![](Engine\Rendering\RenderTargets\BlueprintRenderTargets\HowTo\HeightFieldPainter\HFP_VP_Setup_UV_02.png)
+	[/PARAM]
+	[PARAM:snippet]
+		Begin Object Class=K2Node_CustomEvent Name="K2Node_CustomEvent_0"
+		   CustomFunctionName="ApplyPaint"
+		   NodePosX=400
+		   NodePosY=-144
+		   NodeGuid=489876C04EA134CA0B3422B32A20539F
+		   CustomProperties Pin (PinId=7022CEEC415E3A3584654B9C357CC771,PinName="OutputDelegate",Direction="EGPD_Output",PinType.PinCategory="delegate",PinType.PinSubCategory="",PinType.PinSubCategoryObject=None,PinType.PinSubCategoryMemberReference=(MemberParent=BlueprintGeneratedClass'/Game/Docs/Perception/NewBlueprint.NewBlueprint_C',MemberName="ApplyPaint",MemberGuid=489876C04EA134CA0B3422B32A20539F),PinType.PinValueType=(),PinType.bIsMap=False,PinType.bIsSet=False,PinType.bIsArray=False,PinType.bIsReference=False,PinType.bIsConst=False,PinType.bIsWeakPointer=False,PersistentGuid=00000000000000000000000000000000,bHidden=False,bNotConnectable=False,bDefaultValueIsReadOnly=False,bDefaultValueIsIgnored=False,bAdvancedView=False,)
+		   CustomProperties Pin (PinId=48ABE0014DB082A12371E08B2B184C9A,PinName="then",Direction="EGPD_Output",PinType.PinCategory="exec",PinType.PinSubCategory="",PinType.PinSubCategoryObject=None,PinType.PinSubCategoryMemberReference=(),PinType.PinValueType=(),PinType.bIsMap=False,PinType.bIsSet=False,PinType.bIsArray=False,PinType.bIsReference=False,PinType.bIsConst=False,PinType.bIsWeakPointer=False,LinkedTo=(K2Node_CallFunction_368 527F9720481C49179DFF49B166FB5074,),PersistentGuid=00000000000000000000000000000000,bHidden=False,bNotConnectable=False,bDefaultValueIsReadOnly=False,bDefaultValueIsIgnored=False,bAdvancedView=False,)
+		   CustomProperties Pin (PinId=28D508444B6B999E2B762C9FBF985375,PinName="UV",Direction="EGPD_Output",PinType.PinCategory="struct",PinType.PinSubCategory="",PinType.PinSubCategoryObject=ScriptStruct'/Script/CoreUObject.Vector2D',PinType.PinSubCategoryMemberReference=(),PinType.PinValueType=(),PinType.bIsMap=False,PinType.bIsSet=False,PinType.bIsArray=False,PinType.bIsReference=False,PinType.bIsConst=False,PinType.bIsWeakPointer=False,LinkedTo=(K2Node_CallFunction_112 9C1C769941A0BCA1C1E925B9C4796984,),PersistentGuid=00000000000000000000000000000000,bHidden=False,bNotConnectable=False,bDefaultValueIsReadOnly=False,bDefaultValueIsIgnored=False,bAdvancedView=False,)
+		   CustomProperties UserDefinedPin Name="UV" IsArray=0 IsReference=0 PinDir="EGPD_Output" Category=struct SubCategoryObject=/Script/CoreUObject.Vector2D 
+		End Object
+		Begin Object Class=K2Node_Event Name="K2Node_Event_0"
+		   EventReference=(MemberParent=Class'/Script/Engine.Actor',MemberName="ReceiveBeginPlay")
+		   bOverrideFunction=True
+		   NodePosX=-1120
+		   NodePosY=-160
+		   bCommentBubblePinned=True
+		   NodeGuid=E7512AEA456EF28B6DFB2A90C7C09CF6
+		   CustomProperties Pin (PinId=2BE1644F47C274111DE6EBA0A1855DDE,PinName="OutputDelegate",Direction="EGPD_Output",PinType.PinCategory="delegate",PinType.PinSubCategory="",PinType.PinSubCategoryObject=None,PinType.PinSubCategoryMemberReference=(MemberParent=Class'/Script/Engine.Actor',MemberName="ReceiveBeginPlay"),PinType.PinValueType=(),PinType.bIsMap=False,PinType.bIsSet=False,PinType.bIsArray=False,PinType.bIsReference=False,PinType.bIsConst=False,PinType.bIsWeakPointer=False,PersistentGuid=00000000000000000000000000000000,bHidden=False,bNotConnectable=False,bDefaultValueIsReadOnly=False,bDefaultValueIsIgnored=False,bAdvancedView=False,)
+		   CustomProperties Pin (PinId=4C1A3CD54C93D4A74803118ED8E9BE79,PinName="then",Direction="EGPD_Output",PinType.PinCategory="exec",PinType.PinSubCategory="",PinType.PinSubCategoryObject=None,PinType.PinSubCategoryMemberReference=(),PinType.PinValueType=(),PinType.bIsMap=False,PinType.bIsSet=False,PinType.bIsArray=False,PinType.bIsReference=False,PinType.bIsConst=False,PinType.bIsWeakPointer=False,LinkedTo=(K2Node_CallFunction_371 59EBD0D84FA02620751F0CAC38ADA5FF,),PersistentGuid=00000000000000000000000000000000,bHidden=False,bNotConnectable=False,bDefaultValueIsReadOnly=False,bDefaultValueIsIgnored=False,bAdvancedView=False,)
+		End Object
+		Begin Object Class=K2Node_Event Name="K2Node_Event_57"
+		   EventReference=(MemberParent=Class'/Script/Engine.Actor',MemberName="ReceiveTick")
+		   bOverrideFunction=True
+		   NodePosX=-496
+		   NodePosY=336
+		   bCommentBubblePinned=True
+		   NodeGuid=5CB6CAA84FDF4469CF0555A1833F8DC2
+		   CustomProperties Pin (PinId=40C56D29448EE7B7BC3BD4A6B5EEBE71,PinName="OutputDelegate",Direction="EGPD_Output",PinType.PinCategory="delegate",PinType.PinSubCategory="",PinType.PinSubCategoryObject=None,PinType.PinSubCategoryMemberReference=(MemberParent=Class'/Script/Engine.Actor',MemberName="ReceiveTick"),PinType.PinValueType=(),PinType.bIsMap=False,PinType.bIsSet=False,PinType.bIsArray=False,PinType.bIsReference=False,PinType.bIsConst=False,PinType.bIsWeakPointer=False,PersistentGuid=00000000000000000000000000000000,bHidden=False,bNotConnectable=False,bDefaultValueIsReadOnly=False,bDefaultValueIsIgnored=False,bAdvancedView=False,)
+		   CustomProperties Pin (PinId=0371202549E0ACFD644FD4A3E8B79FD5,PinName="then",Direction="EGPD_Output",PinType.PinCategory="exec",PinType.PinSubCategory="",PinType.PinSubCategoryObject=None,PinType.PinSubCategoryMemberReference=(),PinType.PinValueType=(),PinType.bIsMap=False,PinType.bIsSet=False,PinType.bIsArray=False,PinType.bIsReference=False,PinType.bIsConst=False,PinType.bIsWeakPointer=False,LinkedTo=(K2Node_IfThenElse_25 997FDA76493A40CEE900CD81079C2BFC,),PersistentGuid=00000000000000000000000000000000,bHidden=False,bNotConnectable=False,bDefaultValueIsReadOnly=False,bDefaultValueIsIgnored=False,bAdvancedView=False,)
+		   CustomProperties Pin (PinId=8C4E9AE4471B9BF4C1279EBA05819B9C,PinName="DeltaSeconds",PinToolTip="Delta Seconds\nFloat",Direction="EGPD_Output",PinType.PinCategory="float",PinType.PinSubCategory="",PinType.PinSubCategoryObject=None,PinType.PinSubCategoryMemberReference=(),PinType.PinValueType=(),PinType.bIsMap=False,PinType.bIsSet=False,PinType.bIsArray=False,PinType.bIsReference=False,PinType.bIsConst=False,PinType.bIsWeakPointer=False,DefaultValue="0.0",AutogeneratedDefaultValue="0.0",PersistentGuid=00000000000000000000000000000000,bHidden=False,bNotConnectable=False,bDefaultValueIsReadOnly=False,bDefaultValueIsIgnored=False,bAdvancedView=False,)
+		End Object
+		Begin Object Class=K2Node_CallFunction Name="K2Node_CallFunction_359"
+		   FunctionReference=(MemberParent=Class'/Script/Engine.KismetRenderingLibrary',MemberName="ClearRenderTarget2D")
+		   NodePosX=-304
+		   NodePosY=-160
+		   NodeGuid=23CCCF8F40038BFAF646298B2D4213E2
+		   CustomProperties Pin (PinId=50F50A054429DFDB03CCE884B76C4E0A,PinName="execute",PinType.PinCategory="exec",PinType.PinSubCategory="",PinType.PinSubCategoryObject=None,PinType.PinSubCategoryMemberReference=(),PinType.PinValueType=(),PinType.bIsMap=False,PinType.bIsSet=False,PinType.bIsArray=False,PinType.bIsReference=False,PinType.bIsConst=False,PinType.bIsWeakPointer=False,LinkedTo=(K2Node_CallFunction_371 C8543EC14C417EB94C0A40806572C540,),PersistentGuid=00000000000000000000000000000000,bHidden=False,bNotConnectable=False,bDefaultValueIsReadOnly=False,bDefaultValueIsIgnored=False,bAdvancedView=False,)
+		   CustomProperties Pin (PinId=43C94C3C4FB04A4DEADC37A9E3BDE0B2,PinName="then",Direction="EGPD_Output",PinType.PinCategory="exec",PinType.PinSubCategory="",PinType.PinSubCategoryObject=None,PinType.PinSubCategoryMemberReference=(),PinType.PinValueType=(),PinType.bIsMap=False,PinType.bIsSet=False,PinType.bIsArray=False,PinType.bIsReference=False,PinType.bIsConst=False,PinType.bIsWeakPointer=False,LinkedTo=(K2Node_CallFunction_370 542CC73148D38A37BCCECE90262AFC49,),PersistentGuid=00000000000000000000000000000000,bHidden=False,bNotConnectable=False,bDefaultValueIsReadOnly=False,bDefaultValueIsIgnored=False,bAdvancedView=False,)
+		   CustomProperties Pin (PinId=E0E9F0EF4A1840A57B003C8528C144D0,PinName="self",PinFriendlyName=NSLOCTEXT("K2Node", "Target", "Target"),PinType.PinCategory="object",PinType.PinSubCategory="",PinType.PinSubCategoryObject=Class'/Script/Engine.KismetRenderingLibrary',PinType.PinSubCategoryMemberReference=(),PinType.PinValueType=(),PinType.bIsMap=False,PinType.bIsSet=False,PinType.bIsArray=False,PinType.bIsReference=False,PinType.bIsConst=False,PinType.bIsWeakPointer=False,DefaultObject="/Script/Engine.Default__KismetRenderingLibrary",PersistentGuid=00000000000000000000000000000000,bHidden=True,bNotConnectable=False,bDefaultValueIsReadOnly=False,bDefaultValueIsIgnored=False,bAdvancedView=False,)
+		   CustomProperties Pin (PinId=79C9EF7C406FCCFB6D58E7971F16CDC8,PinName="WorldContextObject",PinType.PinCategory="object",PinType.PinSubCategory="",PinType.PinSubCategoryObject=Class'/Script/CoreUObject.Object',PinType.PinSubCategoryMemberReference=(),PinType.PinValueType=(),PinType.bIsMap=False,PinType.bIsSet=False,PinType.bIsArray=False,PinType.bIsReference=False,PinType.bIsConst=False,PinType.bIsWeakPointer=False,PersistentGuid=00000000000000000000000000000000,bHidden=True,bNotConnectable=False,bDefaultValueIsReadOnly=False,bDefaultValueIsIgnored=False,bAdvancedView=False,)
+		   CustomProperties Pin (PinId=314A49414552DC481563F19C80A6DDDF,PinName="TextureRenderTarget",PinType.PinCategory="object",PinType.PinSubCategory="",PinType.PinSubCategoryObject=Class'/Script/Engine.TextureRenderTarget2D',PinType.PinSubCategoryMemberReference=(),PinType.PinValueType=(),PinType.bIsMap=False,PinType.bIsSet=False,PinType.bIsArray=False,PinType.bIsReference=False,PinType.bIsConst=False,PinType.bIsWeakPointer=False,LinkedTo=(K2Node_VariableGet_178 1015D25442FF3C8A3B214DA5835B8950,),PersistentGuid=00000000000000000000000000000000,bHidden=False,bNotConnectable=False,bDefaultValueIsReadOnly=False,bDefaultValueIsIgnored=False,bAdvancedView=False,)
+		   CustomProperties Pin (PinId=F48572574E7B44CBAA94708BF8EAAAE5,PinName="ClearColor",PinType.PinCategory="struct",PinType.PinSubCategory="",PinType.PinSubCategoryObject=ScriptStruct'/Script/CoreUObject.LinearColor',PinType.PinSubCategoryMemberReference=(),PinType.PinValueType=(),PinType.bIsMap=False,PinType.bIsSet=False,PinType.bIsArray=False,PinType.bIsReference=False,PinType.bIsConst=False,PinType.bIsWeakPointer=False,PersistentGuid=00000000000000000000000000000000,bHidden=False,bNotConnectable=False,bDefaultValueIsReadOnly=False,bDefaultValueIsIgnored=False,bAdvancedView=False,)
+		End Object
+		Begin Object Class=K2Node_VariableGet Name="K2Node_VariableGet_178"
+		   VariableReference=(MemberName="HeightfieldRT",MemberGuid=E71A12524941DA802B08039BF7491FF7,bSelfContext=True)
+		   NodePosX=-512
+		   NodePosY=-80
+		   NodeGuid=83090C184988B8F08C9A39A5FCD157EA
+		   CustomProperties Pin (PinId=1015D25442FF3C8A3B214DA5835B8950,PinName="HeightfieldRT",Direction="EGPD_Output",PinType.PinCategory="object",PinType.PinSubCategory="",PinType.PinSubCategoryObject=Class'/Script/Engine.TextureRenderTarget2D',PinType.PinSubCategoryMemberReference=(),PinType.PinValueType=(),PinType.bIsMap=False,PinType.bIsSet=False,PinType.bIsArray=False,PinType.bIsReference=False,PinType.bIsConst=False,PinType.bIsWeakPointer=False,LinkedTo=(K2Node_CallFunction_359 314A49414552DC481563F19C80A6DDDF,),PersistentGuid=00000000000000000000000000000000,bHidden=False,bNotConnectable=False,bDefaultValueIsReadOnly=False,bDefaultValueIsIgnored=False,bAdvancedView=False,)
+		   CustomProperties Pin (PinId=25B1BE2A4EF98E7B1DA2B2ADF9EF4385,PinName="self",PinFriendlyName=NSLOCTEXT("K2Node", "Target", "Target"),PinType.PinCategory="object",PinType.PinSubCategory="",PinType.PinSubCategoryObject=BlueprintGeneratedClass'/Game/Docs/Perception/NewBlueprint.NewBlueprint_C',PinType.PinSubCategoryMemberReference=(),PinType.PinValueType=(),PinType.bIsMap=False,PinType.bIsSet=False,PinType.bIsArray=False,PinType.bIsReference=False,PinType.bIsConst=False,PinType.bIsWeakPointer=False,PersistentGuid=00000000000000000000000000000000,bHidden=True,bNotConnectable=False,bDefaultValueIsReadOnly=False,bDefaultValueIsIgnored=False,bAdvancedView=False,)
+		End Object
+		Begin Object Class=K2Node_InputKey Name="K2Node_InputKey_19"
+		   InputKey=LeftMouseButton
+		   NodePosX=-1120
+		   NodePosY=304
+		   NodeGuid=28201655455FA58A96734383B5F3AE53
+		   CustomProperties Pin (PinId=4D5FE2B549E275EDF0E659896A9261E4,PinName="Pressed",Direction="EGPD_Output",PinType.PinCategory="exec",PinType.PinSubCategory="",PinType.PinSubCategoryObject=None,PinType.PinSubCategoryMemberReference=(),PinType.PinValueType=(),PinType.bIsMap=False,PinType.bIsSet=False,PinType.bIsArray=False,PinType.bIsReference=False,PinType.bIsConst=False,PinType.bIsWeakPointer=False,LinkedTo=(K2Node_VariableSet_41 DF34F0B548E46FCCB2AB54856F7CFFC5,),PersistentGuid=00000000000000000000000000000000,bHidden=False,bNotConnectable=False,bDefaultValueIsReadOnly=False,bDefaultValueIsIgnored=False,bAdvancedView=False,)
+		   CustomProperties Pin (PinId=B150F43D4C370A3041E5B9A2DD8FC834,PinName="Released",Direction="EGPD_Output",PinType.PinCategory="exec",PinType.PinSubCategory="",PinType.PinSubCategoryObject=None,PinType.PinSubCategoryMemberReference=(),PinType.PinValueType=(),PinType.bIsMap=False,PinType.bIsSet=False,PinType.bIsArray=False,PinType.bIsReference=False,PinType.bIsConst=False,PinType.bIsWeakPointer=False,LinkedTo=(K2Node_VariableSet_42 9312E8C94B08C9FAF76638810A0F61A1,),PersistentGuid=00000000000000000000000000000000,bHidden=False,bNotConnectable=False,bDefaultValueIsReadOnly=False,bDefaultValueIsIgnored=False,bAdvancedView=False,)
+		   CustomProperties Pin (PinId=43D7BA3045BD03142E77E597FD7EAD86,PinName="Key",Direction="EGPD_Output",PinType.PinCategory="struct",PinType.PinSubCategory="",PinType.PinSubCategoryObject=ScriptStruct'/Script/InputCore.Key',PinType.PinSubCategoryMemberReference=(),PinType.PinValueType=(),PinType.bIsMap=False,PinType.bIsSet=False,PinType.bIsArray=False,PinType.bIsReference=False,PinType.bIsConst=False,PinType.bIsWeakPointer=False,DefaultValue="AnyKey",PersistentGuid=00000000000000000000000000000000,bHidden=False,bNotConnectable=False,bDefaultValueIsReadOnly=False,bDefaultValueIsIgnored=False,bAdvancedView=False,)
+		End Object
+		Begin Object Class=K2Node_VariableSet Name="K2Node_VariableSet_41"
+		   VariableReference=(MemberName="MouseDown",MemberGuid=FD84A8304E07A6A8E93E818EBF07943E,bSelfContext=True)
+		   NodePosX=-816
+		   NodePosY=288
+		   NodeGuid=FA5B0DD743D55F513448B5B5BC8244ED
+		   CustomProperties Pin (PinId=DF34F0B548E46FCCB2AB54856F7CFFC5,PinName="execute",PinType.PinCategory="exec",PinType.PinSubCategory="",PinType.PinSubCategoryObject=None,PinType.PinSubCategoryMemberReference=(),PinType.PinValueType=(),PinType.bIsMap=False,PinType.bIsSet=False,PinType.bIsArray=False,PinType.bIsReference=False,PinType.bIsConst=False,PinType.bIsWeakPointer=False,LinkedTo=(K2Node_InputKey_19 4D5FE2B549E275EDF0E659896A9261E4,),PersistentGuid=00000000000000000000000000000000,bHidden=False,bNotConnectable=False,bDefaultValueIsReadOnly=False,bDefaultValueIsIgnored=False,bAdvancedView=False,)
+		   CustomProperties Pin (PinId=A68EDFFF4BFBADBD6B1C60AB8EEE7A3D,PinName="then",Direction="EGPD_Output",PinType.PinCategory="exec",PinType.PinSubCategory="",PinType.PinSubCategoryObject=None,PinType.PinSubCategoryMemberReference=(),PinType.PinValueType=(),PinType.bIsMap=False,PinType.bIsSet=False,PinType.bIsArray=False,PinType.bIsReference=False,PinType.bIsConst=False,PinType.bIsWeakPointer=False,PersistentGuid=00000000000000000000000000000000,bHidden=False,bNotConnectable=False,bDefaultValueIsReadOnly=False,bDefaultValueIsIgnored=False,bAdvancedView=False,)
+		   CustomProperties Pin (PinId=8F7952DE47C7BACDD7565587C283D7EE,PinName="MouseDown",PinType.PinCategory="bool",PinType.PinSubCategory="",PinType.PinSubCategoryObject=None,PinType.PinSubCategoryMemberReference=(),PinType.PinValueType=(),PinType.bIsMap=False,PinType.bIsSet=False,PinType.bIsArray=False,PinType.bIsReference=False,PinType.bIsConst=False,PinType.bIsWeakPointer=False,DefaultValue="true",AutogeneratedDefaultValue="false",PersistentGuid=00000000000000000000000000000000,bHidden=False,bNotConnectable=False,bDefaultValueIsReadOnly=False,bDefaultValueIsIgnored=False,bAdvancedView=False,)
+		   CustomProperties Pin (PinId=6799AEA64465D17287DE96888FB9C1D6,PinName="Output_Get",PinToolTip="Retrieves the value of the variable, can use instead of a separate Get node",Direction="EGPD_Output",PinType.PinCategory="bool",PinType.PinSubCategory="",PinType.PinSubCategoryObject=None,PinType.PinSubCategoryMemberReference=(),PinType.PinValueType=(),PinType.bIsMap=False,PinType.bIsSet=False,PinType.bIsArray=False,PinType.bIsReference=False,PinType.bIsConst=False,PinType.bIsWeakPointer=False,DefaultValue="false",AutogeneratedDefaultValue="false",PersistentGuid=00000000000000000000000000000000,bHidden=False,bNotConnectable=False,bDefaultValueIsReadOnly=False,bDefaultValueIsIgnored=False,bAdvancedView=False,)
+		   CustomProperties Pin (PinId=E69996B6422209F91F0E5084392EBA86,PinName="self",PinFriendlyName=NSLOCTEXT("K2Node", "Target", "Target"),PinType.PinCategory="object",PinType.PinSubCategory="",PinType.PinSubCategoryObject=BlueprintGeneratedClass'/Game/Docs/Perception/NewBlueprint.NewBlueprint_C',PinType.PinSubCategoryMemberReference=(),PinType.PinValueType=(),PinType.bIsMap=False,PinType.bIsSet=False,PinType.bIsArray=False,PinType.bIsReference=False,PinType.bIsConst=False,PinType.bIsWeakPointer=False,PersistentGuid=00000000000000000000000000000000,bHidden=True,bNotConnectable=False,bDefaultValueIsReadOnly=False,bDefaultValueIsIgnored=False,bAdvancedView=False,)
+		End Object
+		Begin Object Class=K2Node_VariableSet Name="K2Node_VariableSet_42"
+		   VariableReference=(MemberName="MouseDown",MemberGuid=209E1D534C15EFFA81760DA2C7F9D56C,bSelfContext=True)
+		   NodePosX=-816
+		   NodePosY=384
+		   NodeGuid=CFBE5C5A4DBDC975EC5236B52A067CBA
+		   CustomProperties Pin (PinId=9312E8C94B08C9FAF76638810A0F61A1,PinName="execute",PinType.PinCategory="exec",PinType.PinSubCategory="",PinType.PinSubCategoryObject=None,PinType.PinSubCategoryMemberReference=(),PinType.PinValueType=(),PinType.bIsMap=False,PinType.bIsSet=False,PinType.bIsArray=False,PinType.bIsReference=False,PinType.bIsConst=False,PinType.bIsWeakPointer=False,LinkedTo=(K2Node_InputKey_19 B150F43D4C370A3041E5B9A2DD8FC834,),PersistentGuid=00000000000000000000000000000000,bHidden=False,bNotConnectable=False,bDefaultValueIsReadOnly=False,bDefaultValueIsIgnored=False,bAdvancedView=False,)
+		   CustomProperties Pin (PinId=CB1027A1495B8948A22EB58160AA4C3E,PinName="then",Direction="EGPD_Output",PinType.PinCategory="exec",PinType.PinSubCategory="",PinType.PinSubCategoryObject=None,PinType.PinSubCategoryMemberReference=(),PinType.PinValueType=(),PinType.bIsMap=False,PinType.bIsSet=False,PinType.bIsArray=False,PinType.bIsReference=False,PinType.bIsConst=False,PinType.bIsWeakPointer=False,PersistentGuid=00000000000000000000000000000000,bHidden=False,bNotConnectable=False,bDefaultValueIsReadOnly=False,bDefaultValueIsIgnored=False,bAdvancedView=False,)
+		   CustomProperties Pin (PinId=3754DE984E6A106CC45C95831AEDF03A,PinName="MouseDown",PinType.PinCategory="bool",PinType.PinSubCategory="",PinType.PinSubCategoryObject=None,PinType.PinSubCategoryMemberReference=(),PinType.PinValueType=(),PinType.bIsMap=False,PinType.bIsSet=False,PinType.bIsArray=False,PinType.bIsReference=False,PinType.bIsConst=False,PinType.bIsWeakPointer=False,DefaultValue="false",AutogeneratedDefaultValue="false",PersistentGuid=00000000000000000000000000000000,bHidden=False,bNotConnectable=False,bDefaultValueIsReadOnly=False,bDefaultValueIsIgnored=False,bAdvancedView=False,)
+		   CustomProperties Pin (PinId=CD9E960E4A4357C090CAFF8D5F8D9AB9,PinName="Output_Get",PinToolTip="Retrieves the value of the variable, can use instead of a separate Get node",Direction="EGPD_Output",PinType.PinCategory="bool",PinType.PinSubCategory="",PinType.PinSubCategoryObject=None,PinType.PinSubCategoryMemberReference=(),PinType.PinValueType=(),PinType.bIsMap=False,PinType.bIsSet=False,PinType.bIsArray=False,PinType.bIsReference=False,PinType.bIsConst=False,PinType.bIsWeakPointer=False,DefaultValue="false",AutogeneratedDefaultValue="false",PersistentGuid=00000000000000000000000000000000,bHidden=False,bNotConnectable=False,bDefaultValueIsReadOnly=False,bDefaultValueIsIgnored=False,bAdvancedView=False,)
+		   CustomProperties Pin (PinId=00D24E1D4EEA0556B7962191FA32BCF8,PinName="self",PinFriendlyName=NSLOCTEXT("K2Node", "Target", "Target"),PinType.PinCategory="object",PinType.PinSubCategory="",PinType.PinSubCategoryObject=BlueprintGeneratedClass'/Game/Docs/Perception/NewBlueprint.NewBlueprint_C',PinType.PinSubCategoryMemberReference=(),PinType.PinValueType=(),PinType.bIsMap=False,PinType.bIsSet=False,PinType.bIsArray=False,PinType.bIsReference=False,PinType.bIsConst=False,PinType.bIsWeakPointer=False,PersistentGuid=00000000000000000000000000000000,bHidden=True,bNotConnectable=False,bDefaultValueIsReadOnly=False,bDefaultValueIsIgnored=False,bAdvancedView=False,)
+		End Object
+		Begin Object Class=K2Node_CallFunction Name="K2Node_CallFunction_360"
+		   FunctionReference=(MemberName="TraceFromCamera",MemberGuid=C96978014B119460D779C7A90DE17A56,bSelfContext=True)
+		   NodePosX=-32
+		   NodePosY=320
+		   NodeGuid=90A4A971455DE64EDA359887F3440815
+		   CustomProperties Pin (PinId=9F3F898442AC24B5F566A9A6E3D64A03,PinName="execute",PinToolTip="\nExec",PinType.PinCategory="exec",PinType.PinSubCategory="",PinType.PinSubCategoryObject=None,PinType.PinSubCategoryMemberReference=(),PinType.PinValueType=(),PinType.bIsMap=False,PinType.bIsSet=False,PinType.bIsArray=False,PinType.bIsReference=False,PinType.bIsConst=False,PinType.bIsWeakPointer=False,LinkedTo=(K2Node_IfThenElse_25 358D51FF4B717F8D875E549790D394DB,),PersistentGuid=00000000000000000000000000000000,bHidden=False,bNotConnectable=False,bDefaultValueIsReadOnly=False,bDefaultValueIsIgnored=False,bAdvancedView=False,)
+		   CustomProperties Pin (PinId=BF6710D74C9EC10D6D5690B01743CC3D,PinName="then",PinToolTip="\nExec",Direction="EGPD_Output",PinType.PinCategory="exec",PinType.PinSubCategory="",PinType.PinSubCategoryObject=None,PinType.PinSubCategoryMemberReference=(),PinType.PinValueType=(),PinType.bIsMap=False,PinType.bIsSet=False,PinType.bIsArray=False,PinType.bIsReference=False,PinType.bIsConst=False,PinType.bIsWeakPointer=False,PersistentGuid=00000000000000000000000000000000,bHidden=False,bNotConnectable=False,bDefaultValueIsReadOnly=False,bDefaultValueIsIgnored=False,bAdvancedView=False,)
+		   CustomProperties Pin (PinId=53418ABD4C9A655761FB14B9217B591D,PinName="self",PinFriendlyName=NSLOCTEXT("K2Node", "Target", "Target"),PinToolTip="Target\nSelf Reference",PinType.PinCategory="object",PinType.PinSubCategory="self",PinType.PinSubCategoryObject=None,PinType.PinSubCategoryMemberReference=(),PinType.PinValueType=(),PinType.bIsMap=False,PinType.bIsSet=False,PinType.bIsArray=False,PinType.bIsReference=False,PinType.bIsConst=False,PinType.bIsWeakPointer=False,PersistentGuid=00000000000000000000000000000000,bHidden=False,bNotConnectable=False,bDefaultValueIsReadOnly=False,bDefaultValueIsIgnored=False,bAdvancedView=False,)
+		End Object
+		Begin Object Class=K2Node_VariableGet Name="K2Node_VariableGet_179"
+		   VariableReference=(MemberName="MouseDown",MemberGuid=209E1D534C15EFFA81760DA2C7F9D56C,bSelfContext=True)
+		   NodePosX=-496
+		   NodePosY=448
+		   NodeGuid=2CB21A884A375F7EB36E19A7CD63CEC1
+		   CustomProperties Pin (PinId=200A09224BC64E9B688537BFADAE5937,PinName="MouseDown",Direction="EGPD_Output",PinType.PinCategory="bool",PinType.PinSubCategory="",PinType.PinSubCategoryObject=None,PinType.PinSubCategoryMemberReference=(),PinType.PinValueType=(),PinType.bIsMap=False,PinType.bIsSet=False,PinType.bIsArray=False,PinType.bIsReference=False,PinType.bIsConst=False,PinType.bIsWeakPointer=False,AutogeneratedDefaultValue="false",LinkedTo=(K2Node_IfThenElse_25 DBE9930D445160424F99CAB156B1E26F,),PersistentGuid=00000000000000000000000000000000,bHidden=False,bNotConnectable=False,bDefaultValueIsReadOnly=False,bDefaultValueIsIgnored=False,bAdvancedView=False,)
+		   CustomProperties Pin (PinId=1DDB72E94FB45FA218F0B0A9A9D8A8DF,PinName="self",PinFriendlyName=NSLOCTEXT("K2Node", "Target", "Target"),PinType.PinCategory="object",PinType.PinSubCategory="",PinType.PinSubCategoryObject=BlueprintGeneratedClass'/Game/Docs/Perception/NewBlueprint.NewBlueprint_C',PinType.PinSubCategoryMemberReference=(),PinType.PinValueType=(),PinType.bIsMap=False,PinType.bIsSet=False,PinType.bIsArray=False,PinType.bIsReference=False,PinType.bIsConst=False,PinType.bIsWeakPointer=False,PersistentGuid=00000000000000000000000000000000,bHidden=True,bNotConnectable=False,bDefaultValueIsReadOnly=False,bDefaultValueIsIgnored=False,bAdvancedView=False,)
+		End Object
+		Begin Object Class=K2Node_IfThenElse Name="K2Node_IfThenElse_25"
+		   NodePosX=-272
+		   NodePosY=336
+		   NodeGuid=51D0E8324194B328FE444CAD1A34E690
+		   CustomProperties Pin (PinId=997FDA76493A40CEE900CD81079C2BFC,PinName="execute",PinType.PinCategory="exec",PinType.PinSubCategory="",PinType.PinSubCategoryObject=None,PinType.PinSubCategoryMemberReference=(),PinType.PinValueType=(),PinType.bIsMap=False,PinType.bIsSet=False,PinType.bIsArray=False,PinType.bIsReference=False,PinType.bIsConst=False,PinType.bIsWeakPointer=False,LinkedTo=(K2Node_Event_57 0371202549E0ACFD644FD4A3E8B79FD5,),PersistentGuid=00000000000000000000000000000000,bHidden=False,bNotConnectable=False,bDefaultValueIsReadOnly=False,bDefaultValueIsIgnored=False,bAdvancedView=False,)
+		   CustomProperties Pin (PinId=DBE9930D445160424F99CAB156B1E26F,PinName="Condition",PinType.PinCategory="bool",PinType.PinSubCategory="",PinType.PinSubCategoryObject=None,PinType.PinSubCategoryMemberReference=(),PinType.PinValueType=(),PinType.bIsMap=False,PinType.bIsSet=False,PinType.bIsArray=False,PinType.bIsReference=False,PinType.bIsConst=False,PinType.bIsWeakPointer=False,DefaultValue="false",AutogeneratedDefaultValue="true",LinkedTo=(K2Node_VariableGet_179 200A09224BC64E9B688537BFADAE5937,),PersistentGuid=00000000000000000000000000000000,bHidden=False,bNotConnectable=False,bDefaultValueIsReadOnly=False,bDefaultValueIsIgnored=False,bAdvancedView=False,)
+		   CustomProperties Pin (PinId=358D51FF4B717F8D875E549790D394DB,PinName="then",PinFriendlyName=NSLOCTEXT("K2Node", "true", "true"),Direction="EGPD_Output",PinType.PinCategory="exec",PinType.PinSubCategory="",PinType.PinSubCategoryObject=None,PinType.PinSubCategoryMemberReference=(),PinType.PinValueType=(),PinType.bIsMap=False,PinType.bIsSet=False,PinType.bIsArray=False,PinType.bIsReference=False,PinType.bIsConst=False,PinType.bIsWeakPointer=False,LinkedTo=(K2Node_CallFunction_360 9F3F898442AC24B5F566A9A6E3D64A03,),PersistentGuid=00000000000000000000000000000000,bHidden=False,bNotConnectable=False,bDefaultValueIsReadOnly=False,bDefaultValueIsIgnored=False,bAdvancedView=False,)
+		   CustomProperties Pin (PinId=53B9B8674F2C0FDC804A088C1129D0E3,PinName="else",PinFriendlyName=NSLOCTEXT("K2Node", "false", "false"),Direction="EGPD_Output",PinType.PinCategory="exec",PinType.PinSubCategory="",PinType.PinSubCategoryObject=None,PinType.PinSubCategoryMemberReference=(),PinType.PinValueType=(),PinType.bIsMap=False,PinType.bIsSet=False,PinType.bIsArray=False,PinType.bIsReference=False,PinType.bIsConst=False,PinType.bIsWeakPointer=False,PersistentGuid=00000000000000000000000000000000,bHidden=False,bNotConnectable=False,bDefaultValueIsReadOnly=False,bDefaultValueIsIgnored=False,bAdvancedView=False,)
+		End Object
+		Begin Object Class=EdGraphNode_Comment Name="EdGraphNode_Comment_119"
+		   NodePosX=-544
+		   NodePosY=272
+		   NodeWidth=768
+		   NodeHeight=224
+		   NodeComment="Shoot when mouse is down"
+		   NodeGuid=C9D019B041EFA04F15F093ABA7EF6018
+		End Object
+		Begin Object Class=K2Node_CallFunction Name="K2Node_CallFunction_366"
+		   FunctionReference=(MemberParent=Class'/Script/Engine.KismetRenderingLibrary',MemberName="DrawMaterialToRenderTarget")
+		   NodePosX=2096
+		   NodePosY=-128
+		   NodeGuid=2B62A983457F2C0BBE9FA9B06CA94F82
+		   CustomProperties Pin (PinId=D89DC81641CA4093CA72DAAF5CACC388,PinName="execute",PinToolTip="\nExec",PinType.PinCategory="exec",PinType.PinSubCategory="",PinType.PinSubCategoryObject=None,PinType.PinSubCategoryMemberReference=(),PinType.PinValueType=(),PinType.bIsMap=False,PinType.bIsSet=False,PinType.bIsArray=False,PinType.bIsReference=False,PinType.bIsConst=False,PinType.bIsWeakPointer=False,LinkedTo=(K2Node_CallFunction_367 F5065258491F11F07B8E7B9BF2A38140,),PersistentGuid=00000000000000000000000000000000,bHidden=False,bNotConnectable=False,bDefaultValueIsReadOnly=False,bDefaultValueIsIgnored=False,bAdvancedView=False,)
+		   CustomProperties Pin (PinId=18C8362646DD324E0B3EE09FE2E9691F,PinName="then",PinToolTip="\nExec",Direction="EGPD_Output",PinType.PinCategory="exec",PinType.PinSubCategory="",PinType.PinSubCategoryObject=None,PinType.PinSubCategoryMemberReference=(),PinType.PinValueType=(),PinType.bIsMap=False,PinType.bIsSet=False,PinType.bIsArray=False,PinType.bIsReference=False,PinType.bIsConst=False,PinType.bIsWeakPointer=False,PersistentGuid=00000000000000000000000000000000,bHidden=False,bNotConnectable=False,bDefaultValueIsReadOnly=False,bDefaultValueIsIgnored=False,bAdvancedView=False,)
+		   CustomProperties Pin (PinId=3E4C4C0C47F9C599C77C0A94DAE5CDBC,PinName="self",PinFriendlyName=NSLOCTEXT("K2Node", "Target", "Target"),PinToolTip="Target\nKismet Rendering Library Reference",PinType.PinCategory="object",PinType.PinSubCategory="",PinType.PinSubCategoryObject=Class'/Script/Engine.KismetRenderingLibrary',PinType.PinSubCategoryMemberReference=(),PinType.PinValueType=(),PinType.bIsMap=False,PinType.bIsSet=False,PinType.bIsArray=False,PinType.bIsReference=False,PinType.bIsConst=False,PinType.bIsWeakPointer=False,DefaultObject="/Script/Engine.Default__KismetRenderingLibrary",PersistentGuid=00000000000000000000000000000000,bHidden=True,bNotConnectable=False,bDefaultValueIsReadOnly=False,bDefaultValueIsIgnored=False,bAdvancedView=False,)
+		   CustomProperties Pin (PinId=5EE9579647D9A511B63300AEE24E89BF,PinName="WorldContextObject",PinToolTip="World Context Object\nObject Reference",PinType.PinCategory="object",PinType.PinSubCategory="",PinType.PinSubCategoryObject=Class'/Script/CoreUObject.Object',PinType.PinSubCategoryMemberReference=(),PinType.PinValueType=(),PinType.bIsMap=False,PinType.bIsSet=False,PinType.bIsArray=False,PinType.bIsReference=False,PinType.bIsConst=False,PinType.bIsWeakPointer=False,PersistentGuid=00000000000000000000000000000000,bHidden=True,bNotConnectable=False,bDefaultValueIsReadOnly=False,bDefaultValueIsIgnored=False,bAdvancedView=False,)
+		   CustomProperties Pin (PinId=B57C770C4D827E19FDE1F885B23E018C,PinName="TextureRenderTarget",PinToolTip="Texture Render Target\nTexture Render Target 2D Reference",PinType.PinCategory="object",PinType.PinSubCategory="",PinType.PinSubCategoryObject=Class'/Script/Engine.TextureRenderTarget2D',PinType.PinSubCategoryMemberReference=(),PinType.PinValueType=(),PinType.bIsMap=False,PinType.bIsSet=False,PinType.bIsArray=False,PinType.bIsReference=False,PinType.bIsConst=False,PinType.bIsWeakPointer=False,LinkedTo=(K2Node_VariableGet_182 A993DA7E42B925003560FEAF90CAD2D6,),PersistentGuid=00000000000000000000000000000000,bHidden=False,bNotConnectable=False,bDefaultValueIsReadOnly=False,bDefaultValueIsIgnored=False,bAdvancedView=False,)
+		   CustomProperties Pin (PinId=D99B80D64C0171739FBC10A4E4C540EC,PinName="Material",PinToolTip="Material\nMaterial Interface Reference",PinType.PinCategory="object",PinType.PinSubCategory="",PinType.PinSubCategoryObject=Class'/Script/Engine.MaterialInterface',PinType.PinSubCategoryMemberReference=(),PinType.PinValueType=(),PinType.bIsMap=False,PinType.bIsSet=False,PinType.bIsArray=False,PinType.bIsReference=False,PinType.bIsConst=False,PinType.bIsWeakPointer=False,LinkedTo=(K2Node_Knot_1 9617DBFB4F732CAC88F1169322A29AC2,),PersistentGuid=00000000000000000000000000000000,bHidden=False,bNotConnectable=False,bDefaultValueIsReadOnly=False,bDefaultValueIsIgnored=False,bAdvancedView=False,)
+		End Object
+		Begin Object Class=EdGraphNode_Comment Name="EdGraphNode_Comment_121"
+		   CommentDepth=-2
+		   NodePosX=336
+		   NodePosY=-224
+		   NodeWidth=2080
+		   NodeHeight=608
+		   NodeComment="Take Damage - paint additive splat into heightfield"
+		   NodeGuid=CB882C4F4EB796FCE22144A411211BE2
+		End Object
+		Begin Object Class=K2Node_CallFunction Name="K2Node_CallFunction_367"
+		   FunctionReference=(MemberParent=Class'/Script/Engine.MaterialInstanceDynamic',MemberName="SetVectorParameterValue")
+		   NodePosX=1712
+		   NodePosY=-144
+		   NodeGuid=8240F9D14E527E479D48A19E3919DC8C
+		   CustomProperties Pin (PinId=88C40CCB4DEBDA28F201A388C6DF16B9,PinName="execute",PinToolTip="\nExec",PinType.PinCategory="exec",PinType.PinSubCategory="",PinType.PinSubCategoryObject=None,PinType.PinSubCategoryMemberReference=(),PinType.PinValueType=(),PinType.bIsMap=False,PinType.bIsSet=False,PinType.bIsArray=False,PinType.bIsReference=False,PinType.bIsConst=False,PinType.bIsWeakPointer=False,LinkedTo=(K2Node_CallFunction_369 5826E451454A5FA8D03AB1B3928B3090,),PersistentGuid=00000000000000000000000000000000,bHidden=False,bNotConnectable=False,bDefaultValueIsReadOnly=False,bDefaultValueIsIgnored=False,bAdvancedView=False,)
+		   CustomProperties Pin (PinId=F5065258491F11F07B8E7B9BF2A38140,PinName="then",PinToolTip="\nExec",Direction="EGPD_Output",PinType.PinCategory="exec",PinType.PinSubCategory="",PinType.PinSubCategoryObject=None,PinType.PinSubCategoryMemberReference=(),PinType.PinValueType=(),PinType.bIsMap=False,PinType.bIsSet=False,PinType.bIsArray=False,PinType.bIsReference=False,PinType.bIsConst=False,PinType.bIsWeakPointer=False,LinkedTo=(K2Node_CallFunction_366 D89DC81641CA4093CA72DAAF5CACC388,),PersistentGuid=00000000000000000000000000000000,bHidden=False,bNotConnectable=False,bDefaultValueIsReadOnly=False,bDefaultValueIsIgnored=False,bAdvancedView=False,)
+		   CustomProperties Pin (PinId=FF3048F146D54F1C4AA4A48C221B5457,PinName="self",PinFriendlyName=NSLOCTEXT("K2Node", "Target", "Target"),PinToolTip="Target\nMaterial Instance Dynamic Reference",PinType.PinCategory="object",PinType.PinSubCategory="",PinType.PinSubCategoryObject=Class'/Script/Engine.MaterialInstanceDynamic',PinType.PinSubCategoryMemberReference=(),PinType.PinValueType=(),PinType.bIsMap=False,PinType.bIsSet=False,PinType.bIsArray=False,PinType.bIsReference=False,PinType.bIsConst=False,PinType.bIsWeakPointer=False,LinkedTo=(K2Node_Knot_3 650149D74494B25528C925B6D9B777EB,),PersistentGuid=00000000000000000000000000000000,bHidden=False,bNotConnectable=False,bDefaultValueIsReadOnly=False,bDefaultValueIsIgnored=False,bAdvancedView=False,)
+		   CustomProperties Pin (PinId=50E108F844D08F285B9864B792806572,PinName="ParameterName",PinToolTip="Parameter Name\nName",PinType.PinCategory="name",PinType.PinSubCategory="",PinType.PinSubCategoryObject=None,PinType.PinSubCategoryMemberReference=(),PinType.PinValueType=(),PinType.bIsMap=False,PinType.bIsSet=False,PinType.bIsArray=False,PinType.bIsReference=False,PinType.bIsConst=False,PinType.bIsWeakPointer=False,DefaultValue="ForcePosition",AutogeneratedDefaultValue="None",PersistentGuid=00000000000000000000000000000000,bHidden=False,bNotConnectable=False,bDefaultValueIsReadOnly=False,bDefaultValueIsIgnored=False,bAdvancedView=False,)
+		   CustomProperties Pin (PinId=1523A2CA42868EC87B90F39450925710,PinName="Value",PinToolTip="Value\nLinear Color Structure",PinType.PinCategory="struct",PinType.PinSubCategory="",PinType.PinSubCategoryObject=ScriptStruct'/Script/CoreUObject.LinearColor',PinType.PinSubCategoryMemberReference=(),PinType.PinValueType=(),PinType.bIsMap=False,PinType.bIsSet=False,PinType.bIsArray=False,PinType.bIsReference=False,PinType.bIsConst=False,PinType.bIsWeakPointer=False,LinkedTo=(K2Node_CallFunction_111 ED6E4C604B7A3F2181D74F98664FDCDE,),PersistentGuid=00000000000000000000000000000000,bHidden=False,bNotConnectable=False,bDefaultValueIsReadOnly=False,bDefaultValueIsIgnored=False,bAdvancedView=False,)
+		End Object
+		Begin Object Class=K2Node_VariableGet Name="K2Node_VariableGet_181"
+		   VariableReference=(MemberName="ForceSplatMID",MemberGuid=5E84E3334DD38E72B951FAB5E2DDFCAC,bSelfContext=True)
+		   NodePosX=432
+		   NodePosY=48
+		   NodeGuid=9F47AC264CC72A60B7FFF980642EE7B6
+		   CustomProperties Pin (PinId=16610BCA473F60D2FB851F8B3E4736A4,PinName="ForceSplatMID",Direction="EGPD_Output",PinType.PinCategory="object",PinType.PinSubCategory="",PinType.PinSubCategoryObject=Class'/Script/Engine.MaterialInstanceDynamic',PinType.PinSubCategoryMemberReference=(),PinType.PinValueType=(),PinType.bIsMap=False,PinType.bIsSet=False,PinType.bIsArray=False,PinType.bIsReference=False,PinType.bIsConst=False,PinType.bIsWeakPointer=False,LinkedTo=(K2Node_CallFunction_368 55F338D54A4933AB828D2DB26782C387,K2Node_Knot_2 21486A174CB3882A9431A9803149133A,),PersistentGuid=00000000000000000000000000000000,bHidden=False,bNotConnectable=False,bDefaultValueIsReadOnly=False,bDefaultValueIsIgnored=False,bAdvancedView=False,)
+		   CustomProperties Pin (PinId=218485C84FA2252331D0F9BE4F94155B,PinName="self",PinFriendlyName=NSLOCTEXT("K2Node", "Target", "Target"),PinType.PinCategory="object",PinType.PinSubCategory="",PinType.PinSubCategoryObject=BlueprintGeneratedClass'/Game/Docs/Perception/NewBlueprint.NewBlueprint_C',PinType.PinSubCategoryMemberReference=(),PinType.PinValueType=(),PinType.bIsMap=False,PinType.bIsSet=False,PinType.bIsArray=False,PinType.bIsReference=False,PinType.bIsConst=False,PinType.bIsWeakPointer=False,PersistentGuid=00000000000000000000000000000000,bHidden=True,bNotConnectable=False,bDefaultValueIsReadOnly=False,bDefaultValueIsIgnored=False,bAdvancedView=False,)
+		End Object
+		Begin Object Class=K2Node_VariableGet Name="K2Node_VariableGet_182"
+		   VariableReference=(MemberName="HeightfieldRT",MemberGuid=00C4CF8E4BAEA36D336DFE8BE323CB97,bSelfContext=True)
+		   NodePosX=1808
+		   NodePosY=48
+		   NodeGuid=0B79E4F94075C08955EC108CBBE48AD2
+		   CustomProperties Pin (PinId=A993DA7E42B925003560FEAF90CAD2D6,PinName="HeightfieldRT",Direction="EGPD_Output",PinType.PinCategory="object",PinType.PinSubCategory="",PinType.PinSubCategoryObject=Class'/Script/Engine.TextureRenderTarget2D',PinType.PinSubCategoryMemberReference=(),PinType.PinValueType=(),PinType.bIsMap=False,PinType.bIsSet=False,PinType.bIsArray=False,PinType.bIsReference=False,PinType.bIsConst=False,PinType.bIsWeakPointer=False,LinkedTo=(K2Node_CallFunction_366 B57C770C4D827E19FDE1F885B23E018C,),PersistentGuid=00000000000000000000000000000000,bHidden=False,bNotConnectable=False,bDefaultValueIsReadOnly=False,bDefaultValueIsIgnored=False,bAdvancedView=False,)
+		   CustomProperties Pin (PinId=31956CB24C33CCA7B8693B9FD3CE372C,PinName="self",PinFriendlyName=NSLOCTEXT("K2Node", "Target", "Target"),PinType.PinCategory="object",PinType.PinSubCategory="",PinType.PinSubCategoryObject=BlueprintGeneratedClass'/Game/Docs/Perception/NewBlueprint.NewBlueprint_C',PinType.PinSubCategoryMemberReference=(),PinType.PinValueType=(),PinType.bIsMap=False,PinType.bIsSet=False,PinType.bIsArray=False,PinType.bIsReference=False,PinType.bIsConst=False,PinType.bIsWeakPointer=False,PersistentGuid=00000000000000000000000000000000,bHidden=True,bNotConnectable=False,bDefaultValueIsReadOnly=False,bDefaultValueIsIgnored=False,bAdvancedView=False,)
+		End Object
+		Begin Object Class=K2Node_CallFunction Name="K2Node_CallFunction_368"
+		   FunctionReference=(MemberParent=Class'/Script/Engine.MaterialInstanceDynamic',MemberName="SetScalarParameterValue")
+		   NodePosX=832
+		   NodePosY=-144
+		   NodeGuid=A332568149C42FB044CCBA886F149530
+		   CustomProperties Pin (PinId=527F9720481C49179DFF49B166FB5074,PinName="execute",PinToolTip="\nExec",PinType.PinCategory="exec",PinType.PinSubCategory="",PinType.PinSubCategoryObject=None,PinType.PinSubCategoryMemberReference=(),PinType.PinValueType=(),PinType.bIsMap=False,PinType.bIsSet=False,PinType.bIsArray=False,PinType.bIsReference=False,PinType.bIsConst=False,PinType.bIsWeakPointer=False,LinkedTo=(K2Node_CustomEvent_0 48ABE0014DB082A12371E08B2B184C9A,),PersistentGuid=00000000000000000000000000000000,bHidden=False,bNotConnectable=False,bDefaultValueIsReadOnly=False,bDefaultValueIsIgnored=False,bAdvancedView=False,)
+		   CustomProperties Pin (PinId=5826E451454A5FA8D03AB1B3928B3090,PinName="then",PinToolTip="\nExec",Direction="EGPD_Output",PinType.PinCategory="exec",PinType.PinSubCategory="",PinType.PinSubCategoryObject=None,PinType.PinSubCategoryMemberReference=(),PinType.PinValueType=(),PinType.bIsMap=False,PinType.bIsSet=False,PinType.bIsArray=False,PinType.bIsReference=False,PinType.bIsConst=False,PinType.bIsWeakPointer=False,LinkedTo=(K2Node_CallFunction_369 527F9720481C49179DFF49B166FB5074,),PersistentGuid=00000000000000000000000000000000,bHidden=False,bNotConnectable=False,bDefaultValueIsReadOnly=False,bDefaultValueIsIgnored=False,bAdvancedView=False,)
+		   CustomProperties Pin (PinId=55F338D54A4933AB828D2DB26782C387,PinName="self",PinFriendlyName=NSLOCTEXT("K2Node", "Target", "Target"),PinToolTip="Target\nMaterial Instance Dynamic Reference",PinType.PinCategory="object",PinType.PinSubCategory="",PinType.PinSubCategoryObject=Class'/Script/Engine.MaterialInstanceDynamic',PinType.PinSubCategoryMemberReference=(),PinType.PinValueType=(),PinType.bIsMap=False,PinType.bIsSet=False,PinType.bIsArray=False,PinType.bIsReference=False,PinType.bIsConst=False,PinType.bIsWeakPointer=False,LinkedTo=(K2Node_VariableGet_181 16610BCA473F60D2FB851F8B3E4736A4,),PersistentGuid=00000000000000000000000000000000,bHidden=False,bNotConnectable=False,bDefaultValueIsReadOnly=False,bDefaultValueIsIgnored=False,bAdvancedView=False,)
+		   CustomProperties Pin (PinId=7E1E1625448C9EC8718DA083EEE57602,PinName="ParameterName",PinToolTip="Parameter Name\nName",PinType.PinCategory="name",PinType.PinSubCategory="",PinType.PinSubCategoryObject=None,PinType.PinSubCategoryMemberReference=(),PinType.PinValueType=(),PinType.bIsMap=False,PinType.bIsSet=False,PinType.bIsArray=False,PinType.bIsReference=False,PinType.bIsConst=False,PinType.bIsWeakPointer=False,DefaultValue="ForceSize",AutogeneratedDefaultValue="None",PersistentGuid=00000000000000000000000000000000,bHidden=False,bNotConnectable=False,bDefaultValueIsReadOnly=False,bDefaultValueIsIgnored=False,bAdvancedView=False,)
+		   CustomProperties Pin (PinId=2997AC484441B6721690EC97301505A9,PinName="Value",PinToolTip="Value\nFloat",PinType.PinCategory="float",PinType.PinSubCategory="",PinType.PinSubCategoryObject=None,PinType.PinSubCategoryMemberReference=(),PinType.PinValueType=(),PinType.bIsMap=False,PinType.bIsSet=False,PinType.bIsArray=False,PinType.bIsReference=False,PinType.bIsConst=False,PinType.bIsWeakPointer=False,DefaultValue="0.0",AutogeneratedDefaultValue="0.0",LinkedTo=(K2Node_VariableGet_183 4CA0FD884AE58C829F746B926BBF77AE,),PersistentGuid=00000000000000000000000000000000,bHidden=False,bNotConnectable=False,bDefaultValueIsReadOnly=False,bDefaultValueIsIgnored=False,bAdvancedView=False,)
+		End Object
+		Begin Object Class=K2Node_VariableGet Name="K2Node_VariableGet_183"
+		   VariableReference=(MemberName="BrushSize",MemberGuid=1CCE6A824DA81A1B69B0828C36569139,bSelfContext=True)
+		   NodePosX=688
+		   NodePosY=32
+		   NodeGuid=4A8189D8435F5591518712BA98DB2A69
+		   CustomProperties Pin (PinId=4CA0FD884AE58C829F746B926BBF77AE,PinName="BrushSize",Direction="EGPD_Output",PinType.PinCategory="float",PinType.PinSubCategory="",PinType.PinSubCategoryObject=None,PinType.PinSubCategoryMemberReference=(),PinType.PinValueType=(),PinType.bIsMap=False,PinType.bIsSet=False,PinType.bIsArray=False,PinType.bIsReference=False,PinType.bIsConst=False,PinType.bIsWeakPointer=False,AutogeneratedDefaultValue="0.0",LinkedTo=(K2Node_CallFunction_368 2997AC484441B6721690EC97301505A9,),PersistentGuid=00000000000000000000000000000000,bHidden=False,bNotConnectable=False,bDefaultValueIsReadOnly=False,bDefaultValueIsIgnored=False,bAdvancedView=False,)
+		   CustomProperties Pin (PinId=2352A2434257CC49A0B2A9A6A21DEFFF,PinName="self",PinFriendlyName=NSLOCTEXT("K2Node", "Target", "Target"),PinType.PinCategory="object",PinType.PinSubCategory="",PinType.PinSubCategoryObject=BlueprintGeneratedClass'/Game/Docs/Perception/NewBlueprint.NewBlueprint_C',PinType.PinSubCategoryMemberReference=(),PinType.PinValueType=(),PinType.bIsMap=False,PinType.bIsSet=False,PinType.bIsArray=False,PinType.bIsReference=False,PinType.bIsConst=False,PinType.bIsWeakPointer=False,PersistentGuid=00000000000000000000000000000000,bHidden=True,bNotConnectable=False,bDefaultValueIsReadOnly=False,bDefaultValueIsIgnored=False,bAdvancedView=False,)
+		End Object
+		Begin Object Class=K2Node_CallFunction Name="K2Node_CallFunction_369"
+		   FunctionReference=(MemberParent=Class'/Script/Engine.MaterialInstanceDynamic',MemberName="SetScalarParameterValue")
+		   NodePosX=1280
+		   NodePosY=-144
+		   NodeGuid=90D4ECCA4E7726211606E38CFE93A3A7
+		   CustomProperties Pin (PinId=527F9720481C49179DFF49B166FB5074,PinName="execute",PinToolTip="\nExec",PinType.PinCategory="exec",PinType.PinSubCategory="",PinType.PinSubCategoryObject=None,PinType.PinSubCategoryMemberReference=(),PinType.PinValueType=(),PinType.bIsMap=False,PinType.bIsSet=False,PinType.bIsArray=False,PinType.bIsReference=False,PinType.bIsConst=False,PinType.bIsWeakPointer=False,LinkedTo=(K2Node_CallFunction_368 5826E451454A5FA8D03AB1B3928B3090,),PersistentGuid=00000000000000000000000000000000,bHidden=False,bNotConnectable=False,bDefaultValueIsReadOnly=False,bDefaultValueIsIgnored=False,bAdvancedView=False,)
+		   CustomProperties Pin (PinId=5826E451454A5FA8D03AB1B3928B3090,PinName="then",PinToolTip="\nExec",Direction="EGPD_Output",PinType.PinCategory="exec",PinType.PinSubCategory="",PinType.PinSubCategoryObject=None,PinType.PinSubCategoryMemberReference=(),PinType.PinValueType=(),PinType.bIsMap=False,PinType.bIsSet=False,PinType.bIsArray=False,PinType.bIsReference=False,PinType.bIsConst=False,PinType.bIsWeakPointer=False,LinkedTo=(K2Node_CallFunction_367 88C40CCB4DEBDA28F201A388C6DF16B9,),PersistentGuid=00000000000000000000000000000000,bHidden=False,bNotConnectable=False,bDefaultValueIsReadOnly=False,bDefaultValueIsIgnored=False,bAdvancedView=False,)
+		   CustomProperties Pin (PinId=55F338D54A4933AB828D2DB26782C387,PinName="self",PinFriendlyName=NSLOCTEXT("K2Node", "Target", "Target"),PinToolTip="Target\nMaterial Instance Dynamic Reference",PinType.PinCategory="object",PinType.PinSubCategory="",PinType.PinSubCategoryObject=Class'/Script/Engine.MaterialInstanceDynamic',PinType.PinSubCategoryMemberReference=(),PinType.PinValueType=(),PinType.bIsMap=False,PinType.bIsSet=False,PinType.bIsArray=False,PinType.bIsReference=False,PinType.bIsConst=False,PinType.bIsWeakPointer=False,LinkedTo=(K2Node_Knot_2 E89B5FC6471F58488288299BF897E596,),PersistentGuid=00000000000000000000000000000000,bHidden=False,bNotConnectable=False,bDefaultValueIsReadOnly=False,bDefaultValueIsIgnored=False,bAdvancedView=False,)
+		   CustomProperties Pin (PinId=7E1E1625448C9EC8718DA083EEE57602,PinName="ParameterName",PinToolTip="Parameter Name\nName",PinType.PinCategory="name",PinType.PinSubCategory="",PinType.PinSubCategoryObject=None,PinType.PinSubCategoryMemberReference=(),PinType.PinValueType=(),PinType.bIsMap=False,PinType.bIsSet=False,PinType.bIsArray=False,PinType.bIsReference=False,PinType.bIsConst=False,PinType.bIsWeakPointer=False,DefaultValue="ForceStrength",AutogeneratedDefaultValue="None",PersistentGuid=00000000000000000000000000000000,bHidden=False,bNotConnectable=False,bDefaultValueIsReadOnly=False,bDefaultValueIsIgnored=False,bAdvancedView=False,)
+		   CustomProperties Pin (PinId=2997AC484441B6721690EC97301505A9,PinName="Value",PinToolTip="Value\nFloat",PinType.PinCategory="float",PinType.PinSubCategory="",PinType.PinSubCategoryObject=None,PinType.PinSubCategoryMemberReference=(),PinType.PinValueType=(),PinType.bIsMap=False,PinType.bIsSet=False,PinType.bIsArray=False,PinType.bIsReference=False,PinType.bIsConst=False,PinType.bIsWeakPointer=False,DefaultValue="0.0",AutogeneratedDefaultValue="0.0",LinkedTo=(K2Node_VariableGet_184 C94B43A14D6EFD6EC40E638FADF9DEDA,),PersistentGuid=00000000000000000000000000000000,bHidden=False,bNotConnectable=False,bDefaultValueIsReadOnly=False,bDefaultValueIsIgnored=False,bAdvancedView=False,)
+		End Object
+		Begin Object Class=K2Node_VariableGet Name="K2Node_VariableGet_184"
+		   VariableReference=(MemberName="BrushStrength",MemberGuid=86BA01A34EF38F4D34549680DD25106A,bSelfContext=True)
+		   NodePosX=1104
+		   NodePosY=64
+		   NodeGuid=4B537FE4436968C31A68B89AC7FA52A7
+		   CustomProperties Pin (PinId=C94B43A14D6EFD6EC40E638FADF9DEDA,PinName="BrushStrength",Direction="EGPD_Output",PinType.PinCategory="float",PinType.PinSubCategory="",PinType.PinSubCategoryObject=None,PinType.PinSubCategoryMemberReference=(),PinType.PinValueType=(),PinType.bIsMap=False,PinType.bIsSet=False,PinType.bIsArray=False,PinType.bIsReference=False,PinType.bIsConst=False,PinType.bIsWeakPointer=False,AutogeneratedDefaultValue="0.0",LinkedTo=(K2Node_CallFunction_369 2997AC484441B6721690EC97301505A9,),PersistentGuid=00000000000000000000000000000000,bHidden=False,bNotConnectable=False,bDefaultValueIsReadOnly=False,bDefaultValueIsIgnored=False,bAdvancedView=False,)
+		   CustomProperties Pin (PinId=A4C769414DF088ACA7CD4AA013C615BA,PinName="self",PinFriendlyName=NSLOCTEXT("K2Node", "Target", "Target"),PinType.PinCategory="object",PinType.PinSubCategory="",PinType.PinSubCategoryObject=BlueprintGeneratedClass'/Game/Docs/Perception/NewBlueprint.NewBlueprint_C',PinType.PinSubCategoryMemberReference=(),PinType.PinValueType=(),PinType.bIsMap=False,PinType.bIsSet=False,PinType.bIsArray=False,PinType.bIsReference=False,PinType.bIsConst=False,PinType.bIsWeakPointer=False,PersistentGuid=00000000000000000000000000000000,bHidden=True,bNotConnectable=False,bDefaultValueIsReadOnly=False,bDefaultValueIsIgnored=False,bAdvancedView=False,)
+		End Object
+		Begin Object Class=K2Node_CallFunction Name="K2Node_CallFunction_370"
+		   FunctionReference=(MemberParent=Class'/Script/Engine.MaterialInstanceDynamic',MemberName="SetTextureParameterValue")
+		   NodePosX=-16
+		   NodePosY=-176
+		   ErrorType=1
+		   ErrorMsg="Error This blueprint (self) is not a MaterialInstanceDynamic, therefore \' self \' must have a connection."
+		   NodeGuid=9E6121CF4A1911A5DF55648A96429C55
+		   CustomProperties Pin (PinId=542CC73148D38A37BCCECE90262AFC49,PinName="execute",PinType.PinCategory="exec",PinType.PinSubCategory="",PinType.PinSubCategoryObject=None,PinType.PinSubCategoryMemberReference=(),PinType.PinValueType=(),PinType.bIsMap=False,PinType.bIsSet=False,PinType.bIsArray=False,PinType.bIsReference=False,PinType.bIsConst=False,PinType.bIsWeakPointer=False,LinkedTo=(K2Node_CallFunction_359 43C94C3C4FB04A4DEADC37A9E3BDE0B2,),PersistentGuid=00000000000000000000000000000000,bHidden=False,bNotConnectable=False,bDefaultValueIsReadOnly=False,bDefaultValueIsIgnored=False,bAdvancedView=False,)
+		   CustomProperties Pin (PinId=A183C34442666D0C8D878F8716DDAC38,PinName="then",Direction="EGPD_Output",PinType.PinCategory="exec",PinType.PinSubCategory="",PinType.PinSubCategoryObject=None,PinType.PinSubCategoryMemberReference=(),PinType.PinValueType=(),PinType.bIsMap=False,PinType.bIsSet=False,PinType.bIsArray=False,PinType.bIsReference=False,PinType.bIsConst=False,PinType.bIsWeakPointer=False,PersistentGuid=00000000000000000000000000000000,bHidden=False,bNotConnectable=False,bDefaultValueIsReadOnly=False,bDefaultValueIsIgnored=False,bAdvancedView=False,)
+		   CustomProperties Pin (PinId=63A3B28B4B0AAB8799485AA1BF3EAE97,PinName="self",PinFriendlyName=NSLOCTEXT("K2Node", "Target", "Target"),PinType.PinCategory="object",PinType.PinSubCategory="",PinType.PinSubCategoryObject=Class'/Script/Engine.MaterialInstanceDynamic',PinType.PinSubCategoryMemberReference=(),PinType.PinValueType=(),PinType.bIsMap=False,PinType.bIsSet=False,PinType.bIsArray=False,PinType.bIsReference=False,PinType.bIsConst=False,PinType.bIsWeakPointer=False,LinkedTo=(K2Node_VariableGet_186 D6A13F044975C578FFF6A98328676BC9,),PersistentGuid=00000000000000000000000000000000,bHidden=False,bNotConnectable=False,bDefaultValueIsReadOnly=False,bDefaultValueIsIgnored=False,bAdvancedView=False,)
+		   CustomProperties Pin (PinId=92D5DD9C482CFAAA46C3FE93F76AC002,PinName="ParameterName",PinType.PinCategory="name",PinType.PinSubCategory="",PinType.PinSubCategoryObject=None,PinType.PinSubCategoryMemberReference=(),PinType.PinValueType=(),PinType.bIsMap=False,PinType.bIsSet=False,PinType.bIsArray=False,PinType.bIsReference=False,PinType.bIsConst=False,PinType.bIsWeakPointer=False,DefaultValue="Heightfield",AutogeneratedDefaultValue="None",PersistentGuid=00000000000000000000000000000000,bHidden=False,bNotConnectable=False,bDefaultValueIsReadOnly=False,bDefaultValueIsIgnored=False,bAdvancedView=False,)
+		   CustomProperties Pin (PinId=FAC8C8B64E1E50423C9F4493F3C27926,PinName="Value",PinType.PinCategory="object",PinType.PinSubCategory="",PinType.PinSubCategoryObject=Class'/Script/Engine.Texture',PinType.PinSubCategoryMemberReference=(),PinType.PinValueType=(),PinType.bIsMap=False,PinType.bIsSet=False,PinType.bIsArray=False,PinType.bIsReference=False,PinType.bIsConst=False,PinType.bIsWeakPointer=False,LinkedTo=(K2Node_VariableGet_185 3214E57245F463BA1D70C1845534BA2C,),PersistentGuid=00000000000000000000000000000000,bHidden=False,bNotConnectable=False,bDefaultValueIsReadOnly=False,bDefaultValueIsIgnored=False,bAdvancedView=False,)
+		End Object
+		Begin Object Class=K2Node_VariableGet Name="K2Node_VariableGet_185"
+		   VariableReference=(MemberName="HeightfieldRT",MemberGuid=E71A12524941DA802B08039BF7491FF7,bSelfContext=True)
+		   NodePosX=-176
+		   NodePosY=80
+		   NodeGuid=76226BD54C8F569560103783C4CD1A86
+		   CustomProperties Pin (PinId=3214E57245F463BA1D70C1845534BA2C,PinName="HeightfieldRT",Direction="EGPD_Output",PinType.PinCategory="object",PinType.PinSubCategory="",PinType.PinSubCategoryObject=Class'/Script/Engine.TextureRenderTarget2D',PinType.PinSubCategoryMemberReference=(),PinType.PinValueType=(),PinType.bIsMap=False,PinType.bIsSet=False,PinType.bIsArray=False,PinType.bIsReference=False,PinType.bIsConst=False,PinType.bIsWeakPointer=False,LinkedTo=(K2Node_CallFunction_370 FAC8C8B64E1E50423C9F4493F3C27926,),PersistentGuid=00000000000000000000000000000000,bHidden=False,bNotConnectable=False,bDefaultValueIsReadOnly=False,bDefaultValueIsIgnored=False,bAdvancedView=False,)
+		   CustomProperties Pin (PinId=6487857E4D70A07AF8065183735A48D7,PinName="self",PinFriendlyName=NSLOCTEXT("K2Node", "Target", "Target"),PinType.PinCategory="object",PinType.PinSubCategory="",PinType.PinSubCategoryObject=BlueprintGeneratedClass'/Game/Docs/Perception/NewBlueprint.NewBlueprint_C',PinType.PinSubCategoryMemberReference=(),PinType.PinValueType=(),PinType.bIsMap=False,PinType.bIsSet=False,PinType.bIsArray=False,PinType.bIsReference=False,PinType.bIsConst=False,PinType.bIsWeakPointer=False,PersistentGuid=00000000000000000000000000000000,bHidden=True,bNotConnectable=False,bDefaultValueIsReadOnly=False,bDefaultValueIsIgnored=False,bAdvancedView=False,)
+		End Object
+		Begin Object Class=K2Node_VariableGet Name="K2Node_VariableGet_186"
+		   VariableReference=(MemberName="PainterMaterialInstance",MemberGuid=C12C0CC9438476601416309CDF044A50,bSelfContext=True)
+		   NodePosX=-256
+		   NodePosY=32
+		   NodeGuid=649ED586453820E68250F4946D9E3F74
+		   CustomProperties Pin (PinId=D6A13F044975C578FFF6A98328676BC9,PinName="PainterMaterialInstance",Direction="EGPD_Output",PinType.PinCategory="object",PinType.PinSubCategory="",PinType.PinSubCategoryObject=Class'/Script/Engine.MaterialInstanceDynamic',PinType.PinSubCategoryMemberReference=(),PinType.PinValueType=(),PinType.bIsMap=False,PinType.bIsSet=False,PinType.bIsArray=False,PinType.bIsReference=False,PinType.bIsConst=False,PinType.bIsWeakPointer=False,LinkedTo=(K2Node_CallFunction_370 63A3B28B4B0AAB8799485AA1BF3EAE97,),PersistentGuid=00000000000000000000000000000000,bHidden=False,bNotConnectable=False,bDefaultValueIsReadOnly=False,bDefaultValueIsIgnored=False,bAdvancedView=False,)
+		   CustomProperties Pin (PinId=F6D062374D80E728415C4BBD46D27021,PinName="self",PinFriendlyName=NSLOCTEXT("K2Node", "Target", "Target"),PinType.PinCategory="object",PinType.PinSubCategory="",PinType.PinSubCategoryObject=BlueprintGeneratedClass'/Game/Docs/Perception/NewBlueprint.NewBlueprint_C',PinType.PinSubCategoryMemberReference=(),PinType.PinValueType=(),PinType.bIsMap=False,PinType.bIsSet=False,PinType.bIsArray=False,PinType.bIsReference=False,PinType.bIsConst=False,PinType.bIsWeakPointer=False,PersistentGuid=00000000000000000000000000000000,bHidden=True,bNotConnectable=False,bDefaultValueIsReadOnly=False,bDefaultValueIsIgnored=False,bAdvancedView=False,)
+		End Object
+		Begin Object Class=K2Node_CallFunction Name="K2Node_CallFunction_371"
+		   FunctionReference=(MemberName="EnableInput",bSelfContext=True)
+		   NodePosX=-800
+		   NodePosY=-176
+		   NodeGuid=ABCDF3454641D37E8EEE988D5C219DC4
+		   CustomProperties Pin (PinId=59EBD0D84FA02620751F0CAC38ADA5FF,PinName="execute",PinType.PinCategory="exec",PinType.PinSubCategory="",PinType.PinSubCategoryObject=None,PinType.PinSubCategoryMemberReference=(),PinType.PinValueType=(),PinType.bIsMap=False,PinType.bIsSet=False,PinType.bIsArray=False,PinType.bIsReference=False,PinType.bIsConst=False,PinType.bIsWeakPointer=False,LinkedTo=(K2Node_Event_0 4C1A3CD54C93D4A74803118ED8E9BE79,),PersistentGuid=00000000000000000000000000000000,bHidden=False,bNotConnectable=False,bDefaultValueIsReadOnly=False,bDefaultValueIsIgnored=False,bAdvancedView=False,)
+		   CustomProperties Pin (PinId=C8543EC14C417EB94C0A40806572C540,PinName="then",Direction="EGPD_Output",PinType.PinCategory="exec",PinType.PinSubCategory="",PinType.PinSubCategoryObject=None,PinType.PinSubCategoryMemberReference=(),PinType.PinValueType=(),PinType.bIsMap=False,PinType.bIsSet=False,PinType.bIsArray=False,PinType.bIsReference=False,PinType.bIsConst=False,PinType.bIsWeakPointer=False,LinkedTo=(K2Node_CallFunction_359 50F50A054429DFDB03CCE884B76C4E0A,),PersistentGuid=00000000000000000000000000000000,bHidden=False,bNotConnectable=False,bDefaultValueIsReadOnly=False,bDefaultValueIsIgnored=False,bAdvancedView=False,)
+		   CustomProperties Pin (PinId=2472C14544442ACFC3F3CAAE2BFC249B,PinName="self",PinFriendlyName=NSLOCTEXT("K2Node", "Target", "Target"),PinType.PinCategory="object",PinType.PinSubCategory="",PinType.PinSubCategoryObject=Class'/Script/Engine.Actor',PinType.PinSubCategoryMemberReference=(),PinType.PinValueType=(),PinType.bIsMap=False,PinType.bIsSet=False,PinType.bIsArray=False,PinType.bIsReference=False,PinType.bIsConst=False,PinType.bIsWeakPointer=False,PersistentGuid=00000000000000000000000000000000,bHidden=False,bNotConnectable=False,bDefaultValueIsReadOnly=False,bDefaultValueIsIgnored=False,bAdvancedView=False,)
+		   CustomProperties Pin (PinId=538663B143E85B7459172FBD475EDB53,PinName="PlayerController",PinType.PinCategory="object",PinType.PinSubCategory="",PinType.PinSubCategoryObject=Class'/Script/Engine.PlayerController',PinType.PinSubCategoryMemberReference=(),PinType.PinValueType=(),PinType.bIsMap=False,PinType.bIsSet=False,PinType.bIsArray=False,PinType.bIsReference=False,PinType.bIsConst=False,PinType.bIsWeakPointer=False,LinkedTo=(K2Node_CallFunction_372 DB581A434D6C7153C9DF6F9236D045F9,),PersistentGuid=00000000000000000000000000000000,bHidden=False,bNotConnectable=False,bDefaultValueIsReadOnly=False,bDefaultValueIsIgnored=False,bAdvancedView=False,)
+		End Object
+		Begin Object Class=K2Node_CallFunction Name="K2Node_CallFunction_372"
+		   bIsPureFunc=True
+		   FunctionReference=(MemberParent=Class'/Script/Engine.GameplayStatics',MemberName="GetPlayerController")
+		   NodePosX=-1088
+		   NodePosY=-16
+		   NodeGuid=1F58152C440A8F20BDCCC791275D56B7
+		   CustomProperties Pin (PinId=2843ED7743C312C47BFD3AB629A45D70,PinName="self",PinFriendlyName=NSLOCTEXT("K2Node", "Target", "Target"),PinType.PinCategory="object",PinType.PinSubCategory="",PinType.PinSubCategoryObject=Class'/Script/Engine.GameplayStatics',PinType.PinSubCategoryMemberReference=(),PinType.PinValueType=(),PinType.bIsMap=False,PinType.bIsSet=False,PinType.bIsArray=False,PinType.bIsReference=False,PinType.bIsConst=False,PinType.bIsWeakPointer=False,DefaultObject="/Script/Engine.Default__GameplayStatics",PersistentGuid=00000000000000000000000000000000,bHidden=True,bNotConnectable=False,bDefaultValueIsReadOnly=False,bDefaultValueIsIgnored=False,bAdvancedView=False,)
+		   CustomProperties Pin (PinId=5BA2917945614FA7C1CCA2ACD9C09317,PinName="WorldContextObject",PinType.PinCategory="object",PinType.PinSubCategory="",PinType.PinSubCategoryObject=Class'/Script/CoreUObject.Object',PinType.PinSubCategoryMemberReference=(),PinType.PinValueType=(),PinType.bIsMap=False,PinType.bIsSet=False,PinType.bIsArray=False,PinType.bIsReference=False,PinType.bIsConst=True,PinType.bIsWeakPointer=False,PersistentGuid=00000000000000000000000000000000,bHidden=True,bNotConnectable=False,bDefaultValueIsReadOnly=False,bDefaultValueIsIgnored=False,bAdvancedView=False,)
+		   CustomProperties Pin (PinId=9D63B5834AA0892C28999DAE7EAD90DD,PinName="PlayerIndex",PinType.PinCategory="int",PinType.PinSubCategory="",PinType.PinSubCategoryObject=None,PinType.PinSubCategoryMemberReference=(),PinType.PinValueType=(),PinType.bIsMap=False,PinType.bIsSet=False,PinType.bIsArray=False,PinType.bIsReference=False,PinType.bIsConst=False,PinType.bIsWeakPointer=False,DefaultValue="0",AutogeneratedDefaultValue="0",PersistentGuid=00000000000000000000000000000000,bHidden=False,bNotConnectable=False,bDefaultValueIsReadOnly=False,bDefaultValueIsIgnored=False,bAdvancedView=False,)
+		   CustomProperties Pin (PinId=DB581A434D6C7153C9DF6F9236D045F9,PinName="ReturnValue",Direction="EGPD_Output",PinType.PinCategory="object",PinType.PinSubCategory="",PinType.PinSubCategoryObject=Class'/Script/Engine.PlayerController',PinType.PinSubCategoryMemberReference=(),PinType.PinValueType=(),PinType.bIsMap=False,PinType.bIsSet=False,PinType.bIsArray=False,PinType.bIsReference=False,PinType.bIsConst=False,PinType.bIsWeakPointer=False,LinkedTo=(K2Node_CallFunction_371 538663B143E85B7459172FBD475EDB53,),PersistentGuid=00000000000000000000000000000000,bHidden=False,bNotConnectable=False,bDefaultValueIsReadOnly=False,bDefaultValueIsIgnored=False,bAdvancedView=False,)
+		End Object
+		Begin Object Class=EdGraphNode_Comment Name="EdGraphNode_Comment_122"
+		   NodePosX=-1136
+		   NodePosY=-224
+		   NodeWidth=544
+		   NodeHeight=320
+		   NodeComment="Enable input on this Actor"
+		   NodeGuid=4F633C864F16B20671F9EBAA54783ABF
+		End Object
+		Begin Object Class=EdGraphNode_Comment Name="EdGraphNode_Comment_123"
+		   NodePosX=-1152
+		   NodePosY=224
+		   NodeWidth=576
+		   NodeHeight=288
+		   NodeComment="Check to see if Left Mouse Button is held down"
+		   NodeGuid=8AECFF5C40918BF6245F4DAFABBFC704
+		End Object
+		Begin Object Class=EdGraphNode_Comment Name="EdGraphNode_Comment_124"
+		   NodePosX=-560
+		   NodePosY=-224
+		   NodeWidth=852
+		   NodeHeight=390
+		   NodeComment="Clear the Render Target off all previous data"
+		   NodeGuid=0C95B6804EB892426DD63AA2CB9C5498
+		End Object
+		Begin Object Class=K2Node_CallFunction Name="K2Node_CallFunction_111"
+		   bIsPureFunc=True
+		   FunctionReference=(MemberParent=Class'/Script/Engine.KismetMathLibrary',MemberName="MakeColor")
+		   NodePosX=1200
+		   NodePosY=176
+		   NodeGuid=2F092D6D4DA90C9572EC9F85924916F3
+		   CustomProperties Pin (PinId=EAE432A9456760D2FD7D9AA26954601A,PinName="self",PinFriendlyName=NSLOCTEXT("K2Node", "Target", "Target"),PinToolTip="Target\nKismet Math Library Reference",PinType.PinCategory="object",PinType.PinSubCategory="",PinType.PinSubCategoryObject=Class'/Script/Engine.KismetMathLibrary',PinType.PinSubCategoryMemberReference=(),PinType.PinValueType=(),PinType.bIsMap=False,PinType.bIsSet=False,PinType.bIsArray=False,PinType.bIsReference=False,PinType.bIsConst=False,PinType.bIsWeakPointer=False,DefaultObject="/Script/Engine.Default__KismetMathLibrary",PersistentGuid=00000000000000000000000000000000,bHidden=True,bNotConnectable=False,bDefaultValueIsReadOnly=False,bDefaultValueIsIgnored=False,bAdvancedView=False,)
+		   CustomProperties Pin (PinId=CE2135D24885DB9D640A02B406C717A5,PinName="R",PinToolTip="R\nFloat",PinType.PinCategory="float",PinType.PinSubCategory="",PinType.PinSubCategoryObject=None,PinType.PinSubCategoryMemberReference=(),PinType.PinValueType=(),PinType.bIsMap=False,PinType.bIsSet=False,PinType.bIsArray=False,PinType.bIsReference=False,PinType.bIsConst=False,PinType.bIsWeakPointer=False,DefaultValue="0.0",AutogeneratedDefaultValue="0.0",LinkedTo=(K2Node_CallFunction_112 6235D41242F07A065692E88F8CE30E41,),PersistentGuid=00000000000000000000000000000000,bHidden=False,bNotConnectable=False,bDefaultValueIsReadOnly=False,bDefaultValueIsIgnored=False,bAdvancedView=False,)
+		   CustomProperties Pin (PinId=7E9481F14863C9C0ED9A6BA40DAFBC2B,PinName="G",PinToolTip="G\nFloat",PinType.PinCategory="float",PinType.PinSubCategory="",PinType.PinSubCategoryObject=None,PinType.PinSubCategoryMemberReference=(),PinType.PinValueType=(),PinType.bIsMap=False,PinType.bIsSet=False,PinType.bIsArray=False,PinType.bIsReference=False,PinType.bIsConst=False,PinType.bIsWeakPointer=False,DefaultValue="0.0",AutogeneratedDefaultValue="0.0",LinkedTo=(K2Node_CallFunction_112 AF8593D64D9E3B65F6F54CACEF7E720E,),PersistentGuid=00000000000000000000000000000000,bHidden=False,bNotConnectable=False,bDefaultValueIsReadOnly=False,bDefaultValueIsIgnored=False,bAdvancedView=False,)
+		   CustomProperties Pin (PinId=AB401BC0477E4D994187F6AAA68F2190,PinName="B",PinToolTip="B\nFloat",PinType.PinCategory="float",PinType.PinSubCategory="",PinType.PinSubCategoryObject=None,PinType.PinSubCategoryMemberReference=(),PinType.PinValueType=(),PinType.bIsMap=False,PinType.bIsSet=False,PinType.bIsArray=False,PinType.bIsReference=False,PinType.bIsConst=False,PinType.bIsWeakPointer=False,DefaultValue="0.0",AutogeneratedDefaultValue="0.0",PersistentGuid=00000000000000000000000000000000,bHidden=False,bNotConnectable=False,bDefaultValueIsReadOnly=False,bDefaultValueIsIgnored=False,bAdvancedView=False,)
+		   CustomProperties Pin (PinId=7D2ABAFE441C9565DDC1C59BD202C62E,PinName="A",PinToolTip="A\nFloat",PinType.PinCategory="float",PinType.PinSubCategory="",PinType.PinSubCategoryObject=None,PinType.PinSubCategoryMemberReference=(),PinType.PinValueType=(),PinType.bIsMap=False,PinType.bIsSet=False,PinType.bIsArray=False,PinType.bIsReference=False,PinType.bIsConst=False,PinType.bIsWeakPointer=False,DefaultValue="1.000000",AutogeneratedDefaultValue="1.000000",PersistentGuid=00000000000000000000000000000000,bHidden=False,bNotConnectable=False,bDefaultValueIsReadOnly=False,bDefaultValueIsIgnored=False,bAdvancedView=False,)
+		   CustomProperties Pin (PinId=ED6E4C604B7A3F2181D74F98664FDCDE,PinName="ReturnValue",PinToolTip="Return Value\nLinear Color Structure",Direction="EGPD_Output",PinType.PinCategory="struct",PinType.PinSubCategory="",PinType.PinSubCategoryObject=ScriptStruct'/Script/CoreUObject.LinearColor',PinType.PinSubCategoryMemberReference=(),PinType.PinValueType=(),PinType.bIsMap=False,PinType.bIsSet=False,PinType.bIsArray=False,PinType.bIsReference=False,PinType.bIsConst=False,PinType.bIsWeakPointer=False,LinkedTo=(K2Node_CallFunction_367 1523A2CA42868EC87B90F39450925710,),PersistentGuid=00000000000000000000000000000000,bHidden=False,bNotConnectable=False,bDefaultValueIsReadOnly=False,bDefaultValueIsIgnored=False,bAdvancedView=False,)
+		End Object
+		Begin Object Class=K2Node_CallFunction Name="K2Node_CallFunction_112"
+		   bIsPureFunc=True
+		   FunctionReference=(MemberParent=Class'/Script/Engine.KismetMathLibrary',MemberName="BreakVector2D")
+		   NodePosX=1008
+		   NodePosY=176
+		   NodeGuid=E4229A8349D48CC06A748CB7F815B373
+		   CustomProperties Pin (PinId=02A0F5B349B6D8AF9F74C0ADCD2BCC71,PinName="self",PinFriendlyName=NSLOCTEXT("K2Node", "Target", "Target"),PinToolTip="Target\nKismet Math Library Reference",PinType.PinCategory="object",PinType.PinSubCategory="",PinType.PinSubCategoryObject=Class'/Script/Engine.KismetMathLibrary',PinType.PinSubCategoryMemberReference=(),PinType.PinValueType=(),PinType.bIsMap=False,PinType.bIsSet=False,PinType.bIsArray=False,PinType.bIsReference=False,PinType.bIsConst=False,PinType.bIsWeakPointer=False,DefaultObject="/Script/Engine.Default__KismetMathLibrary",PersistentGuid=00000000000000000000000000000000,bHidden=True,bNotConnectable=False,bDefaultValueIsReadOnly=False,bDefaultValueIsIgnored=False,bAdvancedView=False,)
+		   CustomProperties Pin (PinId=9C1C769941A0BCA1C1E925B9C4796984,PinName="InVec",PinToolTip="In Vec\nVector 2D Structure",PinType.PinCategory="struct",PinType.PinSubCategory="",PinType.PinSubCategoryObject=ScriptStruct'/Script/CoreUObject.Vector2D',PinType.PinSubCategoryMemberReference=(),PinType.PinValueType=(),PinType.bIsMap=False,PinType.bIsSet=False,PinType.bIsArray=False,PinType.bIsReference=False,PinType.bIsConst=False,PinType.bIsWeakPointer=False,LinkedTo=(K2Node_CustomEvent_0 28D508444B6B999E2B762C9FBF985375,),PersistentGuid=00000000000000000000000000000000,bHidden=False,bNotConnectable=False,bDefaultValueIsReadOnly=False,bDefaultValueIsIgnored=False,bAdvancedView=False,)
+		   CustomProperties Pin (PinId=6235D41242F07A065692E88F8CE30E41,PinName="X",PinToolTip="X\nFloat",Direction="EGPD_Output",PinType.PinCategory="float",PinType.PinSubCategory="",PinType.PinSubCategoryObject=None,PinType.PinSubCategoryMemberReference=(),PinType.PinValueType=(),PinType.bIsMap=False,PinType.bIsSet=False,PinType.bIsArray=False,PinType.bIsReference=False,PinType.bIsConst=False,PinType.bIsWeakPointer=False,AutogeneratedDefaultValue="0.0",LinkedTo=(K2Node_CallFunction_111 CE2135D24885DB9D640A02B406C717A5,),PersistentGuid=00000000000000000000000000000000,bHidden=False,bNotConnectable=False,bDefaultValueIsReadOnly=False,bDefaultValueIsIgnored=False,bAdvancedView=False,)
+		   CustomProperties Pin (PinId=AF8593D64D9E3B65F6F54CACEF7E720E,PinName="Y",PinToolTip="Y\nFloat",Direction="EGPD_Output",PinType.PinCategory="float",PinType.PinSubCategory="",PinType.PinSubCategoryObject=None,PinType.PinSubCategoryMemberReference=(),PinType.PinValueType=(),PinType.bIsMap=False,PinType.bIsSet=False,PinType.bIsArray=False,PinType.bIsReference=False,PinType.bIsConst=False,PinType.bIsWeakPointer=False,AutogeneratedDefaultValue="0.0",LinkedTo=(K2Node_CallFunction_111 7E9481F14863C9C0ED9A6BA40DAFBC2B,),PersistentGuid=00000000000000000000000000000000,bHidden=False,bNotConnectable=False,bDefaultValueIsReadOnly=False,bDefaultValueIsIgnored=False,bAdvancedView=False,)
+		End Object
+		Begin Object Class=K2Node_Knot Name="K2Node_Knot_1"
+		   NodePosX=1920
+		   NodePosY=128
+		   NodeGuid=D95F94C24F89B388724FB08B4C523663
+		   CustomProperties Pin (PinId=5BC713A7464C13624C865AA695CAEC51,PinName="InputPin",PinType.PinCategory="object",PinType.PinSubCategory="",PinType.PinSubCategoryObject=Class'/Script/Engine.MaterialInstanceDynamic',PinType.PinSubCategoryMemberReference=(),PinType.PinValueType=(),PinType.bIsMap=False,PinType.bIsSet=False,PinType.bIsArray=False,PinType.bIsReference=False,PinType.bIsConst=False,PinType.bIsWeakPointer=False,LinkedTo=(K2Node_Knot_3 650149D74494B25528C925B6D9B777EB,),PersistentGuid=00000000000000000000000000000000,bHidden=False,bNotConnectable=False,bDefaultValueIsReadOnly=False,bDefaultValueIsIgnored=True,bAdvancedView=False,)
+		   CustomProperties Pin (PinId=9617DBFB4F732CAC88F1169322A29AC2,PinName="OutputPin",Direction="EGPD_Output",PinType.PinCategory="object",PinType.PinSubCategory="",PinType.PinSubCategoryObject=Class'/Script/Engine.MaterialInstanceDynamic',PinType.PinSubCategoryMemberReference=(),PinType.PinValueType=(),PinType.bIsMap=False,PinType.bIsSet=False,PinType.bIsArray=False,PinType.bIsReference=False,PinType.bIsConst=False,PinType.bIsWeakPointer=False,LinkedTo=(K2Node_CallFunction_366 D99B80D64C0171739FBC10A4E4C540EC,),PersistentGuid=00000000000000000000000000000000,bHidden=False,bNotConnectable=False,bDefaultValueIsReadOnly=False,bDefaultValueIsIgnored=False,bAdvancedView=False,)
+		End Object
+		Begin Object Class=K2Node_Knot Name="K2Node_Knot_2"
+		   NodePosX=944
+		   NodePosY=128
+		   NodeGuid=83D718C1494194B87822DDA4AE701FF2
+		   CustomProperties Pin (PinId=21486A174CB3882A9431A9803149133A,PinName="InputPin",PinType.PinCategory="object",PinType.PinSubCategory="",PinType.PinSubCategoryObject=Class'/Script/Engine.MaterialInstanceDynamic',PinType.PinSubCategoryMemberReference=(),PinType.PinValueType=(),PinType.bIsMap=False,PinType.bIsSet=False,PinType.bIsArray=False,PinType.bIsReference=False,PinType.bIsConst=False,PinType.bIsWeakPointer=False,LinkedTo=(K2Node_VariableGet_181 16610BCA473F60D2FB851F8B3E4736A4,),PersistentGuid=00000000000000000000000000000000,bHidden=False,bNotConnectable=False,bDefaultValueIsReadOnly=False,bDefaultValueIsIgnored=True,bAdvancedView=False,)
+		   CustomProperties Pin (PinId=E89B5FC6471F58488288299BF897E596,PinName="OutputPin",Direction="EGPD_Output",PinType.PinCategory="object",PinType.PinSubCategory="",PinType.PinSubCategoryObject=Class'/Script/Engine.MaterialInstanceDynamic',PinType.PinSubCategoryMemberReference=(),PinType.PinValueType=(),PinType.bIsMap=False,PinType.bIsSet=False,PinType.bIsArray=False,PinType.bIsReference=False,PinType.bIsConst=False,PinType.bIsWeakPointer=False,LinkedTo=(K2Node_CallFunction_369 55F338D54A4933AB828D2DB26782C387,K2Node_Knot_3 DC5D6ACC48AD171865003F8B26BD4BEC,),PersistentGuid=00000000000000000000000000000000,bHidden=False,bNotConnectable=False,bDefaultValueIsReadOnly=False,bDefaultValueIsIgnored=False,bAdvancedView=False,)
+		End Object
+		Begin Object Class=K2Node_Knot Name="K2Node_Knot_3"
+		   NodePosX=1552
+		   NodePosY=128
+		   NodeGuid=123FD8834420E0D9908374AEA6F4D6FF
+		   CustomProperties Pin (PinId=DC5D6ACC48AD171865003F8B26BD4BEC,PinName="InputPin",PinType.PinCategory="object",PinType.PinSubCategory="",PinType.PinSubCategoryObject=Class'/Script/Engine.MaterialInstanceDynamic',PinType.PinSubCategoryMemberReference=(),PinType.PinValueType=(),PinType.bIsMap=False,PinType.bIsSet=False,PinType.bIsArray=False,PinType.bIsReference=False,PinType.bIsConst=False,PinType.bIsWeakPointer=False,LinkedTo=(K2Node_Knot_2 E89B5FC6471F58488288299BF897E596,),PersistentGuid=00000000000000000000000000000000,bHidden=False,bNotConnectable=False,bDefaultValueIsReadOnly=False,bDefaultValueIsIgnored=True,bAdvancedView=False,)
+		   CustomProperties Pin (PinId=650149D74494B25528C925B6D9B777EB,PinName="OutputPin",Direction="EGPD_Output",PinType.PinCategory="object",PinType.PinSubCategory="",PinType.PinSubCategoryObject=Class'/Script/Engine.MaterialInstanceDynamic',PinType.PinSubCategoryMemberReference=(),PinType.PinValueType=(),PinType.bIsMap=False,PinType.bIsSet=False,PinType.bIsArray=False,PinType.bIsReference=False,PinType.bIsConst=False,PinType.bIsWeakPointer=False,LinkedTo=(K2Node_Knot_1 5BC713A7464C13624C865AA695CAEC51,K2Node_CallFunction_367 FF3048F146D54F1C4AA4A48C221B5457,),PersistentGuid=00000000000000000000000000000000,bHidden=False,bNotConnectable=False,bDefaultValueIsReadOnly=False,bDefaultValueIsIgnored=False,bAdvancedView=False,)
+		End Object
+					
+	[/PARAM]
+[/OBJECT]
+[/EXCERPT:FunctionUV02]
 -->
\ No newline at end of file
