Copying //UE4/Dev-Documentation to //UE4/Samples-Main (Source: //UE4/Dev-Documentation @ 3372893)

diff --git a/Engine/Documentation/Source/Engine/Rendering/RenderTargets/BlueprintRenderTargets/HowTo/HeightFieldPainter/HeightFieldPainterHowTo.INT.udn b/Engine/Documentation/Source/Engine/Rendering/RenderTargets/BlueprintRenderTargets/HowTo/HeightFieldPainter/HeightFieldPainterHowTo.INT.udn
--- a/Engine/Documentation/Source/Engine/Rendering/RenderTargets/BlueprintRenderTargets/HowTo/HeightFieldPainter/HeightFieldPainterHowTo.INT.udn
+++ b/Engine/Documentation/Source/Engine/Rendering/RenderTargets/BlueprintRenderTargets/HowTo/HeightFieldPainter/HeightFieldPainterHowTo.INT.udn
@@ -1,4 +1,4 @@
-availability:docs
+availability:Public
 title:Creating a Height Field Painter with Blueprints and Render Targets
 description:How-To use Blueprints and Render Targets to create realtime Height Filed Painter.
 crumbs:%ROOT%, Engine, Engine\Rendering\RenderTargets
@@ -13,7 +13,7 @@ Related:Engine/Blueprints/
 Related:Engine/Rendering/Materials/
 Related:Engine/Blueprints/UserGuide/UserConstructionScript/
 tags:Render Targets
-tags:Blueprint
+tags:Blueprints
 checkpoint:HFP
 
 ![](HFP_Header_Image.png)
diff --git a/Engine/Documentation/Source/Engine/Rendering/RenderTargets/BlueprintRenderTargets/Reference/BlueprintRenderTargetsReference.INT.udn b/Engine/Documentation/Source/Engine/Rendering/RenderTargets/BlueprintRenderTargets/Reference/BlueprintRenderTargetsReference.INT.udn
--- a/Engine/Documentation/Source/Engine/Rendering/RenderTargets/BlueprintRenderTargets/Reference/BlueprintRenderTargetsReference.INT.udn
+++ b/Engine/Documentation/Source/Engine/Rendering/RenderTargets/BlueprintRenderTargets/Reference/BlueprintRenderTargetsReference.INT.udn
@@ -1,4 +1,4 @@
-availability:docs
+availability:Public
 title: Blueprint and Render Targets Reference
 description: A listing of properties and descriptions for using Blueprint & Render Targets.
 crumbs:%ROOT%, Engine, Engine/Rendering, Engine/Rendering/RenderTargets, Engine/Rendering/RenderTargets/BlueprintRenderTargets
@@ -8,7 +8,7 @@ parent:Engine/Rendering/RenderTargets/BlueprintRenderTargets
 order:2
 Navigation:Reference
 tags: Render Targets
-tags: Blueprint
+tags: Blueprints
 topic-image:Engine/Rendering/RenderTargets/Render_Targets_Reference.png
 
 [TOC (start:2 end:3)]
diff --git a/Engine/Documentation/Source/Engine/Rendering/RenderTargets/RenderTargets.INT.udn b/Engine/Documentation/Source/Engine/Rendering/RenderTargets/RenderTargets.INT.udn
--- a/Engine/Documentation/Source/Engine/Rendering/RenderTargets/RenderTargets.INT.udn
+++ b/Engine/Documentation/Source/Engine/Rendering/RenderTargets/RenderTargets.INT.udn
@@ -1,4 +1,4 @@
-availability:Docs
+availability:Public
 title: Render Targets
 description:Using Render Targets in Unreal Engine 4.
 crumbs:%ROOT%, Engine, Engine/Rendering
@@ -8,7 +8,6 @@ parent:Engine/Rendering
 order:7
 Navigation:topic
 tags: Render Targets
-tags: Blueprint
 
 ![](BlueprintRenderTarget.png)
 
diff --git a/Engine/Documentation/Source/Engine/Rendering/SceneCapture/SceneCaptureMobile/SceneCaptureMobile.INT.udn b/Engine/Documentation/Source/Engine/Rendering/SceneCapture/SceneCaptureMobile/SceneCaptureMobile.INT.udn
--- a/Engine/Documentation/Source/Engine/Rendering/SceneCapture/SceneCaptureMobile/SceneCaptureMobile.INT.udn
+++ b/Engine/Documentation/Source/Engine/Rendering/SceneCapture/SceneCaptureMobile/SceneCaptureMobile.INT.udn
@@ -1,57 +1,51 @@
 availability:Docs
 title:Scene Capture for Mobile
 description: Things to be aware of when using Scene Captures for Mobile devices.
-crumbs:%ROOT%, Engine, Engine/Rendering, Engine/Rendering/SceneCapture
+crumbs:
 version:4.14
-Type:topic
+Type:reference
 SkillLevel:Intermediate
 parent:Engine/Rendering/SceneCapture
 order:2
 Navigation:topic
 topic-image:Scene_Capture_Mobile.png
+related: Resources/ContentExamples/Reflections/1_6/
+related: Resources/ContentExamples/Reflections/1_7
+tags:mobile
+tags:rendering
+tags:Daydream
+tags:iOS
+tags:Android
+tags:User Interface
+tags:UMG UI Designer
+tags:Materials
+prereq:Gameplay/HowTo/UsingCameras
+prereq:Engine/Rendering/RenderTargets
+prereq:Engine/UMG
 
-When using the Scene Capture Actor in project that are to be used on Mobile devices you will need to aware of the following information to make sure the Scene Capture Actors are as performant as possible when used on Mobile Devices.
 
 
 
-![](Scene_Capture_Mobile.png)
+When using Scene Capture Actors in a project that will be deployed to moblie devices, there are some guidelines you may want to follow to ensure that they are as performant as possible when running on mobile devices.
 
+![](Scene_Capture_Mobile.png)
 
 [region:note]
-When rendering Scene Captures, the Scene Capture Source settings that output Inverse Opacity and Depth values are  supported on Mobile devices. 
+When rendering Scene Captures, the Scene Capture Source settings that output **Inverse Opacity** and **Depth Values** are supported on Mobile devices. 
 [/region]
 
-* The **SceneColor (HDR) in RGB, Inv Opacity in A** option can be used to render objects with translucency into a Texture which can then be alpha-blended over a scene or widget Blueprint. 
-
-* Similarly, the depth value can be used as a mask when using the resulting Texture.
-
-* Generating the Ppacity data has some cost, so use **SceneColor (HDR) in RGB, 0 in A** for improved performance if you do not need Opacity.
-
-* Scene captures now work correctly on devices that do not support floating point targets, such as Galaxy S6 prior to Android 6.
-
-
-
-
-
-
-
-
-
-
-
-
-
-
-
-
-
-
-
-
-
-
+* The **SceneColor (HDR) in RGB, Inv Opacity in A** option can be used to render objects with translucency into a Texture which can then be alpha-blended over a scene or Widget Blueprint. 
 
+* Similarly, the Depth Value can be used as a mask when using the resulting Texture.
 
+* Generating the Opacity data has some cost, so use **SceneColor (HDR) in RGB, 0 in A** for improved performance if you do not need Opacity.
 
+**Supported Capture Sources**
 
+The following is a list of capture sources that are supported on mobile:
 
+* **SceneColor** (HDR) in **RGB**, **Inv Opacity** in **A**
+* **SceneColor** (HDR) in **RGB**, **0** in **A**
+* **SceneColor** (HDR) in **RGB**, **SceneDepth** in **A**
+* **Final Color** (LDR) in **RGB**
+* **Scene Depth** in **R**
diff --git a/Engine/Documentation/Source/Engine/Replay/Replay.INT.udn b/Engine/Documentation/Source/Engine/Replay/Replay.INT.udn
--- a/Engine/Documentation/Source/Engine/Replay/Replay.INT.udn
+++ b/Engine/Documentation/Source/Engine/Replay/Replay.INT.udn
@@ -32,9 +32,11 @@ The basic functions of the replay system are to start or stop recording a demo,
 | `AGameMode::bHandleDedicatedServerReplays` (variable) | N/A | Setting this variable to can be set to true to cause dedicated servers to record games automatically. |
 
 ## DemoNetDriver and Streamers
+
 The `DemoNetDriver` uses **Streamers** to perform the actual extraction and recording of the information needed to create the replay. Depending on how the replay data is intended to be viewed, different Streamers can be attached to the `DemoNetDriver`. There are three different Streamers included with the engine, each suited to different use cases. The **NULL Streamer**, which is the default, records events from the host machine directly to disk, making it best suited for single-player games and games that keep replays locally, on the host player's own machine. The **Memory Streamer** runs on a client machine and stores data in memory, making it ideal for "instant replay" features in sports titles or "kill cams" in shooters. Finally, the **HTTP Streamer** is used to send replay data to a second machine over LAN or the internet, which makes it a great choice for dedicated-server games and for games that may need to be streamed live to a large number of viewers while still remaining responsive to players.
 
+For further details on the features of the `DemoNetDriver` and its Streamers, please see the [](Engine/Replay/Streamers) page.
+
 ## Backward Compatibility
-Backward compatibility is supported by the Replay system, starting in version 4.13. This means you can make modifications to a build, i.e. adding or removing replicated properties, and still load and view replays that were recorded with the original build. In most cases, this is handled automatically, as the Replay system will skip over old replicated fields that have been removed, or use defaults for new replicated fields that have been added. Note that structures with `NetSerialize` overridden will require these differences in data format to be handled manually. To this end, `FArchive` now provides the `EngineNetVer` and `GameNetVer` functions, which identify the engine and game version, respectively, enabling you to adjust replay data as it comes in so that it fits with your current build. The delegate `GetLocalNetworkVersionOverride` in `FNetworkVersion` can be bound to a function that returns a `uint32` version number of your choosing. In case exact matches are not required for compatibility, `FNetworkVersion` has a delegate called `IsNetworkCompatibleOverride` that can be overridden to perform the comparison between two version numbers and decide if they are compatible.
 
-For further details on the features of the `DemoNetDriver` and its Streamers, please see the [](Engine/Replay/Streamers) page.
+Backward compatibility is supported by the Replay system, starting in version 4.13. This means you can make modifications to a build, i.e. adding or removing replicated properties, and still load and view replays that were recorded with the original build. In most cases, this is handled automatically, as the Replay system will skip over old replicated fields that have been removed, or use defaults for new replicated fields that have been added. Note that structures with `NetSerialize` overridden will require these differences in data format to be handled manually. To this end, `FArchive` now provides the `EngineNetVer` and `GameNetVer` functions, which identify the engine and game version, respectively, enabling you to adjust replay data as it comes in so that it fits with your current build. The delegate `GetLocalNetworkVersionOverride` in `FNetworkVersion` can be bound to a function that returns a `uint32` version number of your choosing. In case exact matches are not required for compatibility, `FNetworkVersion` has a delegate called `IsNetworkCompatibleOverride` that can be overridden to perform the comparison between two version numbers and decide if they are compatible.
diff --git a/Engine/Documentation/Source/Engine/Replay/Streamers/HTTPRESTAPI/HTTPRESTAPI.INT.udn b/Engine/Documentation/Source/Engine/Replay/Streamers/HTTPRESTAPI/HTTPRESTAPI.INT.udn
--- a/Engine/Documentation/Source/Engine/Replay/Streamers/HTTPRESTAPI/HTTPRESTAPI.INT.udn
+++ b/Engine/Documentation/Source/Engine/Replay/Streamers/HTTPRESTAPI/HTTPRESTAPI.INT.udn
@@ -1,4 +1,4 @@
-Availability:Docs
+Availability:Public
 Title: HTTP Streamer REST API
 Crumbs: %ROOT%, Engine, Engine/Replay, Engine/Replay/Streamers
 Description:Overview of the REST API for the HTTP Streamer
@@ -9,18 +9,33 @@ tags:Replays
 
 [TOC (start:2 end:3)]
 
-## HTTP REST API Implementation
-
 The following requests will be made by the HTTP Streamer and must be handled by your replay server. In most cases, simply recording the information given in the appropriate location, or providing that information back to a viewer upon request, will be sufficient. It is good to note at this point that replays are not generally stored as a single file. Each individual chunk of a replay generally occupies its own file. Event data, although associated with a replay, is stored separately, since events can be searched without knowing what replay they belong to. Even event groups should be stored as separate files with lists of all events (across all replays) belonging to that group.
 
 The HTTP Streamer's REST API provides functionality for streaming replays from the game server to the replay server (uploading), for streaming live or pre-recorded replays from the replay server to the game, or viewer, server (downloading), or for querying information about the replays available on the server (searching). To implement your own replay server, you will need to respond to the following HTTP requests, as well as any additional requests specific to your game.
 
-### Upload Requests
+## Upload Requests
+
+### Start Streaming
 
----
+Sends a request to begin uploading a replay stream.
+
+#### Syntax
 
 <!-- From StartStreaming(): -->
-* **(Server URL)**replay/**(Session Name Override)**?app=**(Replay Version AppString)**&version=**(Replay Network Version)**&cl=**(Replay CL)**&friendlyName=**(Platform Friendly Name)**
+
+	<Server URL>replay/<Session Name Override>?app=<Replay Version AppString>&version=<Replay Network Version>&cl=<Replay CL>&friendlyName=<Platform Friendly Name>
+	
+#### Details
+
+[REGION:simpletable]
+| Item | Value |
+| ---- | ----- |
+| Request Method | POST | 
+| Content Type | application/json |
+[/REGION]
+
+#### Parameters
+
 | Element | Value | Intended Use / Notes |
 | --- | --- | --- |
 | Server URL | Replay server URL. Drawn from DefaultEngine.Ini, unmodified. | The final "/" in the URL must be included - it will not be added automatically. |
@@ -29,28 +44,72 @@ The HTTP Streamer's REST API provides functionality for streaming replays from t
 | Replay Network Version | A user-defined string that describes the version of the game being played. | This value is also used for filtering lists of replays, so it should be stored with the replay data. |
 | Replay CL | The changelist of the build being played. | This value is used for filtering lists of replays. |
 | Platform Friendly Name | The platform on which the game is being played. | This value might be displayed to users, and is expected to be returned when lists of replays are requested. However, it will not be sent to the game server with requests for lists of replays and therefore is not expected to be used as a filter. |
-Uses POST. "Content-Type" is "application/json".
-<br>JSON format:
-| Variable Name | Variable Type | Content |
+
+#### Response
+
+Expects a response of type `FNetworkReplayStartUploadingResponse`, which contains JSON data with the string "sessionId" defined. This value will be used as the **Session Name** in future communication from the HTTP Streamer. If the optional <Session Name Override> was provided, the returned "sessionId", and thus the ultimate <Session Name>, should be the same value. This is the game's way of demanding the session name.
+
+| Name | Type | Content |
 | --- | --- | --- |
 | Users | String Array | List of users present in the game. Should be stored as tags on the replay. |
-Expects a response of type `FNetworkReplayStartUploadingResponse`, which contains JSON data with the string "sessionId" defined. This value will be used as the **Session Name** in future communication from the HTTP Streamer. If the optional **(Session Name Override)** was provided, the returned "sessionId", and thus the ultimate <Session Name>, should be the same value. This is the game's way of demanding the session name.
 
----
+### Upload Header
+
+Sends a request to upload the binary data constituting the archive header, which should be recorded in a header file associated with **Session Name**.
+
+#### Syntax
 
 <!-- From UploadHeader(): -->
-* **(Server URL)**replay/**(Session Name)**/file/replay.header?numChunks=**(Stream Chunk Index)**&time=**(Total Demo Time In MS)**
+
+	<Server URL>replay/<Session Name>/file/replay.header?numChunks=<Stream Chunk Index>&time=<Total Demo Time In MS>
+	
+#### Details
+
+[REGION:simpletable]
+| Item | Value |
+| ---- | ----- |
+| Request Method | POST | 
+| Content Type | application/octet-stream |
+[/REGION]
+
+#### Parameters
+
 | Element | Value | Intended Use / Notes |
 | --- | --- | --- |
 | Server URL | Replay server URL. Drawn from DefaultEngine.Ini, unmodified. | The final "/" in the URL must be included - it will not be added automatically. |
 | Session Name | Name of the replay being uploaded. | This name will match the value that the replay server provided in response to the initial upload request. |
 | Stream Chunk Index | The index of the current chunk of replay data being sent. | Stream chunks are not uniform in size, but they are strictly ordered by this value. |
 | Total Demo Time In MS | The total running time of the demo up to the end of the current stream chunk. | Tells how much replay data, in milliseconds, has been uploaded so far without having to interpret binary replay data. |
-Uses POST. "Content-Type" is "application/octet-stream".
-The content sent will be the binary data constituting the archive header, which should be recorded in a header file associated with **Session Name**.
+
+### Flush Stream
+
+Sends a request to flush the stream. The content is binary replay data, which can be compressed. This data represents the chunk indicated by **Stream Chunk Index**. Since chunks can vary in data size as well as the length of time they cover, it is recommended that each chunk be stored in a separate file. The HTTP Streamer will attempt to flush its current stream data via this system every ten seconds. The interval between flushes can be adjusted by changing the console variable "httpReplay.ChunkUploadDelayInSeconds"
+
+[REGION:note]
+To enable compression, three functions must be overridden in a child class of `FHttpNetworkReplayStreamer`:
+
+* `SupportsCompression` must return `true`.
+* `CompressBuffer` must perform compression and return `true`, except in the case of an error.
+* `DecompressBuffer` must perform decompression and return `true`, except in the case of an error.
+[/REGION:note]
+
+#### Syntax
 
 <!-- From FlushStream(): -->
-* **(Server URL)**replay/**(Session Name)**/file/stream.**(Stream Chunk Index)**?numChunks=**(Total Chunks)**&time=**(Total Demo Time In MS)**&mTime1=**(Start Time In MS)**&mTime2=**(End Time In MS)**&absSize=**(Total Uploaded Bytes)**
+
+	<Server URL>replay/<Session Name>/file/stream.<Stream Chunk Index>?numChunks=<Total Chunks>&time=<Total Demo Time In MS>&mTime1=<Start Time In MS>&mTime2=<End Time In MS>&absSize=<Total Uploaded Bytes>
+
+#### Details
+
+[REGION:simpletable]
+| Item | Value |
+| ---- | ----- |
+| Request Method | POST | 
+| Content Type | application/octet-stream |
+[/REGION]
+
+#### Parameters
+
 | Element | Value | Intended Use / Notes |
 | --- | --- | --- |
 | Server URL | Replay server URL. Drawn from DefaultEngine.Ini, unmodified. | The final "/" in the URL must be included - it will not be added automatically. |
@@ -61,20 +120,34 @@ The content sent will be the binary data constituting the archive header, which
 | Start Time In MS | The total running time as of the point where this stream chunk begins. | Time is in milliseconds. For streaming replays, can be used to identify which chunk a scrubbing viewer needs. |
 | End Time In MS | The total running time as of the point where this stream chunk ends. | Time is in milliseconds. For streaming replays, can be used to identify which chunk a scrubbing viewer needs. |
 | Total Uploaded Bytes | The total size of the demo up to the end of the current stream chunk. | Tells how much replay data, in bytes, has been uploaded so far. |
-Uses POST. "Content-Type" is "application/octet-stream".
-The content is binary replay data, which can be compressed. This data represents the chunk indicated by **Stream Chunk Index**. Since chunks can vary in data size as well as the length of time they cover, it is recommended that each chunk be stored in a separate file. The HTTP Streamer will attempt to flush its current stream data via this system every ten seconds. The interval between flushes can be adjusted by changing the console variable "httpReplay.ChunkUploadDelayInSeconds"
+
+### Add/Update Event
+
+Sends a request to add or update a replay event. Content contains binary event data to be recorded. This data can be sent through DemoNetDriver and is be user-defined. The second version of this HTTP request can be used to modify or update an existing event, whereas the first version is used to create a new event.
+
 [REGION:note]
-To enable compression, three functions must be overridden in a child class of `FHttpNetworkReplayStreamer`:
-`SupportsCompression` must return `true`.
-`CompressBuffer` must perform compression and return `true`, except in the case of an error.
-`DecompressBuffer` must perform decompression and return `true`, except in the case of an error.
+Note that events are meant to be stored separately from replays, which is why event names use GUIDs. Each event knows what replay it belongs to, which is why events are not embedded in the replays themselves.
 [/REGION:note]
 
----
+#### Syntax
 
 <!-- From AddOrUpdateEvent(): -->
-* **(Server URL)**replay/**(Session Name)**/event?group=**(Group Name)**&time1=**(Event Time In MS)**&time2=**(Event Time In MS)**&meta=**(Meta Tag)**&incrementSize=false
-* **(Server URL)**replay/**(Session Name)**/event/**(Session Name**)`_`**(Event Name)**?group=**(Group Name)**&time1=**(Event Time In MS)**&time2=**(Event Time In MS)**&meta=**(Meta Tag)**&incrementSize=false
+
+	<Server URL>replay/<Session Name>/event?group=<Group Name>&time1=<Event Time In MS>&time2=<Event Time In MS>&meta=<Meta Tag>&incrementSize=false
+~~~
+	<Server URL>replay/<Session Name>/event/<Session Name>_<Event Name>?group=<Group Name>&time1=<Event Time In MS>&time2=<Event Time In MS>&meta=<Meta Tag>&incrementSize=false
+
+#### Details
+
+[REGION:simpletable]
+| Item | Value |
+| ---- | ----- |
+| Request Method | POST | 
+| Content Type | application/octet-stream |
+[/REGION]
+
+#### Parameters
+
 | Element | Value | Intended Use / Notes |
 | --- | --- | --- |
 | Server URL | Replay server URL. Drawn from DefaultEngine.Ini, unmodified. | The final "/" in the URL must be included - it will not be added automatically. |
@@ -83,15 +156,28 @@ To enable compression, three functions must be overridden in a child class of `F
 | Event Name | The GUID of the event being added or updated. | (Optional.) Information about specific, named events may be requested by viewers. If not specified, the event can still be found by listing the events in its group. |
 | Event Time In MS | The time in the replay when this event takes place. | Time is in milliseconds. Although "time1" and "time2" can be used as a begin time and an end time respectively, they are always the same value in the default HTTP Streamer. |
 | Meta Tag | Meta tag associated with this event. | Replay lists can be filtered by meta tag, if the request specifies any. **Meta Tag** will be URL-encoded by the HTTP Streamer during transmission. |
-Uses POST. "Content-Type" is "application/octet-stream" and contains binary event data to be recorded. This data can be sent through DemoNetDriver and is be user-defined. The second version of this HTTP request can be used to modify or update an existing event, whereas the first version is used to create a new event.
-[REGION:note]
-Note that events are meant to be stored separately from replays, which is why event names use GUIDs. Each event knows what replay it belongs to, which is why events are not embedded in the replays themselves.
-[/REGION:note]
 
----
+### Stop Streaming
+
+Sends a request to indicate that the replay is done uploading.
+
+#### Syntax
 
 <!-- From StopUploading(): -->
-* **(Server URL)**replay/**(Session Name)**/stopUploading?numChunks=**(Total Num Chunks)**&time=**(Total Time In MS)**&absSize=**(Total Bytes Uploaded)**
+
+	<Server URL>replay/<Session Name>/stopUploading?numChunks=<Total Num Chunks>&time=<Total Time In MS>&absSize=<Total Bytes Uploaded>
+	
+#### Details
+
+[REGION:simpletable]
+| Item | Value |
+| ---- | ----- |
+| Request Method | POST | 
+| Content Type | application/octet-stream (no content is sent) |
+[/REGION]
+
+#### Parameters
+
 | Element | Value | Intended Use / Notes |
 | --- | --- | --- |
 | Server URL | Replay server URL. Drawn from DefaultEngine.Ini, unmodified. | The final "/" in the URL must be included - it will not be added automatically. |
@@ -99,62 +185,104 @@ Note that events are meant to be stored separately from replays, which is why ev
 | Total Num Chunks | The number of chunks in the finished replay. | |
 | Total Time In MS | The total running time, in milliseconds, of the replay. | |
 | Total Bytes Uploaded | The total size, in bytes, of the replay. | |
-Uses POST. "Content-Type" is "application/octet-stream", but no content is sent with this type of request. The request itself indicates that the replay is done uploading.
 
----
+## Download Requests
+
+### Request Start Download
 
-### Download Requests
+The purpose of this request is to start downloading a replay from the replay server by establishing a connection based on the user name. Immediately after this request goes through successfully, the HTTP Streamer will send a request to download the replay's header, and then another to enumerate all checkpoints in the replay. 
 
----
+#### Syntax
 
 <!-- From StartStreaming(): -->
-* **(Server URL)**replay/**(Session Name)**?shouldKeep=**(Should Keep)**
-| Element | Value | Intended Use / Notes |
-| --- | --- | --- |
-| Server URL | Replay server URL. Drawn from DefaultEngine.Ini, unmodified. | The final "/" in the URL must be included - it will not be added automatically. |
-| Session Name | Name of the replay being uploaded. | This name will match the value that the replay server provided in response to the initial upload request. |
-| Should Keep | "true" or "false" | Indicates whether or not the replay being sent should be saved. |
-Uses POST. "Content-Type" is "application/json" but is empty. This request is intended for debug purposes, and is not called in the existing code.
 
----
+	<Server URL>replay/<Session Name>/startDownloading?user=<User Name>
+			
+#### Details
+
+[REGION:simpletable]
+| Item | Value |
+| ---- | ----- |
+| Request Method | POST | 
+| Content Type | None |
+[/REGION]
+
+#### Parameters
 
-<!-- From StartStreaming(): -->
-* **(Server URL)**replay/**(Session Name)**/startDownloading?user=**(User Name)**
 | Element | Value | Intended Use / Notes |
 | --- | --- | --- |
 | Server URL | Replay server URL. Drawn from DefaultEngine.Ini, unmodified. | The final "/" in the URL must be included - it will not be added automatically. |
 | Session Name | Name of the replay being downloaded. | This name will match the value that the server provided when the replay was originally uploaded. |
 | User Name | Unique Net Id (GUID) | The preferred network GUID that is logged into the client machine. This value comes from the engine and is not intended to be user-supplied. |
-Uses POST. Does not send content. The purpose of this request is to start downloading a replay from the replay server by establishing a connection based on the user name. Immediately after this request goes through successfully, the HTTP Streamer will send a request to download the replay's header, and then another to enumerate all checkpoints in the replay. The expected response to this request is of type `FNetworkReplayStartDownloadingResponse`, which contains JSON data.
-<br>JSON format:
-| Variable Name | Variable Type | Content |
+
+#### Response
+
+The expected response to this request is of type `FNetworkReplayStartDownloadingResponse`, which contains JSON data.
+
+| Name | Type | Content |
 | --- | --- | --- |
 | State | String | The word "Live", to indicate that the game being viewed is happening live. Otherwise, the replay is assumed to be pre-recorded. |
 | NumChunks | Integer | This is expected to be nonzero, and can be updated as streaming continues (used for live replays). |
 | Time | Integer | The total time (in MS) of the replay. Subject to change if **State** indicated that this is a live stream. |
 | ViewerId | String | A unique, name assigned by the replay server that the game server will use when requesting further chunks of replay data. Often referred to as **Viewer Name** in this document. |
 
----
+### Request Header
+
+#### Syntax
 
 <!-- From DownloadHeader(): -->
-* **(Server URL)**replay/**(Session Name)**/file/replay.header
+
+	<Server URL>replay/<Session Name>/file/replay.header
+				
+#### Details
+
+[REGION:simpletable]
+| Item | Value |
+| ---- | ----- |
+| Request Method | GET | 
+| Content Type | None |
+[/REGION]
+
+#### Parameters
+
 | Element | Value | Intended Use / Notes |
 | --- | --- | --- |
 | Server URL | Replay server URL. Drawn from DefaultEngine.Ini, unmodified. | The final "/" in the URL must be included - it will not be added automatically. |
 | Session Name | Name of the replay being downloaded. | This name will match the value that the replay server provided when the replay was originally uploaded. |
-Uses GET. Expects a response with the replay's header in the content portion, as binary data.
 
----
+#### Response
+
+Expects a response with the replay's header in the content portion, as binary data.
+
+### Request Checkpoint Data
+
+#### Syntax
 
 <!-- From EnumerateCheckpoints(): -->
-* **(Server URL)**replay/**(Session Name)**/event?group=checkpoint
+	
+	<Server URL>replay/<Session Name>/event?group=checkpoint
+				
+#### Details
+
+[REGION:simpletable]
+| Item | Value |
+| ---- | ----- |
+| Request Method | GET | 
+| Content Type | None |
+[/REGION]
+
+#### Parameters
+
 | Element | Value | Intended Use / Notes |
 | --- | --- | --- |
 | Server URL | Replay server URL. Drawn from DefaultEngine.Ini, unmodified. | The final "/" in the URL must be included - it will not be added automatically. |
 | Session Name | Name of the replay being downloaded. | This name will match the value that the replay server provided when the replay was originally uploaded. |
-Uses GET. Expects a response of type `FHttpResponsePtr` containing checkpoint data as a JSON-formatted string.
-<br>JSON format:
-| Variable Name | Variable Type | Content |
+
+#### Response
+
+Expects a response of type `FHttpResponsePtr` containing checkpoint data as a JSON-formatted string.
+
+| Name | Type | Content |
 | --- | --- | --- |
 | events | Array | Contains array entries which consist of all the other fields in this table. |
 | id | String | The name (a GUID) of the checkpoint. |
@@ -162,20 +290,42 @@ Uses GET. Expects a response of type `FHttpResponsePtr` containing checkpoint da
 | meta | String | Meta Tag string associated with this event, if any. |
 | time1 | Integer | Start time (in MS) of the event. |
 | time2 | Integer | End time (in MS) of the event. Often the same as the start time. |
+
 [REGION:note]
 Note that this is technically a generic request for a group of events. While "checkpoint" is explicitly specified here, the replay server should be able to handle other values for "group" and report all events for any given group.
 [/REGION:note]
 
----
+### Request Stream Chunk
+
+This is a request for a specific stream chunk, named by the viewer. This request will not be sent while there is a chunk already being downloaded. 
+
+#### Syntax
 
 <!-- From ConditionallyDownloadNextChunk(): -->
-* **(Server URL)**replay/**(Session Name)**/file/stream.**(Stream Chunk Index)**
+
+	<Server URL>replay/<Session Name>/file/stream.<Stream Chunk Index>
+					
+#### Details
+
+[REGION:simpletable]
+| Item | Value |
+| ---- | ----- |
+| Request Method | GET | 
+| Content Type | None |
+[/REGION]
+
+#### Parameters
+
 | Element | Value | Intended Use / Notes |
 | --- | --- | --- |
 | Server URL | Replay server URL. Drawn from DefaultEngine.Ini, unmodified. | The final "/" in the URL must be included - it will not be added automatically. |
 | Session Name | Name of the replay being downloaded. | This name will match the value that the replay server provided when the replay was originally uploaded. |
 | Stream Chunk Index | The index of the current chunk of replay data being requested. | Stream chunks are not uniform in size, but they are strictly ordered by this value. |
-Uses GET. This is a request for a specific stream chunk, named by the viewer. This request will not be sent while there is a chunk already being downloaded. The expected response has a header with some information about the chunk, and its content is a binary chunk of replay data. The binary data can be compressed, if compression is supported. The following fields are expected in the header:
+
+#### Response
+
+The expected response has a header with some information about the chunk, and its content is a binary chunk of replay data. The binary data can be compressed, if compression is supported. The following fields are expected in the header:
+
 | Header Field Name | Description |
 | --- | --- |
 | NumChunks | The total number of chunks in the stream. This can be useful with streamed games, as it lets the user know about new chunks that have been created while viewing. |
@@ -183,20 +333,39 @@ Uses GET. This is a request for a specific stream chunk, named by the viewer. Th
 | State | This variable has one recognized value, "Live". If this variable is missing or is a different value, the streamer will switch to non-live streaming mode. |
 | MTime1 | The start time of the stream, in MS. Can be non-zero for live replay streams. This variable will only be checked if a chunk of data sent, and either this is the first chunk the streamer has downloaded, or the stream is live. |
 | MTime2 | The end time of the stream, in MS. Can change as more chunks come in. This variable will always be checked if a chunk of data is sent, or if the replay is live. |
+
 [REGION:note]
 To enable compression, three functions must be overridden in a child class of `FHttpNetworkReplayStreamer`:
-`SupportsCompression` must return `true`.
-`CompressBuffer` must perform compression and return `true`, except in the case of an error.
-`DecompressBuffer` must perform decompression and return `true`, except in the case of an error.
+
+* `SupportsCompression` must return `true`.
+* `CompressBuffer` must perform compression and return `true`, except in the case of an error.
+* `DecompressBuffer` must perform decompression and return `true`, except in the case of an error.
 [/REGION:note]
 
----
+### Enumerate Events for Replay
+
+Can be called directly from DemoNetDriver, meaning your game project can do this at any time. This request asks the replay server to send a list of all events pertaining to a specific replay (identified by **Session Name**) and filtered by a certain event group. It is recommended that empty group names, as well as a lack of the "group" argument, be supported. 
+
+#### Syntax
 
 <!-- From EnumerateEvents(): -->
-* **(Server URL)**replay/**(Session Name)**/event?group=**(Group Name)**"
-Uses GET. Can be called directly from DemoNetDriver, meaning your game project can do this at any time. This request asks the replay server to send a list of all events pertaining to a specific replay (identified by **Session Name**) and filtered by a certain event group. It is recommended that empty group names, as well as a lack of the "group" argument, be supported. Content received in response to this request is expected to be a JSON-formatted string corresponding to `FReplayEventList` and `FReplayEventListItem`.
-<br>JSON format:
-| Variable Name | Variable Type | Content |
+	
+	<Server URL>replay/<Session Name>/event?group=<Group Name>"
+						
+#### Details
+
+[REGION:simpletable]
+| Item | Value |
+| ---- | ----- |
+| Request Method | GET | 
+| Content Type | None |
+[/REGION]
+
+#### Response
+
+Content received in response to this request is expected to be a JSON-formatted string corresponding to `FReplayEventList` and `FReplayEventListItem`.
+
+| Name | Type | Content |
 | --- | --- | --- |
 | events | Array | Contains elements which consist of all the other fields in this table. |
 | id | String | The name (a GUID) of the checkpoint. |
@@ -205,26 +374,57 @@ Uses GET. Can be called directly from DemoNetDriver, meaning your game project c
 | time1 | Integer | Start time (in MS) of the event. |
 | time2 | Integer | End time (in MS) of the event. Often the same as the start time. |
 
----
+### Refresh Viewer
+
+This request is a heartbeat, sent every ten seconds. This frequency is currently hard-coded, but will likely be made user-configurable without code changes in the future. If a viewer fails to send this request for a long period of time (perhaps one minute), then the viewer has most likely disconnected unexpectedly and it would be reasonable to disconnect them. If the "final=true" version is sent, the user is indicating that this is the final heartbeat, and the user has stopped watching the replay.
+
+#### Syntax
 
 <!-- From RefreshViewer(): -->
-* **(Server URL)**replay/**(Session Name)**/viewer/**(Viewer Name)**
-* **(Server URL)**replay/**(Session Name)**/viewer/**(Viewer Name)**?final=**(Final Heartbeat)**
+	
+	<Server URL>replay/<Session Name>/viewer/<Viewer Name>
+~~~
+	<Server URL>replay/<Session Name>/viewer/<Viewer Name>?final=<Final Heartbeat>
+							
+#### Details
+
+[REGION:simpletable]
+| Item | Value |
+| ---- | ----- |
+| Request Method | POST | 
+| Content Type | None |
+[/REGION]
+
+#### Parameters
+
 | Element | Value | Intended Use / Notes |
 | --- | --- | --- |
 | Server URL | Replay server URL. Drawn from DefaultEngine.Ini, unmodified. | The final "/" in the URL must be included - it will not be added automatically. |
 | Final Heartbeat | "true" or "false" | By default, "final" will not be sent unless **Final Heartbeat** is true. If true, the viewer has stopped watching and no further data should be sent. |
-Uses POST. Does not send content.
-This request is a heartbeat, sent every ten seconds. This frequency is currently hard-coded, but will likely be made user-configurable without code changes in the future. If a viewer fails to send this request for a long period of time (perhaps one minute), then the viewer has most likely disconnected unexpectedly and it would be reasonable to disconnect them. If the "final=true" version is sent, the user is indicating that this is the final heartbeat, and the user has stopped watching the replay.
 
----
+## Search Requests
 
-### Search Requests
+### Search Replays
 
----
+Request a list of all replays that the replay server can send. These replays are expected to be pre-filtered by the user's criteria in the URL. The replay server should be able to handle some of the filter fields being empty or missing.
+
+#### Syntax
 
 <!-- From EnumerateStreams(): -->
-* **(Server URL)**replay?app=**(Replay Version AppString)**&cl=**(Replay CL)**&version=**(Replay Network Version)**&meta=**(Meta Tag)**&user=**(User Name)**&recent=**(In Recent Viewer)**"
+
+	<Server URL>replay?app=<Replay Version AppString>&cl=<Replay CL>&version=<Replay Network Version>&meta=<Meta Tag>&user=<User Name>&recent=<In Recent Viewer>"
+								
+#### Details
+
+[REGION:simpletable]
+| Item | Value |
+| ---- | ----- |
+| Request Method | GET | 
+| Content Type | None |
+[/REGION]
+
+#### Parameters
+
 | Element | Value | Intended Use / Notes |
 | --- | --- | --- |
 | Server URL | Replay server URL. Drawn from DefaultEngine.Ini, unmodified. | The final "/" in the URL must be included - it will not be added automatically. |
@@ -234,14 +434,17 @@ This request is a heartbeat, sent every ten seconds. This frequency is currently
 | Meta | Tags applied to this replay. | User-applied meta tags can be used for filtering games that have certain features. This could be anything from "all Capture The Flag matches" to to "games that went into double-overtime". |
 | User | Unique Net Id (GUID) | This is used to filter games based on their inclusion of a specific player. |
 | Recent | "true" or "false" | If true, only recently-uploaded replays should be shown. What defines "recent" is up to the replay server. |
-Uses GET. Expects a response with content in the form of a JSON-formatted string, containing a list of replays that the replay server can send. These replays are expected to be pre-filtered by the user's criteria in the URL. The replay server should be able to handle some of the filter fields being empty or missing.
-<br>JSON format:
-| Variable Name | Variable Type | Content |
+
+#### Response
+
+Expects a response with content in the form of a JSON-formatted string.
+
+| Name | Type | Content |
 | --- | --- | --- |
 | replays | Array | Contains elements which consist of all the other fields in this table. |
 | AppName | String | The **Replay Version AppString** that was sent when the replay was uploaded. Should match the parameter you sent with the request. |
 | SessionName | String | The name (GUID) of the replay. |
-| FriendlyName | String | The **(Platform Friendly Name)** that was provided when the replay was uploaded. |
+| FriendlyName | String | The <Platform Friendly Name> that was provided when the replay was uploaded. |
 | Timestamp | DateTime | The time and date associated with the replay (i.e. when this game happened). |
 | SizeInBytes | Integer | Total size, in bytes, of the replay. |
 | DemoTimeInMs | Integer | Running length, in MS, of the replay. |
@@ -250,22 +453,42 @@ Uses GET. Expects a response with content in the form of a JSON-formatted string
 | Changelist | Integer | The changelist (CL) of the engine running the game when this replay was recorded. Corresponds to the **Replay CL** argument sent when uploading the replay. Should match the parameter you sent with the request. |
 | shouldKeep | Boolean | True if the debug "shouldKeep" request was made. |
 
----
+### Search Replays by Event
+
+Sends a request for a list of replays from the server. These replays are expected to be pre-filtered to include only those that contain at least one event in the named **Event Group**.
+
+#### Syntax
 
 <!-- From SearchEvents(): -->
-* **(ServerURL)**event?group=**(Event Group)**
+
+	<ServerURL>event?group=<Event Group>
+								
+#### Details
+
+[REGION:simpletable]
+| Item | Value |
+| ---- | ----- |
+| Request Method | GET | 
+| Content Type | None |
+[/REGION]
+
+#### Parameters
+
 | Element | Value | Intended Use / Notes |
 | --- | --- | --- |
 | Server URL | Replay server URL. Drawn from DefaultEngine.Ini, unmodified. | The final "/" in the URL must be included - it will not be added automatically. |
-| Group Name | The name of the event group containing this event. | Events can be put into groups for filtering lists, such as grouping all penalty calls in a sports title. |
-Uses GET. Expects a response with content as a JSON-formatted string containing a list of replays from the server. These replays are expected to be pre-filtered to include only those that contain at least one event in the named **Event Group**.
-<br>JSON format:
-| Variable Name | Variable Type | Content |
+| Event Group | The name of the event group containing this event. | Events can be put into groups for filtering lists, such as grouping all penalty calls in a sports title. |
+
+#### Response
+
+Expects a response with content as a JSON-formatted string containing a list of replays from the server.
+
+| Name | Type | Content |
 | --- | --- | --- |
 | replays | Array | Contains elements which consist of all the other fields in this table. |
 | AppName | String | The **Replay Version AppString** that was sent when the replay was uploaded. |
 | SessionName | String | The name (GUID) of the replay. |
-| FriendlyName | String | The **(Platform Friendly Name)** that was provided when the replay was uploaded. |
+| FriendlyName | String | The <Platform Friendly Name> that was provided when the replay was uploaded. |
 | Timestamp | DateTime | The time and date associated with the replay (i.e. when this game happened). |
 | SizeInBytes | Integer | Total size, in bytes, of the replay. |
 | DemoTimeInMs | Integer | Running length, in MS, of the replay. |
@@ -274,14 +497,32 @@ Uses GET. Expects a response with content as a JSON-formatted string containing
 | Changelist | Integer | The changelist (CL) of the engine running the game when this replay was recorded. Corresponds to the **Replay CL** argument sent when uploading the replay. |
 | shouldKeep | Boolean | True if the debug "shouldKeep" request was made. |
 
----
+### Request Replay
+
+This is called directly from the DemoNetDriver, so it can be done from the game project at any time, but will usually be called while viewing a replay with an **Event Name** retrieved from a list provided by a previous request.
+
+#### Syntax
 
 <!-- From RequestEventData(): -->
-* **(Server URL)**event/**(Event Name)**
+	
+	<Server URL>event/<Event Name>
+								
+#### Details
+
+[REGION:simpletable]
+| Item | Value |
+| ---- | ----- |
+| Request Method | GET | 
+| Content Type | None |
+[/REGION]
+
+#### Parameters
+
 | Element | Value | Intended Use / Notes |
 | --- | --- | --- |
 | Server URL | Replay server URL. Drawn from DefaultEngine.Ini, unmodified. | The final "/" in the URL must be included - it will not be added automatically. |
 | Event Name | The GUID of the event requested. | Since events are identified by GUID, the event can be from any replay on file. Make sure that the event you request is from a list of events pertaining to the replay you are currently watching. |
-Uses GET. Expects a response with Content as a binary stream from the replay server. This is called directly from the DemoNetDriver, so it can be done from the game project at any time, but will usually be called while viewing a replay with an **Event Name** retrieved from a list provided by a previous request.
 
----
+#### Response
+
+Expects a response with Content as a binary stream from the replay server.
\ No newline at end of file
diff --git a/Engine/Documentation/Source/Engine/Replay/Streamers/Streamers.INT.udn b/Engine/Documentation/Source/Engine/Replay/Streamers/Streamers.INT.udn
--- a/Engine/Documentation/Source/Engine/Replay/Streamers/Streamers.INT.udn
+++ b/Engine/Documentation/Source/Engine/Replay/Streamers/Streamers.INT.udn
@@ -1,5 +1,5 @@
-Availability:Docs
-Title: Replay System
+Availability:Public
+Title: DemoNetDriver and Streamers
 Crumbs: %ROOT%, Engine, Engine/Replay
 Description:Overview of the Replay system for recording and playback of gameplay
 Version: 4.14
diff --git a/Engine/Documentation/Source/Engine/Sequencer/HowTo/AnimateDynamicObjects/AnimateDynamicObjects.INT.udn b/Engine/Documentation/Source/Engine/Sequencer/HowTo/AnimateDynamicObjects/AnimateDynamicObjects.INT.udn
new file mode 100644
--- /dev/null
+++ b/Engine/Documentation/Source/Engine/Sequencer/HowTo/AnimateDynamicObjects/AnimateDynamicObjects.INT.udn
@@ -0,0 +1,172 @@
+Availability: Public
+Title:Animating Dynamic Objects at Runtime
+Crumbs: %ROOT%, Engine, Engine/Sequencer, Engine/Sequencer/HowTo
+Description:This example illustrates how you can apply Sequences to dynamic objects at runtime that are different than the one the Sequence was authored with.
+Version: 4.15
+SkillLevel: Advanced
+parent:Engine/Sequencer/HowTo
+order:1
+checkpoint: editorqs
+Related: Engine/Sequencer/Overview
+Related: Engine\Sequencer\ReferenceEditor\
+tags:Sequencer
+
+When creating cinematics for gameplay purposes, there may be times when you want to animate an object that is dynamically spawned at runtime through Sequencer. 
+For instance, you create an animated object that moves along a path but that object is something that can be defined by the player. 
+With the help of [Blueprint Visual Scripting](Engine/Blueprints), you can expose binding identifiers from a Level Sequence using the **Get Sequence Bindings** node and override any of those bindings with your own. 
+
+In this example, we animate and move an object along a path and allow the player to change that object by pressing a button as seen below. 
+
+[OBJECT:EmbeddedVideo]
+[PARAMLITERAL:width]
+640
+[/PARAMLITERAL]
+[PARAMLITERAL:height]
+360
+[/PARAMLITERAL]
+[PARAMLITERAL:videoid]
+XdahqJz2cyM
+[/PARAMLITERAL]
+[/OBJECT]
+
+[REGION:note]
+For this how-to, we are using the **Blueprint Third Person Template** project and with **Starter Content** enabled. 
+[/REGION]
+
+## Steps
+
+1.	In the **Content Browser**, create two new **Blueprints** based off the **Actor** class and name one **Cube_Blueprint** and the other **Sphere_Blueprint**. 
+
+	![](DynamicAnim_01.png)
+
+	We will use the Cube Blueprint in our default Sequence but will swap in our Sphere Blueprint when the player presses a key. 
+
+1.	Inside each Blueprint, add a **Static Mesh** component with one using a **Shape_Cube** mesh and the other using the **Shape_Sphere** mesh. 
+
+	![](DynamicAnim_02.png)
+
+1.	From the **Modes** panel, drag an **Empty Actor** anywhere into the Level Viewport. 
+
+	![](DynamicAnim_03.png)(w:940)
+
+1.	From the **Main Toolbar**, click the **Cinematics** button and **Add Level Sequence**, then assign it any name. 
+
+	![](DynamicAnim_04.png)
+
+1.	In your sequence, with the **Empty Actor** in the level selected, click **Add** then **Add Actor** to add it to the sequence. 
+
+	![](DynamicAnim_05.png)
+
+1.	Add a **Transform** track for the **Empty Actor** and add a key at **0** for its initial position in the level. 
+
+	![](DynamicAnim_06.png)
+
+1.	Move the timeline to frame **50** then move the **Empty Actor** in the level and press **S** to add a new key at its new position. 
+
+	![](DynamicAnim_07.png)(w:740)
+
+1.	Move the timeline to frame **100** then move the **Empty Actor** back to its starting position and press **S** to add a key at its ending transform. 
+
+	![](DynamicAnim_08.png)(w:740)
+
+	You can also press the **]** key inside Sequencer while at frame 100 to end the Sequence at this key. 
+
+1.	Select your Level Sequence in the Level Viewport, then in the **Details** panel, set **Loop** to **Lopp Indefinitely**. 
+
+	![](DynamicAnim_09.png)
+
+1.	With your Level Sequence selected in the Level Viewport, from the **Main Toolbar**, select **Blueprints** and **Open Level Blueprint**. 
+
+	![](DynamicAnim_10.png)
+
+1.	**Right-click** in the graph and create a reference to your Level Sequence. 
+
+	![](DynamicAnim_10b.png)
+
+	In our case, we named our Level Sequence **Master**, however, your name may differ. 
+
+1.	**Right-click** and add an **Event Begin Play** node, then drag off your Level Sequence reference and add a **Play (Sequence Player)** node. 
+
+	![](DynamicAnim_11.png)
+
+	Then connect the nodes as shown below. 
+
+	![](DynamicAnim_12.png)
+
+	This will automatically start playing our Sequence when we start the game. 
+
+1.	**Right-click** in the graph and add a **1** Keyboard Event and connect it to a **Flip Flop** node. 
+
+	![](DynamicAnim_13.png)
+
+	When we press 1 the first time, we will change the Empty Actor to our Cube, when we press 1 again, we will change it to our Sphere. 
+
+1.	**Right-click** and create a **Spawn Actor From Class** node and set the class as **Cube Blueprint**. 
+
+1.	**Right-click** the **Return Value** and **Promote to Variable** called **Cube**, then drag off **Spawn Transform** and use **Make Transform**. 
+
+	![](DynamicAnim_14.png)
+
+	Here we are spawning our Cube Blueprint when 1 is pressed, storing it as a variable called Cube and giving it a default transform to spawn it at. 
+
+1.	Off the **B** pin of the **Flip Flop** node, use **Spawn Actor from Class** (set to **Sphere Blueprint**), create a variable to store it and connect as shown below. 
+
+	![](DynamicAnim_15.png)
+
+	When we press 1 a second time, we are now creating and spawning in our Sphere Blueprint and storing it as a variable for use later. 
+
+1.	**Right-click** in the graph and add the **Get Sequence Bindings** node and in the **Details** panel, point it to your Level Sequence. 
+
+	![](DynamicAnim_16.png)
+
+	This will allow us to access and override the contains within our Level Sequence through Blueprint Script. 
+
+1.	**Right-click** in the graph and add the **Add Binding** node and connect the following nodes: 
+
+	![](DynamicAnim_17.png)
+
+	Above we have connected the reference to our Level Sequence as the **Target**, the **Get Sequence Bindings** node as the **Binding** and the variable we created for our spawned Cube as the **Actor**. 
+	This will allow us to replace the Empty Actor (**Actor**) in our Level Sequence with the Actor we specify (our **Cube** that we spawn or our **Sphere**). The animation that was applied to the Empty Actor will also be applied to the Actor we set as part of the binding. 
+
+1.	Repeat the previous step for the **Sphere** portion of script, replacing **Cube** with **Sphere** as the **Actor**. 
+
+	![](DynamicAnim_18.png)
+
+1.	Following the **Add Binding** node for the Cube portion of script, use the **Sphere** variable and perform an **? IsValid** check and **Destroy Actor** if valid. 
+
+	![](DynamicAnim_19.png)
+
+	This will destroy the Sphere when we spawn in the Cube. 
+
+1.	Repeat the previous step for the **Sphere** portion of script, using **? IsValid** on the **Cube** variable and **Destroy Actor** if valid. 
+
+	![](DynamicAnim_20.png)
+
+	This will destroy the Cube when we spawn the Sphere. 
+
+1.	Play in the editor, then press the **1** key to cycle between our spawned Actors. 
+	
+
+
+## End Result
+
+When we play in the editor and press the 1 key, our Empty Actor is replaced by a spawned Actor either the Cube or the Sphere but continues to use the animation we applied. 
+
+[OBJECT:EmbeddedVideo]
+[PARAMLITERAL:width]
+640
+[/PARAMLITERAL]
+[PARAMLITERAL:height]
+360
+[/PARAMLITERAL]
+[PARAMLITERAL:videoid]
+XdahqJz2cyM
+[/PARAMLITERAL]
+[/OBJECT]
+
+You can also set and bind an array of Actors to a binding with the **Set Binding** node, **Reset Binding** for a particular binding or **Reset Bindings** for an entire Level Sequence. 
+
+![](AdditionalFunctions.png)
+
+
+
diff --git a/Engine/Documentation/Source/Engine/Sequencer/HowTo/BurnIns/BurnIns.INT.udn b/Engine/Documentation/Source/Engine/Sequencer/HowTo/BurnIns/BurnIns.INT.udn
--- a/Engine/Documentation/Source/Engine/Sequencer/HowTo/BurnIns/BurnIns.INT.udn
+++ b/Engine/Documentation/Source/Engine/Sequencer/HowTo/BurnIns/BurnIns.INT.udn
@@ -9,6 +9,7 @@ order:1
 checkpoint: editorqs
 Related: Engine/Sequencer/Overview
 Related: Engine/Sequencer/ReferenceEditor
+tags:Sequencer
 
 [EXCERPT:intro]
 When creating and rendering out cinematic movies, you may want to include overlays with information about the scene you are viewing such as the name of the shot, the date, time or frame information. 
diff --git a/Engine/Documentation/Source/Engine/Sequencer/HowTo/CameraRigCrane/CameraRigCrane.INT.udn b/Engine/Documentation/Source/Engine/Sequencer/HowTo/CameraRigCrane/CameraRigCrane.INT.udn
--- a/Engine/Documentation/Source/Engine/Sequencer/HowTo/CameraRigCrane/CameraRigCrane.INT.udn
+++ b/Engine/Documentation/Source/Engine/Sequencer/HowTo/CameraRigCrane/CameraRigCrane.INT.udn
@@ -1,19 +1,33 @@
 Availability: Public
 Title:Shooting from a Camera Rig Crane
-Crumbs: %ROOT%, Engine, Engine/Sequencer, Engine/Sequencer/HowTo
-Description:Here we use a Camera Rig Crane to shoot a scene with sweeping movement. 
-Version: 4.12
+Crumbs:
+Description:Learn how to use a Camera Rig Crane to shoot a scene with sweeping movement. 
+Version: 4.14
 SkillLevel: Intermediate
 parent:Engine/Sequencer/HowTo
 order:1
 checkpoint: editorqs
 Related: Engine/Sequencer/Overview
-Related: Engine\Sequencer\ReferenceEditor\
+Related: Engine/Sequencer/ReferenceEditor
+tags:Sequencer
 
-One of the tools real-world filmmakers use to produce smooth sweeping shots is by attaching a camera to a crane and controlling the shot through the crane's movement.
-You can create similar shots in Sequencer with the use of the **Camera Rig Crane** Actor and an attached **Camera**. 
+One of the methods real-world filmmakers use to produce smooth, sweeping shots, is by attaching a camera to a crane and controlling the shot with the crane's movement.
+You can create similar shots in Sequencer with the use of the **Camera Rig Crane** Actor and an attached **Camera**. You can keyframe the crane's Pitch, Yaw, or the length 
+of the Crane Arm, as well as Lock the Mounted Camera's Pitch or Yaw (which will follow the crane's movement). 
 
-In this how-to we will add a Camera Rig Crane, attach a Camera, and create a sample Crane Shot. 
+In this guide, we will add a Camera Rig Crane, attach a Camera, and create a sample Crane Shot as indicated in the example below: 
+
+[OBJECT:EmbeddedVideo]
+[PARAMLITERAL:width]
+640
+[/PARAMLITERAL]
+[PARAMLITERAL:height]
+360
+[/PARAMLITERAL]
+[PARAMLITERAL:videoid]
+WnwtVEBUORc
+[/PARAMLITERAL]
+[/OBJECT]
 
 
 ## Steps
@@ -22,95 +36,121 @@ In this how-to we will add a Camera Rig Crane, attach a Camera, and create a sam
 For this how-to, we are using the **Blueprint Third Person Template** project with **Starter Content** enabled. 
 [/REGION]
 
-1.	In your project on the **Modes** panel under **Cinematic**, drag a **Camera Rig Crane** asset into the level. 
+1.	In the Level Viewport of your project, select the **ThirdPersonCharacter**, then hold **Alt**, drag out, and rotate the copy so that it faces the existing character. 
 
 	![](Crane01.png)
 
-1.	Also under **Cinematic**, drag a **Cine Camera Actor** into the level. 
+	These two characters are going to be the subject of our crane shot. 
+
+1.	From the Main Toolbar, click the **Cinematics** button, then select **Add Level Sequence** from the drop-down menu. 
 
 	![](Crane02.png)
 
-1.	In the **World Outliner** window, drag-and-drop the **CineCameraActor** onto the **CameraRig_Crane**. 
+1.	From the **Modes** panel under **Cinematic**, drag a **Camera Rig Crane** into the level. 
 
 	![](Crane03.png)
 
-	This will attach the camera to the crane so that it will follow the crane's movement. 
-
-1.	Select the **CineCameraActor**, then in the **Details** panel under **Transform** enter "0" for **Location** and **Rotation** values. 
+1.	From the **Modes** panel under **Cinematic**, drag a **Cine Camera Actor** into the level. 
 
 	![](Crane04.png)
 
-	This will zero out the camera's transform and attach it correctly to the crane. 
-
-1.	From the Main Toolbar, click **Cinematics** then select **Add Level Sequence**. 
+1.	In the **World Outliner** panel, drag the **Cine Camera Actor** on top of the **Camera Rig Crane** to attach it. 
 
 	![](Crane05.png)
 
-1.	Hold **Ctrl** and select both the **CameraRigCrane** and **CineCameraActor**, then add both to **Sequencer**. 
+	This will attach the camera to the crane, allowing it to move where the crane moves. 
+
+1.	In the **Details** panel for the **Cine Camera Actor**, set the **Location** and **Rotation** values to **0.0**. 
 
 	![](Crane06.png)
 
-1.	Click **Add** then select the **Camera Cut Track**. 
+	This will allow the camera to be attached to the Camera Rig Crane's mount position. 
 
-	![](Crane07.png)
+1.	Make any adjustments to the position of the **Camera Rig Crane** to set up your shot (below, the position is behind the character). 
 
-	This will allow us to define which camera to shoot this Level Sequence from. 
+	![](Crane07.png)
 
-1.	On the **Camera Cuts** track, click the **+ Camera** button then select the **CineCameraActor**. 
+1.	Hold **Ctrl** and select the **Camera Rig Crane** and **Cine Camera Actor** (so that both are selected), then click **Add** in Sequencer and add both to the Level Sequence. 
 
 	![](Crane08.png)
 
-	We are now ready to shoot our scene using our Camera Rig Crane and Cine Camera Actor. 
-	You can preview what the scene will look like through the perspective of the camera we added by clicking the **Lock Viewport to Camera** button on the **Camera Cuts Track**. 
+1.	Select the **Camera Rig Crane**, then in the **Details** panel, click the **Add Key** button for **Crane Pitch**, **Crane Yaw** and **Crane Arm Length**. 
 
 	![](Crane09.png)
 
-	Clicking **Lock Viewport to Camera** again will unlock the camera. 	
+	This will set the default position of each to start the sequence. 
 
-1.	Using the **Translate** tool (**W** key) and **Rotation** tool (**E** key), rotate the crane/camera as shown below. 
+1.	In Sequencer, select the **Cine Camera Actor**, then press the **S** key. 
 
 	![](Crane10.png)
 
-	We are going to raise the crane and push it out to focus and reveal the character standing in the level for this particular shot. 
+	This is a shortcut for adding a key for the current **Transform** values, initializing the camera's position. 
 
-1.	Add a key for the **Transform** location of the **Cine Camera Actor** by clicking the **Add Key** button on the **Transform Track**. 
+1.	Scrub the Timeline to frame **50**, changing the **Crane Pitch** value on the **Camera Rig Crane** to **40.0**, and then click the **Add Key** button. 
 
 	![](Crane11.png)
 
-1.	Click the **+ Track** button on the **CameraRig_Crane** and add the **Crane Arm Length**, **Crane Pitch** and **Crane Yaw** tracks. 
+	If you move the Timeline back to 0 and scrub through to 50, you will see the crane move between the two keyframes (as seen below). 
 
-	![](Crane12.png)
+	[OBJECT:EmbeddedVideo]
+	[PARAMLITERAL:width]
+	640
+	[/PARAMLITERAL]
+	[PARAMLITERAL:height]
+	360
+	[/PARAMLITERAL]
+	[PARAMLITERAL:videoid]
+	LiJ2TXYz7A8
+	[/PARAMLITERAL]
+	[/OBJECT]
 
-1.	Scrub the Timeline Marker ahead to **50** then set **Crane Arm Length** to **1000** and **Crane Pitch** to **20** and add keys for both. 
+	[REGION:note]
+	Also included in the video (above), is how enabling the **Lock Mount Pitch** option affects the attached camera. 
+	With the option disabled, the camera maintains its set position. When enabling the option, the mounted camera follows the pitch of the crane and automatically changes as the crane moves. 
+	This option can be keyframed so that you can turn the option on/off through Sequencer based on your shot needs. 
+	[/REGION]
 
-	![](Crane13.png)
+1.	With the Timeline at frame **50**, select the **Cine Camera Actor**, and then press **S** to add another keyframe. 
 
-	If you scrub back-and-forth between the start and frame 50 you should see the crane movement in the viewport.
+	![](Crane12.png)
 
-1.	Scrub the Timeline Marker to the end of the Sequence (**150**).
+1.	Scrub the Timeline ahead to the end of the shot, change the **Crane Yaw** and **Crane Arm Length** to **75** and **600** respectively, then add keys for each. 
 
-1.	In the **Details** panel for the **CameraRig_Crane**, set **Crane Yaw** to **-25** and add a key. 
+	![](Crane13.png)
 
-	![](Crane14.png)
+	You should see the values shift between keyframes for the Crane Yaw and Crane Arm Length as you scrub through the Timeline. 
 
-	You can set the **Crane Pitch**, **Crane Yaw** and **Crane Arm Length** from the **Details** panel by clicking the **Add Key** button next to each property. 
+	[OBJECT:EmbeddedVideo]
+	[PARAMLITERAL:width]
+	640
+	[/PARAMLITERAL]
+	[PARAMLITERAL:height]
+	360
+	[/PARAMLITERAL]
+	[PARAMLITERAL:videoid]
+	d7K3mHA0bXM
+	[/PARAMLITERAL]
+	[/OBJECT]
 
-1.	In the **Details** panel for the **CineCameraActor**, set the **XYZ** values for **Rotation** to **0, -25, -50**. 
+	[REGION:note]
+	Also, similar to the **Lock Mount Pitch**, the **Lock Mount Yaw** option allows you to force the mounted camera to be locked to the yaw position of the crane. 
+	If you enable both Lock Mount Pitch and Lock Mount Yaw, the mounted camera will pitch and yaw in the same direction of the crane. 
+	[/REGION]
 
-	![](Crane16.png)
+1.	Rotate the camera in the Level Viewport to frame up the two characters, then press **S** to add a keyframe for the position. 
 
-	The values may not be exact and might be auto adjusted which is fine. 
+	![](Crane14.png)
 
-1.	With the **Cine Camera Actor** still selected, press **Shift + E** to add a key for the rotational values. 
+1.	In Sequencer, click the **Camera Lock** button on the **Cine Camera Actor** track. 
 
-	![](Crane17.png)
+	![](Crane15.png)
 
-1.	Lock the viewport to the **Camera Cuts** track then click the **Preview Play** button. 
+	This will lock the viewport to the perspective of the camera and allow us to view what the shot will look like when using this camera.
 
 	
 ## End Result
 
-You should have something similar to below where the camera pans up along with the crane and pushes in towards the character. 
+You should have something similar to the following video, where the camera movement follows the path defined by your crane. 
 
 [OBJECT:EmbeddedVideo]
 [PARAMLITERAL:width]
@@ -120,14 +160,18 @@ You should have something similar to below where the camera pans up along with t
 360
 [/PARAMLITERAL]
 [PARAMLITERAL:videoid]
-RAsfP5G0Shk
+WnwtVEBUORc
 [/PARAMLITERAL]
 [/OBJECT]
 
-You can adjust the keys by dragging them to new values in the Tracks Area. Above we slowed the cranes upward movement by moving the **Crane Arm Length** and **Crane Pitch** keyframes we added at frame **50** and moved them to frame **100**. 
+In the video, we lock the viewport to our camera and turn on **Game View** (**G** key) to hide all editor icons to give us an idea of what the shot will look like. 
+
+In this how-to, we added the Camera Rig Crane to the level, then added a Cine Camera Actor and manually attached it to the crane.
+There is an advanced method you can use when adding a Camera Rig Crane asset to a level that will automate several of the steps from this how-to, which will automatically add and attach a Cine Camera Actor to your Camera Rig Crane. 
+It will also add both assets to an open Level Sequence, assign a Camera Cuts Track, and assign the camera to use for the shot. 
 
 [REGION:tip]
-There is an advanced method you can use to work with the **Camera Rig Crane** asset in Sequencer that will automate several of the steps from this how-to. Please see the [**Workflow Shortcuts**](Engine\Sequencer\Workflow\CheatSheet) section for more information. 
+Please see the [Workflow Shortcuts](Engine\Sequencer\Workflow\CheatSheet) Sequencer documentation for the automated method and other Sequencer tips. 
 [/REGION]
 
 
@@ -159,6 +203,10 @@ There is an advanced method you can use to work with the **Camera Rig Crane** as
 
 
 
+
+
+
+
 
 
 
diff --git a/Engine/Documentation/Source/Engine/Sequencer/HowTo/CameraRigRail/CameraRigRail.INT.udn b/Engine/Documentation/Source/Engine/Sequencer/HowTo/CameraRigRail/CameraRigRail.INT.udn
--- a/Engine/Documentation/Source/Engine/Sequencer/HowTo/CameraRigRail/CameraRigRail.INT.udn
+++ b/Engine/Documentation/Source/Engine/Sequencer/HowTo/CameraRigRail/CameraRigRail.INT.udn
@@ -9,6 +9,7 @@ order:1
 checkpoint: editorqs
 Related: Engine/Sequencer/Overview
 Related: Engine\Sequencer\ReferenceEditor\
+tags:Sequencer
 
 The **Camera Rig Rail** Actor can be used to produce fly-through, 360 degree, or any other moving shots by attaching a camera to a rail to move along.
 The path for the rail can be defined using a [**Spline Component**](Engine/Components/Shapes/) and in Sequencer, you can keyframe the position along the rail where the camera should be at any given time. 
diff --git a/Engine/Documentation/Source/Engine/Sequencer/HowTo/CharacterAnimation/CharacterAnimation.INT.udn b/Engine/Documentation/Source/Engine/Sequencer/HowTo/CharacterAnimation/CharacterAnimation.INT.udn
--- a/Engine/Documentation/Source/Engine/Sequencer/HowTo/CharacterAnimation/CharacterAnimation.INT.udn
+++ b/Engine/Documentation/Source/Engine/Sequencer/HowTo/CharacterAnimation/CharacterAnimation.INT.udn
@@ -1,85 +1,98 @@
 Availability: Public
-Title:Animating a Character
-Crumbs: %ROOT%, Engine, Engine/Sequencer, Engine/Sequencer/HowTo
-Description:Illustrates how you can apply animation to a Skeletal Mesh Actor in Sequencer
-Version: 4.12
+Title:Animation and Animation Blending
+Crumbs: 
+Description:Illustrates how you can apply animation to a Skeletal Mesh Actor in Sequencer and blend animations. 
+Version: 4.15
 SkillLevel: Intermediate
 parent:Engine/Sequencer/HowTo
 order:1
 checkpoint: editorqs
 Related: Engine/Sequencer/Overview
-Related: Engine\Sequencer\ReferenceEditor\
+Related: Engine/Sequencer/ReferenceEditor
+tags:Sequencer
+type:how-to
 
 When creating cinematic sequences, there may be times where you want to have a character in your scene move around or perform some kind of animation. 
-In **Sequencer** this can be done by adding a Skeletal Mesh Actor to a Level Sequence then adding an Animation sub-track and specifying the animation that you want the Actor to perform. 
+In **Sequencer**, this can be done by adding a Skeletal Mesh Actor to a Level Sequence, then adding an Animation sub-track, and then specifying the animation that you want the Actor to perform. 
 Repeating this process and adding different kinds of animations for your Actor to perform is what can bring your characters to life in your cinematic sequences! 
 
-Here in this guide we will create a small scene that has a character jump up and down by applying different animations. 
+In the following guide, you will create a small scene that starts with a character walking, then blending from walking into a jumping animation. 
+
+<!--Add Hero Image here, showing users what they will create during the activity (an End Result picture is fine)-->
 
 ## Steps
 
 [REGION:note]
-For this how-to, we are using the **Blueprint Third Person Template** project with **Starter Content Enabled**. 
+For this guide, we are using the **Blueprint Third Person Template** project with **Starter Content Enabled**. 
 [/REGION]
 
-1.	With your project open, select the **ThirdPersonCharacter** in the level and press **Delete** to delete it. 
+1.	Inside your project, from the Main Toolbar, click the **Cinematics** button, and then click the **Add Level Sequence** button. 
 
 	![](CharacterAnim_01.png)
 
-	This is the **Character Blueprint** you use when playing the game that comes with this template which we will not need. 
-
-1.	In the **Content Browser** under **Content/Mannequin/Character/Mesh**, drag the **SK_Mannequin** asset into the level viewport. 
+1.	In the **Content Browser** under **Content/Mannequin/Character/Mesh**, drag the **SK_Mannequin** into the level. 
 
 	![](CharacterAnim_02.png)
 
-1.	From the **Toolbar**, click the **Cinematics** drop-down option and select **Add Level Sequence**. 
+1.	With the **SK_Mannequin** asset still selected, click the **Add** button in Sequencer to add it to your Level Sequence. 
 
 	![](CharacterAnim_03.png)
 
-	When prompted, enter a save name and location for the new Level Sequence. 
-
-1.	Select the **SK_Mannequin** in the level.
-
-1.	In the **Sequencer** window, click the **Add** button and select **Add Actor to Sequencer** and **Add SK_Mannequin**. 
+1.	Click the **Animation** button on the new Skeletal Mesh Track and assign the **ThirdPersonWalk** Animation Sequence. 
 
 	![](CharacterAnim_04.png)
 
-	This will add a **Track** for the SK_Mannequin to Sequencer which we can then use to manipulate the Actor in our scene. 
-
-	![](CharacterAnim_04b.png)
-
-1.	Click the **+** on the **Animation** Sub-Track to bring up the content menu, then select **ThirdPersonIdle**.  
+1.	Expand **Animation**, then drag the **Timeline** to the **30** frame mark. 
 
 	![](CharacterAnim_05.png)
 
-	Once selected, the Actor will assume an idle position and an Animation will be added to the Animation Track. 
+1.	Click the **Animation** button again and add the **ThirdPersonJump_Start** animation. 
 
-	![](CharacterAnim_05b.png)
+	![](CharacterAnim_06.png)
 
-1.	Select the **Animation** track, then **Left-click** drag the timeline marker to the end of the idle animation and press **Enter**. 
+	If you play the animation now, you may notice a pop between the two animations. 
 
-	![](CharacterAnim_06.png)
+	[OBJECT:EmbeddedVideo]
+	[PARAMLITERAL:width]
+	640
+	[/PARAMLITERAL]
+	[PARAMLITERAL:height]
+	360
+	[/PARAMLITERAL]
+	[PARAMLITERAL:videoid]
+	oHjDpYPAdCo
+	[/PARAMLITERAL]
+	[/OBJECT]
 
-	This will add a new **Keyframe** at that point in the sequence and a context menu will appear where you can select another animation. 
+	[REGION:tip]
+	When adding animations, you can use the weight value to blend between multiple animations, which is how we will address the pop between animations. 
+	[/REGION]
 
-1.	Add the **ThirdPersonJump_Start** animation. 
+1.	Drag the **ThirdPersonJump_Start** animation below and under the walk animation, back to frame **25** of the sequence. 
 
 	![](CharacterAnim_07.png)
 
-	With our current setup, after the idle animation plays it will proceed to the jump start animation. You can preview this by dragging the timeline marker back-and-forth along the timeline or by pressing the **Play** button in the viewport to play our current scene. 
-
-1.	At the end of the jump start animation (with the Animation track selected), press **Enter** and add the **ThirdPersonJump_End**. 
+1.	Expand both sub-animation tracks and click the **add key** button on the walk track with a value of **1.0** and **0.0** for the jump start animation. 
 
 	![](CharacterAnim_08.png)
 
-1.	At the end of the jump end animation (with the Animation track selected), press **Enter** and add the **ThirdPersonIdle** animation.
+	[REGION:note]
+	This will turn off the playback of the jump start animation since it's weighted value is `0.0`. 
+	[/REGION]
+
+1.	Scrub to frame **30** and add keys for both animations of **0.0** and **1.0** for the walk and jump animations respectively. 
 
 	![](CharacterAnim_09.png)
 
+	[REGION:note]
+	This will blend in the jump start animation and blend out the walk animation. 
+	[/REGION]
+
+1.	Finally, click the **Preview Play** button to see the animation blending. 
 
 ## End Result
 
-When you click the **Play** button in the viewport, you will see your character jump up and down. 
+Your sequence now includes a Skeletal Mesh with multiple animations of which blend together smoothly instead of popping between animations. 
 
 [OBJECT:EmbeddedVideo]
 [PARAMLITERAL:width]
@@ -89,19 +102,12 @@ When you click the **Play** button in the viewport, you will see your character
 360
 [/PARAMLITERAL]
 [PARAMLITERAL:videoid]
-FMZG4yGriqc
+kVotxriT1Y4
 [/PARAMLITERAL]
 [/OBJECT]
 
-Above we are seeing what the cinematic sequence would look like from inside the Editor. If we were to play the game, we would not see anything as we have not issued a command to the Level Sequence itself that it should start playing. 
-Similar to a **Matinee** Actor, you can select the Level Sequence in your level then in the **Details** panel check the **Auto Play** checkbox to play the sequence when the level is loaded. 
-
-![](CharacterAnim_10.png)
-
-Also illustrated in the video is how playback is handled in Sequencer. Only what is contained in the **Section Bounds** will be played back (between the green and red markers). 
-You can drag those to increase/decrease the bounds of the sequence or through the **Sequencer** options, enable the **Keep Playback Range in Section Bounds** option (shown below) to ensure that the playback range includes all content. 
+If you want to continue working on this example, you can add the jump loop and jump end animations, blending them in so that the character jumps and lands. 
 
-![](CharacterAnim_11.png)
 
 
 
diff --git a/Engine/Documentation/Source/Engine/Sequencer/HowTo/ImportExportEDLs/ImportExportEDLs.INT.udn b/Engine/Documentation/Source/Engine/Sequencer/HowTo/ImportExportEDLs/ImportExportEDLs.INT.udn
--- a/Engine/Documentation/Source/Engine/Sequencer/HowTo/ImportExportEDLs/ImportExportEDLs.INT.udn
+++ b/Engine/Documentation/Source/Engine/Sequencer/HowTo/ImportExportEDLs/ImportExportEDLs.INT.udn
@@ -1,6 +1,6 @@
 Availability: Public
 Title:Importing & Exporting Edit Decision Lists (EDLs)
-Crumbs: %ROOT%, Engine, Engine/Sequencer, Engine/Sequencer/HowTo
+Crumbs:
 Description:Illustrates how you can import and export EDLs for use in external video editing software applications. 
 Version: 4.13
 SkillLevel: Advanced
@@ -9,24 +9,26 @@ order:1
 checkpoint: editorqs
 Related: Engine/Sequencer/Overview
 Related: Engine/Sequencer/ReferenceEditor
+tags:Sequencer
+type:how-to
 
 [EXCERPT:intro]
-Sequencer in Unreal Engine 4 not only allows you to render and export out your entire cinematic, but you can also export your cinematic split up into each of your shots along with an **Edit Decision List** or **EDL** which is a file that can be used with most video editing applications like Premiere, Avid or Final Cut. 
+Sequencer in Unreal Engine 4 (UE4) not only allows you to render and export your entire cinematic, but you can also export your cinematic split up into each of your shots along with an **Edit Decision List** (**EDL**), which is a file that can be used with most video editing applications like Premiere, Avid or Final Cut. 
 [/EXCERPT:intro]
 
 The EDL contains an ordered list of reel and timecode data representing where each video clip can be obtained in order to produce the final cut. 
-In this how-to we export a cinematic and EDL file, take it into an external program like Premiere, make edits to it and re-import the edited cinematic back into Unreal along with those changes. 
+In this how-to, we export a cinematic and EDL file, take it into an external program (like Premiere), make edits to it, and re-import the edited cinematic back into UE4 along with those changes. 
 
 ## Steps
 
 [REGION:note]
-For this how-to, we are using the **Blueprint Third Person Template** project and have created a sample cinematic however you can use your own project and cinematic sequence. 
-We are also using Adobe Premiere Pro CS6 as our external editing software; however you can use any software that supports EDL importing. 
+For this how-to, we are using the **Blueprint Third Person Template** project, having created a sample cinematic; however, you can use your own project and cinematic sequence. 
+We are also using Adobe Premiere Pro CS6 as our external editing software; however, you can use any software that supports EDL importing. 
 [/REGION]
 
-1.	Make sure your cinematic you want to export includes a **Shot Track** containing your shots, then click the **Render Movie** button.
+1.	Make sure the cinematic you want to export includes a **Shot Track** containing your shots, then click the **Render Movie** button.
 
-1.	Define an **Output Directory** and **Filename Format** using the **{shot}** token and make sure **Write Edit Decision List** is checked. Here we show the sample cinematic that we want to export along with an EDL and click the **Render Movie** button to bring up the settings dialog window. Inside the settings window, we create a new folder to export the cinematic to and prefix our {shot} token with the name of our cinematic. By using the {shot} token each of our shots will be exported as individual video files along with the full cinematic.
+1.	Set the **Output Format** to **Video Sequence**, choose an **Output Directory** and **Filename Format** using the **{shot}** token, then enable **Write Edit Decision List**. 
 
 	[OBJECT:EmbeddedVideo]
 	[PARAMLITERAL:width]
@@ -40,13 +42,13 @@ We are also using Adobe Premiere Pro CS6 as our external editing software; howev
 	[/PARAMLITERAL]
 	[/OBJECT] 
 
+	Here, we show the sample cinematic that we want to export (along with an EDL), and click the **Render Movie** button to bring up the settings dialog window. Inside the settings window, we create a new folder to export the cinematic to, and prefix our **{shot}** token with the name of our cinematic. By using the **{shot}** token, each of our shots will be exported as individual video files along with the full cinematic.
+
 	[REGION:note]
-	This will only work when the **Output Format** is set to **Video Sequence**. It does not work with individual frames. 
+	You must set the **Output Format** to **Video Sequence** in order for this to work, it does not work with individual frames. 
 	[/REGION]
 
-
-1.	Inside your software editing suite, **Import** the **EDL** file exported from the previous step. In Premiere Pro, you can **Right-click** in the window in the lower-left and choose **Import** then point to the EDL file. 
-	We left the EDL Information and New Sequence dialog windows with the default information. 
+1.	**Import** the **EDL** file exported from the previous step into your external video editing software. 
 
 	[OBJECT:EmbeddedVideo]
 	[PARAMLITERAL:width]
@@ -60,7 +62,9 @@ We are also using Adobe Premiere Pro CS6 as our external editing software; howev
 	[/PARAMLITERAL]
 	[/OBJECT]
 
-1.	Next you'll need to link the data from the EDL to the exported media files in your editing software. In Premiere Pro, **Right-click** the media folder and select **Link Media** then point each file to its corresponding media file. 
+	In Premiere Pro, you can right-click in the lower left of the window, choose **Import**, then point to the EDL file. 
+	
+1.	Link the data from the EDL to the exported media files in your external video editing software. 
 
 	[OBJECT:EmbeddedVideo]
 	[PARAMLITERAL:width]
@@ -73,8 +77,10 @@ We are also using Adobe Premiere Pro CS6 as our external editing software; howev
 	iUJenhTwFRk
 	[/PARAMLITERAL]
 	[/OBJECT]
+	
+	In Premiere Pro, you can right-click the media folder, select **Link Media**, then point each file to its corresponding media file. 
 
-1.	Modify the shots in your editing software as desired, then **Export** as an **EDL**. For our sample scene, we adjust the length of the shots as well as re-arrange their position in the sequence. 
+1.	Modify the shots in your external video editing software as desired, then **Export** as an **EDL**. 
 
 	[OBJECT:EmbeddedVideo]
 	[PARAMLITERAL:width]
@@ -88,7 +94,9 @@ We are also using Adobe Premiere Pro CS6 as our external editing software; howev
 	[/PARAMLITERAL]
 	[/OBJECT]
 
-1.	Return to **Unreal Engine 4** and in your Sequence, **Right-click** your **Shot Track** and **Import EDL**. The changes we made externally are now reflected inside our Sequence in Unreal Engine 4. 
+	For our sample scene, we adjust the length of the shots as well as re-arrange their position in the sequence. 
+
+1.	Returning to your Sequence in UE4, right-click your **Shot Track** and **Import EDL**. 
 
 	[OBJECT:EmbeddedVideo]
 	[PARAMLITERAL:width]
@@ -102,22 +110,16 @@ We are also using Adobe Premiere Pro CS6 as our external editing software; howev
 	[/PARAMLITERAL]
 	[/OBJECT]
 
+	The changes we made externally are now reflected inside our Sequence in UE4. 
+
 ## End Result
 
 The import/export EDL option within Sequencer gives you the flexibility to determine where you want to do your editing. 
+You can do most of your setup inside of Sequencer, then pass your sequence off to an editor that works inside an external application (like Premiere) to perform the final touches, who can then pass it back to you to tweak in Sequencer. 
+By exporting your sequence as an EDL and working externally, you have the ability to add post production effects (or after effects) to your scene before rendering out the final cinematic. 
 
-You can do most of your setup inside Sequencer, then pass your sequence off to an editor who works inside an external application like Premiere to perform the final touches who can then pass it back to you to tweak in Sequencer. 
-
-By exporting your sequence as an EDL and working externally, you have the ability to add post production effects or after effects to your scene before you rendering out the final cinematic. 
-
-
-
-
-
-
-
-
-
-
+Optionally, when rendering out your sequence and exporting an EDL, you can add **Frame Handles** automatically to each shot, and define the number of frames to include. 
 
+![](HandleFrames.png)
 
+These extra frames pad each shot and are cut into and out of by the EDL, which can be used in your external video editing package to adjust the cuts between shots. 
diff --git a/Engine/Documentation/Source/Engine/Sequencer/HowTo/SequenceRecorder/SequenceRecorder.INT.udn b/Engine/Documentation/Source/Engine/Sequencer/HowTo/SequenceRecorder/SequenceRecorder.INT.udn
--- a/Engine/Documentation/Source/Engine/Sequencer/HowTo/SequenceRecorder/SequenceRecorder.INT.udn
+++ b/Engine/Documentation/Source/Engine/Sequencer/HowTo/SequenceRecorder/SequenceRecorder.INT.udn
@@ -1,34 +1,35 @@
 Availability: Public
 Title:Recording to Sequencer
-Crumbs: %ROOT%, Engine, Engine/Sequencer, Engine/Sequencer/HowTo
-Description:Record your actions with the Sequence Recorder and use that data inside Sequencer.
-Version: 4.12
+Crumbs:
+Description:Record your actions with the Sequence Recorder, and use that data inside of Sequencer.
+Version: 4.14
 SkillLevel: Advanced
 parent:Engine/Sequencer/HowTo
 order:1
 checkpoint: editorqs
 Related: Engine/Sequencer/Overview
-Related: Engine\Sequencer\ReferenceEditor\
+Related: Engine/Sequencer/ReferenceEditor
+tags:Sequencer
+type:how-to
 
-The **Sequence Recorder** allows you to capture specified Actors during level editing or gameplay that can be saved as a new **Level Sequence** and can be edited inside **Sequencer**. 
-This is useful for quickly capturing content for scenes as you can take a playable character, perform some actions during gameplay while recording with the Sequence Recorder then take that data into Sequencer and create a cinematic around it. 
-
-In this example, we will use the Sequence Recorder to record our playable character's movement which we can then edit. 
+The **Sequence Recorder** allows you to capture specified Actors during level editing or gameplay that can be saved as a new **Level Sequence**, and can be edited inside of **Sequencer**. 
+This is useful for quickly capturing content for scenes, as you can take a playable character, perform some actions during gameplay while recording with the Sequence Recorder, then take that data into Sequencer, creating a cinematic around it.
 
+In this example, we will use the Sequence Recorder to record our playable character's movement, which we can then edit.
 
 ## Steps
 
 [REGION:note]
-For this how-to, we are using the **Blueprint Third Person Template** project with **Starter Content Enabled**. 
+For this How-to, we are using the **Blueprint Third Person Template** project with **Starter Content Enabled**. 
 [/REGION]
 
 1.	From the **Menu Bar** under **Window**, select **Sequence Recorder**. 
 
 	![](RecorderStep1.png)
 
-	The Sequence Recorder window will automatically open. There are some options under **Sequence Recording** which will determine how (and where) the new Level Sequence asset will be saved. You can choose to record Actors that are spawned (such as particle effects, other characters, etc.), determine the naming convention and save location, and change other settings like the duration to record or delay before starting to record. 
+	The Sequence Recorder window will automatically open. There are some options under **Sequence Recording** which will determine how (and where) the new Level Sequence asset will be saved. You can choose to record Actors that are spawned (such as particle effects, other characters, etc.), determine the naming convention and save location, and change other settings like the duration to record or delay before starting to record.
 	
-1.	From the Main Editor Toolbar, click the **Play** button to start a Play in Editor session. 
+1.	From the Main Editor Toolbar, click the **Play** button to start a Play in Editor (PIE) session. 
 
 1.	Press **Shift+F1** to gain mouse control once inside the game. 
 
@@ -36,11 +37,19 @@ For this how-to, we are using the **Blueprint Third Person Template** project wi
 	
 	![](SequenceRecorderAdd.png)
 
-1.	Click on the new recording, then for **Actor to Record** click the drop-down and select **ThirdPersonCharacter**. This is where we specify which Actor to target recording around before starting the recording process. 
+1.	Click on the new recording (which will say **None**), then for **Actor to Record** click the drop-down and select **ThirdPersonCharacter**. 
 
 	![](ActorToRecord.png)
 
-1.	Click the **Record** button. After 4 seconds (which is the **Record Delay** option under the Sequence Recording section), the record process will start. 
+	This is where we specify which Actor to target before starting the recording process, updating the UI as shown above.
+
+1.	Optionally, you can choose to record audio and set the audio gain levels along with your recorded clip. 
+
+	![](RecordAudio.png)
+
+	Audio recording requires an attached microphone and will start recording when the sequence starts recording. 
+
+1.	Click the **Record** button. After 4 seconds (which is the **Record Delay** option under the Sequence Recording section), the recording process will start. 
 
 	![](RecordStart.png)
 
@@ -48,7 +57,7 @@ For this how-to, we are using the **Blueprint Third Person Template** project wi
 	When clicking the **Record** button, all Actors in the list that are set to be tracked will also be recorded. 
 	[/REGION]
 
-1.	Use the **WASD** and **Spacebar** keys to move the character around and jump, then press **ESC** when done. Below is the sequence we recorded. However, if you have a character that can attack or perform other actions, feel free to use that.
+1.	Use the **WASD** and **Spacebar** keys to move the character around and jump, then press **ESC** when done.
 
 	[OBJECT:EmbeddedVideo]
 	[PARAMLITERAL:width]
@@ -58,18 +67,17 @@ For this how-to, we are using the **Blueprint Third Person Template** project wi
 	360
 	[/PARAMLITERAL]
 	[PARAMLITERAL:videoid]
-	xo5JRkV3_nI
+	9mZI75qBsx8
 	[/PARAMLITERAL]
 	[/OBJECT] 
 
-1.	Inside the **Content Browser** a new folder will be created containing assets related to the recorded sequence. When recording is complete, a **Level Sequence** asset containing your cniematic data as well as an **Animation Sequence** containing your animation data will be created.
+	Inside the **Content Browser**, a new folder will be created containing assets related to the recorded sequence. 
 
-	![](RecordedAssets.png) 
-	 
+ 
 
 ## End Result
 
-You can open the **RecordedSequence** asset and begin editing it as you would a normal Level Sequence. Below is our recorded sequence to which we could add cameras and a [Camera Cuts](Engine/Sequencer/HowTo/TracksCameraCut) track to provide multiple angles, music, effects or anything else we wanted. We could even take this sequence and embed it in another sequences as part of a [Shots Track](Engine/Sequencer/HowTo/TracksShot). 
+You can open the **RecordedSequence** asset and begin editing it as you would a normal Level Sequence. Below is our recorded sequence, to which we could add cameras, and a [Camera Cuts](Engine/Sequencer/HowTo/TracksCameraCut) track to provide multiple angles, music, effects, or anything else we'd like to add. We could even take this sequence and embed it in other sequences as part of a [Shots Track](Engine/Sequencer/HowTo/TracksShot). 
 
 [OBJECT:EmbeddedVideo]
 [PARAMLITERAL:width]
@@ -79,7 +87,7 @@ You can open the **RecordedSequence** asset and begin editing it as you would a
 360
 [/PARAMLITERAL]
 [PARAMLITERAL:videoid]
-hR5zBjYO5Oc
+1nFX_bpWfy0
 [/PARAMLITERAL]
 [/OBJECT]
 
@@ -93,18 +101,14 @@ In addition to recording gameplay, you can record your actions during level edit
 360
 [/PARAMLITERAL]
 [PARAMLITERAL:videoid]
-rC_ZStHXm38
+fMzumqLk7IQ
 [/PARAMLITERAL]
 [/OBJECT]
 
-Above we have placed a cube in our level and instructed the Sequence Recorder to record the cube. 
+Above, we have placed a cube in our level and instructed the Sequence Recorder to record the cube. 
 We then moved the cube around and the Sequence Recorder captured the movements we entered through keyframes in a newly created Level Sequence. 
-When we playback our Level Sequence, a new Cube Actor is created as a spawnable in the Level Sequence which is why a second Cube appears when the sequence is active.
+When we play back our Level Sequence, a new Cube Actor is created (as a spawnable) in the Level Sequence, which is why a second Cube appears when the sequence is active.
 
 [REGION:note]
-Only properties that can be keyframed can be captured and recorded when Level Editing recording. 
+Only properties that can be keyframed, can be captured and recorded when Level Editing recording. 
 [/REGION]
-
-
-
-
diff --git a/Engine/Documentation/Source/Engine/Sequencer/HowTo/SequencerHowTo.INT.udn b/Engine/Documentation/Source/Engine/Sequencer/HowTo/SequencerHowTo.INT.udn
--- a/Engine/Documentation/Source/Engine/Sequencer/HowTo/SequencerHowTo.INT.udn
+++ b/Engine/Documentation/Source/Engine/Sequencer/HowTo/SequencerHowTo.INT.udn
@@ -4,6 +4,7 @@ Crumbs: %ROOT%, Engine, Engine\Sequencer
 Description:A How To Guide for creating cinematic sequences with the Sequencer Editor. 
 Related: Engine\Sequencer
 version: 4.11
+tags:Sequencer
 
 
 The **Sequencer - How To...** pages provide step-by-step examples on how to work with the **Sequencer** Editor in various capacties inside Unreal Engine 4. 
@@ -37,14 +38,13 @@ Refer to the table below for more information and click each link to view its co
 
 ## Advanced
 [REGION:wrapper transparent]
+[](FANCY:Engine\Sequencer\HowTo\AnimateDynamicObjects "%Engine\Sequencer\HowTo\AnimateDynamicObjects:description%")
 [](FANCY:Engine\Sequencer\HowTo\TracksSub "%Engine\Sequencer\HowTo\TracksSub:description%")
 [](FANCY:Engine\Sequencer\HowTo\Spawnables "%Engine\Sequencer\HowTo\Spawnables:description%")
 [](FANCY:Engine\Sequencer\HowTo\SequenceRecorder "%Engine\Sequencer\HowTo\SequenceRecorder:description%")
 [](FANCY:Engine\Sequencer\HowTo\TracksLevelVisibility "%Engine\Sequencer\HowTo\TracksLevelVisibility:description%")
 [](FANCY:Engine\Sequencer\HowTo\CustomBurnIns "%Engine\Sequencer\HowTo\CustomBurnIns:description%")
-[PUBLISH:docs]
 [](FANCY:Engine\Sequencer\HowTo\ImportExportEDLs "%Engine\Sequencer\HowTo\ImportExportEDLs:description%")
-[/PUBLISH:docs]
 [/REGION]
 
 
