Copying //UE4/Dev-Documentation to //UE4/Samples-Main (Source: //UE4/Dev-Documentation @ 3372893)

diff --git a/Engine/Documentation/Source/Engine/Sequencer/HowTo/TracksAudio/Audio.INT.udn b/Engine/Documentation/Source/Engine/Sequencer/HowTo/TracksAudio/Audio.INT.udn
--- a/Engine/Documentation/Source/Engine/Sequencer/HowTo/TracksAudio/Audio.INT.udn
+++ b/Engine/Documentation/Source/Engine/Sequencer/HowTo/TracksAudio/Audio.INT.udn
@@ -1,23 +1,20 @@
 Availability: Public
-Title:Adding Audio	
+Title:Working with Audio Tracks
 Crumbs: %ROOT%, Engine, Engine/Sequencer, Engine/Sequencer/HowTo
-Description:Example of how to add audio to your scenes.
-Version: 4.12
+Description:Illustrates how to add Audio Tracks and adjust volume/pitch values through keyframes. 
+Version: 4.15
 SkillLevel: Beginner
 parent:Engine/Sequencer/HowTo
 order:1
 checkpoint: editorqs
 Related: Engine/Sequencer/Overview
 Related: Engine\Sequencer\ReferenceEditor\
+tags:Sequencer
 
-Whether your are looking to score your cinematic with a musical track or want to add sound effects such as impact sounds, voice overs, or others, you'll need to do so with an **Audio Track** inside Sequencer. 
-
-In this how-to, we add a music track to a sample scene and show how you can fade in/out an audio track with a **Sound Cue**.
-
-[REGION:warning]
-Currently **Audio Tracks** have limited editing capabilities. Additional support will come online in future engine updates. 
-[/REGION]
+Whether you are looking to score your cinematic with a musical track or want to add sound effects such as impact sounds, voice overs, or others, you'll need to do so with an **Audio Track** inside Sequencer. 
+Sequencer's Audio Track also allows you to adjust the Volume or Pitch of your audio assets at any point during your cinematic through keyframed value changes. 
 
+In this how-to, we add a music track to a sample scene and fade it in over time as well as adjust the pitch during the cinematic. 
 
 ## Steps
 
@@ -29,68 +26,50 @@ For this how-to, we are using the **Blueprint Third Person Template** with **Sta
 
 	![](Audio01.png)
 
-1.	Inside the Level Sequence, click the **Add** button and select the **Audio Track**. 
+1.	In the **Save Asset As** window, choose a save location and save name then click **Save**. 
 
-	![](Audio02.png)
+	![](Audio02.png)	
 
-1.	Click the **+ Audio** button then select the **Starter_Music_Cue** asset. 
+1.	In the **Level Sequence**, click the **Add** button then select **Audio Track**. 
 
 	![](Audio03.png)
 
-1.	Click the **+ Audio** button and select additional audio assets and arrange on the timeline as desired. 
+1.	On the **Audio Track**, click the **+** button then search for and add the **Starter_Music_Cue** asset. 
 
 	![](Audio04.png)
 
-	You can layer sounds by arranging one audio asset on top of another. 
-
-1.	**Right-click** on the **Starter_Music_Cue** asset to reveal the **Properties**, then click the magnifying glass icon. 
+	You can click the **Expand** icon to expand out and see your audio sub-tracks. 
 
 	![](Audio05.png)
 
-	Along with the properties that you can change for the asset, you can browse to and find the asset in the **Content Browser** by clicking the **Browse To** button. 
-
-1.	**Double-click** and open the **Starter_Music_Cue** Sound Cue asset. 
-
-1.	In the **Sound Cue Editor**, drag an **Enveloper** into the graph. 
+	You can add additional sub-tracks to the Audio Track by clicking the **+** button and can drag them along the timeline to reposition them where you want in your sequence. 
 
 	![](Audio06.png)
 
-	Inside the Sound Cue Editor we can alter and mix sound assets. For this example we want to fade in the sound which we can do with the **Enveloper** node. 
-
-	[REGION:note]
-	In future engine releases, the ability to edit Audio Tracks within Sequencer will come online and this guide will be updated. 
-	[/REGION]
-
-1.	**Left-click** and connect the nodes as shown below. 
+1.	Expand the **Starter_Music_Cue** track, then change the **Volume** and **Pitch Multiplier** values to **0.0** and add keys for both. 
 
 	![](Audio07.png)
 
-	The Starter_Music01 asset will get passed through the Enveloper before outputting the final sound. 
+	This will turn the Volume to off and the Pitch to low. 
 
-1.	In the **Details** panel of the **Enveloper**, click the key in the **Volume Curve** and change the **Value** to **0.0**. 
+1.	Scrub the Timeline marker ahead to frame **50**, change the **Volume** and **Pitch Multiplier** values to **1.0** and add keys for both. 
 
 	![](Audio08.png)
 
-	Here we are stating that at time 0.0 (the start), the value of Volume should be 0.0 (or off). 
-
-1.	Hold **Shift** and **Left-click** in the curve window to create a new key. 
+	This will turn the Volume up to its default value (on) and the Pitch to its default value (not affecting the sound asset). 
 
-1.	Select the key, then change the **Time** to **4.0** and **Value** to **1.0**. 
+1.	You can delete sub-tracks within the Audio Track by clicking on them in the Timeline window, then pressing **Delete** or **Right-click** and **Delete**. 
 
 	![](Audio09.png)
 
-	Here we are stating that at 4 seconds, the value of the Volume should be 1.0 (or full volume) giving us a fade in. You can also click the two buttons next to the **Time** and **Value** properties to adjust the size of the curve window to display the curve. 
+1.	Return the Timeline to **0** then click the **Play** button inside Sequencer to play your sequence. 
 
-	[REGION:tip]
-	If you wanted to fade out for example starting at 10 seconds, you would add a key and set the **Time** to **10.0** with **Value** being **1.0**. If you wanted the fade duration to last 5 seconds, you would add another key with the **Time** being **15.0** and the **Value** being **0.0**. 
-	[/REGION]
-
-1.	Return to **Sequencer** and click the **Preview Play** button. 
+	![](Audio10.png)
 
 
 ## End Result
 
-You should now hear all your sounds that you've placed in the **Audio Track** and the music should fade in over time. 
+When you play the Level Sequence, you will hear the audio fade in and the pitch adjusts as the sequence begins. 
 
 [OBJECT:EmbeddedVideo]
 [PARAMLITERAL:width]
@@ -100,11 +79,11 @@ You should now hear all your sounds that you've placed in the **Audio Track** an
 360
 [/PARAMLITERAL]
 [PARAMLITERAL:videoid]
-1STSDGBc5KM
+Tt_fKZnk8I4
 [/PARAMLITERAL]
 [/OBJECT]
 
-
+While this example is fairly simplified, you can add more keyframes to fade in/out audio at any point during your sequence or adjust the pitch of audio assets as desired. 
 Audio Tracks support either **Sound Cue** or **Sound Wave** assets and can be added by dragging-and-dropping from the **Content Browser** into a Level Sequence or by adding through the **Add** option of an Audio Track. 
 
 
@@ -131,5 +110,6 @@ Audio Tracks support either **Sound Cue** or **Sound Wave** assets and can be ad
 
 
 
+
 
 
diff --git a/Engine/Documentation/Source/Engine/Sequencer/HowTo/TracksEvent/Event.INT.udn b/Engine/Documentation/Source/Engine/Sequencer/HowTo/TracksEvent/Event.INT.udn
--- a/Engine/Documentation/Source/Engine/Sequencer/HowTo/TracksEvent/Event.INT.udn
+++ b/Engine/Documentation/Source/Engine/Sequencer/HowTo/TracksEvent/Event.INT.udn
@@ -1,19 +1,36 @@
 Availability: Public
-Title:Calling Custom Events
-Crumbs: %ROOT%, Engine, Engine/Sequencer, Engine/Sequencer/HowTo
-Description:Shows how you can use the Event Track to call Custom Events and script in Blueprint. 
-Version: 4.12
+Title:Using Custom Events and Structs
+Crumbs:
+Description:Shows how you can use the Event Track to call Custom Events with Custom Structs for use in Blueprint. 
+Version: 4.15
 SkillLevel: Intermediate
 parent:Engine/Sequencer/HowTo
 order:1
 checkpoint: editorqs
 Related: Engine/Sequencer/Overview
-Related: Engine\Sequencer\ReferenceEditor\
+Related: Engine/Sequencer/ReferenceEditor
+tags:Sequencer
 
-The **Event Track** can be used to call [**Custom Events**](Engine/Blueprints/UserGuide/Events/Custom) and fire off [**Blueprint Script**](Engine/Blueprints) inside the  [**Level Blueprint**](Engine/Blueprints/UserGuide/Types/LevelBlueprint). 
+The **Event Track** can be used to call [Custom Events](Engine/Blueprints/UserGuide/Events/Custom) and fire off [Blueprint Script](Engine/Blueprints) inside the  [Level Blueprint](Engine/Blueprints/UserGuide/Types/LevelBlueprint). 
 Once the Custom Event is called, any script attached to it will also be executed and can be used for any number of scenarios based on your needs. 
+When using the Event Track and calling a Custom Event, you can also pass through your own Custom [Struct](Engine/Blueprints/UserGuide/Variables/Structs) which can be used to update the variables contained within. 
+This can be useful in situations where you want to update other variables during the playback of a level sequence. 
 
-In this example, we will execute an Event which will be used to activate a particle effect in the level. 
+In this example, we will use the Event Track to call a Custom Event with a Custom Struct that is used to move an object in the level when the sequence plays. 
+
+[OBJECT:EmbeddedVideo]
+[PARAMLITERAL:width]
+640
+[/PARAMLITERAL]
+[PARAMLITERAL:height]
+360
+[/PARAMLITERAL]
+[PARAMLITERAL:videoid]
+wlLcUr8BXvY
+[/PARAMLITERAL]
+[/OBJECT]
+
+In the example above, when we press a key, the level sequence plays and our event to move the cube upwards is executed each time. 
 
 ## Steps
 
@@ -21,74 +38,92 @@ In this example, we will execute an Event which will be used to activate a parti
 For this how-to, we are using the **Blueprint Third Person Template** project with **Starter Content** enabled. 
 [/REGION]
 
-1.	From the **Main Toolbar** click the **Cinematics** button and **Add Level Sequence**. 
+1.	From the **Content Browser** under the **Content/StarterContent/Shapes** folder, drag any mesh into the level. 
 
 	![](EventTrack1.png)
 
-	Choose a Save Location and give the Level Sequence a name when prompted. 
+	Position the mesh anywhere inside the level. 
+
+1.	In the **Details** panel for your mesh, set the mesh to **Movable**. 
+
+	![](EventTrack1b.png)
+
+	This will allow us to move the mesh during gameplay.
 
-1.	In Sequencer, click the **Add** button and select to add an **Event Track**. 
+1.	Click the **Add New** button, then under **Blueprints**, select the **Structure** option. 
 
 	![](EventTrack2.png)
 
-1.	On the **Event Track**, move the Timeline Marker to **75** then click the **+** sign to add a key. 
+1.	Enter the name **MyStruct**, then open the asset and add a **Vector** variable called **MoveCube**. 
 
 	![](EventTrack3.png)
 
-1.	**Right-click** on the key, then under **Properties** enter the name **MyEvent** for the **Value**. 
+	Naming of your variables is important here as these names will need to be used when you want to reference them later. 
 
-	![](EventTrack4.png)
+1.	Add a new **Level Sequence** from the **Cinematics** button called **Master** to the level, then inside **Sequencer** add an **Event Track**. 
 
-	The name you provide here is the name you **must** use when creating your **Custom Event** that you want to call. 
+	![](EventTrack4.png)
 
-1.	In the **Content Browser** under **StarterContent/Particles**, drag the **P_Explosion** into the level. 
+1.	Move the timeline to frame **5** and add a key, then **Right-click** the key and name the event **LiftCube**, assign **MyStruct** and set the **Z** value to **100**. 
 
 	![](EventTrack5.png)
 
-1.	With the **P_Explosion** asset selected, in the **Details** panel un-check the **Auto Activate** option. 
+	When we get to this keyframe in our sequence, we are going to move our mesh up by 100 units along the Z axis. 
 
-	![](EventTrack6.png)
+1.	Move the timeline to frame **10** and press the **]** key to end the sequence at frame 10. 
 
-	This will ensure that the particle effect is not played upon starting the level. 
+	![](EventTrack6.png)
 
-1.	With the **P_Explosion** asset still selected, click the **Blueprints** button from the Toolbar and select **Open Level Blueprint**. 
+1.	Select your mesh that you added and the Level Sequence in the level, then click the **Blueprints** button from the toolbar and select **Open Level Blueprint**. 
 
 	![](EventTrack7.png)
 
-1.	In the **Level Blueprint**, **Right-click** and select **Create a Reference to P_Explosion**. 
+1.	In the **Level Blueprint**, **Right-click** and search for **Reference** then select **Create Reference to 2 selected Actors**. 
 
 	![](EventTrack8.png)
 
-1.	**Right-click** in the graph again and search for and add a **Custom Event** node and name it **MyEvent** (the same as in **Step 4**). 
+1.	**Right-click** in the graph and add a **1** Keyboard Event, then drag off your **Master** Level Sequence and search for and add the **Set Playback Position** node. 
 
 	![](EventTrack9.png)
-	
-	[REGION:note]
-	This must be the text you entered as the value in Step 4. You can press the **F2** key on the node to rename it. 
-	[/REGION]
 
-1.	**Left-click** drag off the **P_Explosion** node and search for and add the **Activate (ParticleSystemComponent)** node. 
+	This will be used to reset our sequence whenever we start to play it before we actually issue a play function call. 	
+	
+1.	Drag off the **Sequence Player** node that is added and use the **Play** node and connect as shown below. 
 
 	![](EventTrack10.png)
 
-1.	Connect the **Custom Event** to the **Activate** node as shown below. 
+1.	**Right-click** in the graph and create a **Custom Event** called **LiftCube**. 
 
 	![](EventTrack11.png)
 
-	When the event is called it will activate (turn on) the particle effect we placed in the level.
+	[REGION:warning]
+	Your Custom Event name must be the same name that is specified under **Event Name** for your key in Sequencer or it will not get called. 
+	[/REGION]
 
-1.	In the Main Editor window from the **World Outliner** select your Level Sequence then in the **Details** enable **Auto Play**. 
+1.	In the **Details** panel for the Custom Event, click the **+** button to add an input and name it **MoveCube** with the variable type of **Vector**. 
 
 	![](EventTrack12.png)
 
-	When we play in the level, this Level Sequence will automatically play and trigger our effect. 
+	[REGION:warning]
+	Your Input names must be the same names that are used in your Struct or the event will not get called. 
+	Also, the number of inputs must match the number of variables defined in your custom struct or the event will not get called. 
+	For example, if you struct has 2 variables in it, your Custom Event must have 2 inputs with the correct variable names and types. 
+	[/REGION]
+	
+1.	Drag off your mesh reference and search for and use the **AddActorLocalOffset** node. 
+
+	![](EventTrack13.png)
+
+1.	Connect the **MoveCube** vector pin to the **Delta Location** pin and the **LiftCube** executable to the **AddActorLocalOffset** node as shown below. 
+
+	![](EventTrack14.png)
 
 1.	Click the **Play** button from the Main Toolbar to play in the editor. 
 
 
 ## End Result
 
-When playing in the editor after a few seconds, the explosion will activate.
+When playing in the editor, each time you press the **1** key, the cube will move up along the **Z** axis. 
 
 [OBJECT:EmbeddedVideo]
 [PARAMLITERAL:width]
@@ -98,12 +133,13 @@ When playing in the editor after a few seconds, the explosion will activate.
 360
 [/PARAMLITERAL]
 [PARAMLITERAL:videoid]
-krFTHenh9zw
+wlLcUr8BXvY
 [/PARAMLITERAL]
 [/OBJECT]
 
-In the video above, we also illustrate that you will not see the script being executed when you click the **Play** button inside Sequencer.
-You will only be able to preview your scripted events being called from your Event Track during gameplay when the Level Sequence is played.
+In this example we are not assigning a camera to view our cinematic from but instead are using Sequencer and Blueprint to perform scripted sequences. 
+In your projects, an Event Track with a Custom Event and Custom Struct can be used to update player information during a cutscene (such as the player receiving a health increase, new weapon, or other gameplay related effects of your choosing). 
+
 
 
 
diff --git a/Engine/Documentation/Source/Engine/Sequencer/HowTo/TracksShot/Shot.INT.udn b/Engine/Documentation/Source/Engine/Sequencer/HowTo/TracksShot/Shot.INT.udn
--- a/Engine/Documentation/Source/Engine/Sequencer/HowTo/TracksShot/Shot.INT.udn
+++ b/Engine/Documentation/Source/Engine/Sequencer/HowTo/TracksShot/Shot.INT.udn
@@ -1,6 +1,6 @@
 Availability: Public
-Title:Creating Shots & Takes
-Crumbs: %ROOT%, Engine, Engine/Sequencer, Engine/Sequencer/HowTo
+Title:Creating Shots and Takes
+Crumbs:
 Description:Shows a form of non-linear editing by arranging Shots and variations of Shots with Takes. 
 Version: 4.12
 SkillLevel: Intermediate
@@ -8,7 +8,8 @@ parent:Engine/Sequencer/HowTo
 order:1
 checkpoint: editorqs
 Related: Engine/Sequencer/Overview
-Related: Engine\Sequencer\ReferenceEditor\
+Related: Engine/Sequencer/ReferenceEditor
+tags:Sequencer
 
 Similar to traditional film editors, the **Shot Track** allows you to add, arrange and re-arrange Level Sequences that will be played back in the order you arrange them.
 Each Level Sequence you add is referred to as a **Shot** and each Shot can have multiple **Takes**. With Takes, you can quickly swap the shot out with an alternate version giving you the ability to quickly experiment with different looks for your cinematic. 
diff --git a/Engine/Documentation/Source/Engine/Sequencer/Overview/SequencerOverview.INT.udn b/Engine/Documentation/Source/Engine/Sequencer/Overview/SequencerOverview.INT.udn
--- a/Engine/Documentation/Source/Engine/Sequencer/Overview/SequencerOverview.INT.udn
+++ b/Engine/Documentation/Source/Engine/Sequencer/Overview/SequencerOverview.INT.udn
@@ -1,56 +1,58 @@
 Availability:Public
 Title: Sequencer Overview
-Crumbs: %ROOT%, Engine
-Description:Sequencer is Unreal Engine 4's multi-track editor used for creating and previewing cinematic sequences in real time. 
+Crumbs: 
+Description:Sequencer is Unreal Engine 4's multi-track editor, used for creating and previewing cinematic sequences in real time. 
 Navigation:topic
 parent:Engine/Sequencer
 order:1
-Version: 4.12
-Related: Engine\Sequencer\ReferenceEditor
-Related: Engine\Sequencer\HowTo
-
-The **Sequencer** Editor gives users the ability to create in-game cinematics through a specialized multi-track editor similar to [Matinee](Engine\Matinee\). 
-By creating **Level Sequences** and adding **Tracks**, users can define the make-up of each Track which will determine the content for their scene.
-Tracks can consist of things like Animation (for animating a character), Transformation (moving things around in the scene like Camera), Audio (for including a music or sound effect), as well as several other Track types.  
+Version: 4.15
+Related: Engine/Sequencer/ReferenceEditor
+Related: Engine/Sequencer/HowTo
+tags:Sequencer
+type:overview
 
 [TOC(start:2 end:2)]
 
+The **Sequencer** Editor gives users the ability to create in-game cinematics through a specialized multi-track editor (similar to [Matinee](Engine\Matinee\)). 
+By creating **Level Sequences** and adding **Tracks**, users can define the makeup of each Track, which will determine the content for the scene.
+Tracks can consist of things like Animations (for animating a character), Transformations (moving things around in the scene), Audio (for including music or sound effects), as well as several other Track types.  
+
 ## Creating Level Sequences
 
-The **Level Sequence** is the "container" for your cinematic scenes and must be created in order to begin working inside the Sequencer Editor. You can create a Level Sequence and directly place it inside of your level by choosing to create one from the **Toolbar** under **Cinematics** as seen below. 
+The **Level Sequence** is the "container" for your cinematic scenes, and must be created in order to begin working inside of the Sequencer Editor. You can create a Level Sequence and directly place it inside of your level by choosing to create one from the **Toolbar** under **Cinematics** (as seen below). 
 
 ![](LevelSequence.png)
 
-This will add it to the level at which point it can be selected and its properties can be manipulated in the **Details** panel similar to a **Matinee Actor**. In the Details panel (below), you can define whether the Level Sequence will automatically play upon level start, the Play Rate for the sequence and whether the sequence should loop and other settings. 
+This will add it to the level, at which point it can be selected and its properties can be manipulated in the **Details** panel, similar to a **Matinee Actor**. In the Details panel (below), you can define whether the Level Sequence will automatically play upon level start, the Play Rate for the sequence, whether the sequence should loop, and other settings. 
 
 ![](DetailsPanel.png)
 
-How this differs from Matinee is that Level Sequences are self-contained assets and you can embed a Level Sequence within another Level Sequence. 
+How this differs from Matinee, is that Level Sequences are self-contained assets, and you can embed a Level Sequence within another Level Sequence.
 For example, you can create a Level Sequence that has animated characters, cameras, etc. as one scene that is part of a larger cinematic sequence.
 
-The second method of creating Level Sequences can be performed in the **Content Browser** by clicking the **Add New** button and selecting **Level Sequence** under _Animation_. When doing so in this manor you are creating the Level Sequence asset but haven't explicitly placed it in the level. 
+An alternate method of creating Level Sequences can be performed in the **Content Browser** by clicking the **Add New** button and selecting **Level Sequence** under _Animation_. When doing so in this manner, you are creating the Level Sequence asset, but haven't placed it in a level.
 
 ![](LevelSequence_B.png)
 
 
 ## Adding Tracks to Sequencer
 
-With a Level Sequence created and the Sequencer Editor open, you can begin creating your cinematic. 
+After creating a Level Sequence, double-click on it to open the Sequencer Editor so that you can begin creating your cinematic. 
 
 ![](BlankScene.png)
 
 [REGION:caption]
-Above we have a newly created and empty Level Sequence.
+Above, we have a newly-created, empty Level Sequence.
 [/REGION]
 
-The first thing you will need to do is add a **Track** type which you can do from the **Add** button drop-down menu. 
+The first thing you will need to do is add a **Track** type, which you can do from the **Add** button's drop-down menu. 
 
 ![](AddToSequencer.png)
 
-From the drop-down, you will see several Track types you can select from as well as the ability to add an **Actor To Sequencer**. The Actor To Sequencer option will allow you to add any Actors that you have selected in your level at the time to Sequencer so that you can manipulate them during your scene. 
+From the drop-down menu, you will see several Track types you can select from, as the ability to add an **Actor To Sequencer**. The Actor To Sequencer option will enable you to add any Actors that you have selected in your level at the time to Sequencer so that you can manipulate them during your scene.
 
-Typically if you are creating a cinematic that has characters, animals, creatures or anything along those lines that will animate and move you will have a Skeletal Mesh for it which you will need to add to Sequencer. 
-For example below we have a Skeletal Mesh of a bear that we have placed in our level. With it selected, we can then click the **Add** button in Sequencer and choose **Actor To Sequencer** so that we can add it and control it through the Sequencer Editor.
+Typically, if you are creating a cinematic that has characters, animals, creatures, or anything along those lines that will animate and move, you will have a [Skeletal Mesh](Engine/Content/Types/SkeletalMeshes) for it, which you will need to add to Sequencer.
+For example, below, we have a Skeletal Mesh of a bear that we have placed in our level. With it selected, we can then click the **Add** button in Sequencer and choose **Actor To Sequencer** so that we can add it and control it through the Sequencer Editor.
 
 ![](AddBear.png)(w:640)
 
@@ -62,7 +64,7 @@ Once we have added our Skeletal Mesh, we can then add **Sub-tracks** to affect t
 Based on the type of Track you create, Sub-Tracks may be added and the ability to add Sub-Tracks may be available. 
 [/REGION]
 
-Below we select the Animation Sub-Track and assign an Animation for our Skeletal Mesh Bear to play. 
+Below, we select the Animation Sub-Track and assign an Animation for our Skeletal Mesh Bear to play. 
 
 [OBJECT:EmbeddedVideo]
 [PARAMLITERAL:width]
@@ -76,14 +78,13 @@ KoksbI7ZOTM
 [/PARAMLITERAL]
 [/OBJECT]
 
-##Possessables vs. Spawnables
+## Possessables vs. Spawnables
 
-Sequencer is very similar to Matinee in that it uses the concept of "possessables" meaning that an Actor exists in the level and Sequencer will take possession of it in order to apply changes you make to it via Sequencer. 
+Sequencer is very similar to Matinee in that it uses the concept of "possessables", meaning that an Actor exists in the level and Sequencer will take possession of it in order to apply changes you make to it via Sequencer. 
+As illustrated in our bear example above, you can place a Skeletal Mesh in the level and add it to Sequencer, then inside Sequencer, you can add an Animation Track associated with the Skeletal Mesh and assign different animations for it to play during your scene. 
+In this instance, we are "possessing" the Skeletal Mesh that exists in the level in order to tell it what to do. 
 
-As illustrated in our bear example above, you can place a Skeletal Mesh in the level and add it to Sequencer, then inside Sequencer you can add an Animation Track to the Skeletal Mesh and assign different animations for it to play during your scene. 
-We are "possessing" that Skeletal Mesh that exists in order to tell it what to do. Sequencer also includes another form of manipulation which is referred to as "spawnables" meaning that the object we are possessing does not exist and Sequencer will "spawn" the object when told to and has the authority over it and can determine its lifecycle. 
-
-So how is this beneficial and when would I elect to use a "spawnable" object? Since Sequencer is the owner and spawns the object, it is not bound to a particular level. Anything deemed as spawnable can be used in any level, so you could create a scene and have it take place in any environment and reuse it in new levels without having to recreate it. 
+Sequencer also includes another form of manipulation, which is referred to as "spawnables", meaning that the object we are possessing does not yet exist, so Sequencer will "spawn" the object when told to, and will have authority over it to determine its lifecycle. Since Sequencer spawns and owns the object, the object is not bound to a particular level. Anything marked as "spawnable" can be used in any level, so you could create a scene and have it take place in any environment and reuse it in new levels without having to recreate it. 
 
 [REGION:note]
 Please see [Creating Spawnables](Engine\Sequencer\HowTo\Spawnables\) for more information. 
@@ -92,9 +93,8 @@ Please see [Creating Spawnables](Engine\Sequencer\HowTo\Spawnables\) for more in
 
 ## Keyframing Scenes
 
-If you are familiar with external animation tools, the concept of using Keyframes to drive content may be familiar to you. 
-If you are not, in short Sequencer allows you to manipulate your content by adding **Keys** (referred to as Keyframing) with defined properties at desired points along a timeline. 
-When reaching those keys in the Timeline, the properties you have defined at each keyframe are updated to reflect the value you've entered.
+If you are familiar with most animation tools, the concept of using Keyframes to drive content may be familiar to you. Sequencer allows you to manipulate your content by adding **Keys** (referred to as "Keyframing") with defined properties at desired points along a timeline. 
+Upon reaching those keys in the Timeline, the properties you have defined at each keyframe are updated to reflect the values you've entered.
 
 A sample scene below illustrates the process of adding movement by keyframing different transformations of a Skeletal Mesh Actor. 
 
@@ -111,14 +111,14 @@ dVDNJkQxlRk
 [/OBJECT]
 
 [REGION:caption]
-Above we added a keyframe for the starting position and another keyframe for the ending position for the **Location** track.
+Above, we added a keyframe for the starting position, and another keyframe for the ending position, for the **Location** track.
 [/REGION]
 
-You can add a keyframe by selecting a track and pressing **Enter** or by clicking the add keyframe button on each track. 
+You can add a keyframe by selecting a track and pressing **Enter**, or by clicking the "add keyframe" button on each track. 
 
 ![](KeyframeButton.png)
 
-Some properties in the **Details** panel of an Actor can be keyframed directly by clicking the add keyframe button next to each property. 
+Some properties in the **Details** panel of an Actor can be keyframed directly by clicking the "add keyframe" button next to each property. 
 
 ![](KeyframeDetails.png)
 
@@ -126,50 +126,84 @@ Some properties in the **Details** panel of an Actor can be keyframed directly b
 Above, keyframe buttons are shown for the Current Focal Length and Aperture settings of a Cine Camera Actor. 
 [/REGION]
 
-
-
 [REGION:tip]
-With a Track containing Keyframes selected, you can press the **,** and **.** keys to quickly move between each placed keyframe. 
+With a Track containing Keyframes selected, you can press the **,** (comma) and **.** (period) keys to quickly move between each placed keyframe. 
 [/REGION]
 
+With a key (or multiple keys or sections) selected, you can use the **Transform Keys/Sections** tools to reposition and rescale your selection. 
+
+![](SelectionOption.png)
+
+You can also open the Transform Keys/Selection tools by pressing **Ctrl+M**. 
+
+This will allow you offset the keys/selection by a specified amount or the amount to scale the keys/section by the specified time. 
+
+[OBJECT:EmbeddedVideo]
+[PARAMLITERAL:width]
+890					
+[/PARAMLITERAL]
+[PARAMLITERAL:height]
+196
+[/PARAMLITERAL]
+[PARAMLITERAL:videoid]
+Kd3CwzMwTLs
+[/PARAMLITERAL]
+[/OBJECT]
+
+
+
 
 ### Playback Range
 
-Sequencer will playback your scenes based on the defined Start and Stop playback ranges (green/red markers below). 
+Sequencer will play back your scenes based on the defined Start and Stop playback ranges (green/red markers below). 
 
 ![](PlaybackRange.png)
 
 [REGION:caption]
-Above our content contains two shots which end at frame 600 but our sequence extends to frame 1000 which results in dead space. 
+Above, our content contains two shots which end at frame 600, but our sequence extends to frame 1000, resulting in dead space.
 [/REGION]
 
-You can drag the Start and Stop Playback markers to fit your content or there is an option that can be set which will keep your content within these ranges. You can find it under the Sequencer Options drop-down menu and enabling **Keep Playback Range In Section Bounds**. 
+You can drag the **Start** and **Stop Playback** markers to fit your content, or there is an option that can be set which will keep your content within these ranges. You can find this option under the **Sequencer Options** drop-down menu, called **Keep Playback Range In Section Bounds**. 
 
 ![](KeepInRange.png)
 
+When working with [Shot Tracks](Engine\Sequencer\HowTo\TracksShot) and viewing sequences from the Master level as well as the Shot level, we evaluate the entire sequence at its relative time. 
+In the image below we have two red markers for the ending of playback ranges, the first denotes the end of the Shot (from the Shot level) while the second denotes the end of the Master Sequence. 
+What is happening in this example is that our Shot0020_001 is set to use 50 frames but we've set it to use 100 frames on the Master Level. 
+
+![](Evaluation_1.png)
+
+If we dive into our Shot, the first red marker denotes the end of the Shot at frame 50, while the second red marker denotes the end of the Shot at the master level. 
+
+![](Evaluation_2.png)
+
+Frames 50 through 100 are grayed out to indicate that they are not being considered for evaluation. 
+To address the difference in length, we could increase the amount of frames being evaluated at the Shot level to 100 or at the Master Sequence level, we could reduce the length of the Shot to 50 frames. 
+
+
 ##Specialized Viewports
 
 Sequencer allows you to use specialized viewports to make the editing process easier. 
 
 ![](SpecialViewports.png)
 
-These [Cinematic Viewports](Engine\Sequencer\HowTo\ViewportOptions) give a better idea of what the scene will look like and can be enabled from the **Viewport** options button. 
+These [Cinematic Viewports](Engine\Sequencer\HowTo\ViewportOptions) give a better idea of what the scene will look like, and can be enabled from the **Viewport** options button.
 
 ##Cinematic Actors
 
-From the **Modes** panel under **Cinematic** there are three Cinematic Actors you can use for crafting your cinematics. 
+From the **Modes** panel, under **Cinematic**, there are three Cinematic Actors you can use for crafting your cinematics.
 
 ![](CineTools.png)
 
-You can drag any of these into your level then add it to Sequencer or you can drag it into Sequencer (which will make it a Spawnable). 
+You can drag any of these into your level then add it to Sequencer, or you can drag it into Sequencer (which will make it a Spawnable). 
 
 ### Camera Rig Crane
 
-The Camera Rig Crane Actor can be used to simulate crane-like camera movements that are used in traditional film making. 
+The **Camera Rig Crane** Actor can be used to simulate crane-like camera movements that are used in traditional filmmaking. 
 
 ![](Crane.png)(w:480)
 
-You can attach a Camera to the Camera Rig Crane Actor then manipulate the crane's movement through the **Details** panel and **Crane Controls** values. Those values can affect the Crane Pitch, Crane Yaw and Crane Arm Length. These values can be keyframed within Sequencer which will allow you to adjust them during your cinematics as desired. 
+You can attach a Camera to the Camera Rig Crane Actor, then manipulate the crane's movement through the **Details** panel and **Crane Controls** values. Those values can affect the Crane Pitch, Crane Yaw, and Crane Arm Length. These values can be keyframed within Sequencer, which will allow you to adjust them during your cinematics (as desired). 
 
 [REGION:note]
 Please see [Shooting from a Camera Rig Crane](Engine\Sequencer\HowTo\CameraRigCrane) for more information. 
@@ -177,11 +211,11 @@ Please see [Shooting from a Camera Rig Crane](Engine\Sequencer\HowTo\CameraRigCr
 
 ### Camera Rig Rail
 
-The Camera Rig Rail Actor is a spline based tool that a Camera can be attached to which provides a "path" for movement. 
+The **Camera Rig Rail** Actor is a spline based tool that a Camera can be attached to, which provides a "path" for movement. 
 
 ![](Rail.png)(w:480)
 
-You can select each spline point and alter the tangents to produce the path you want your camera to follow. The camera that is attached to the rail can be rotated independently and the position of where the camera is on the rail can be adjusted with the **Current Position on Rail** property inside the **Details** panel. This value can be keyframed allowing you to alter the value through Sequencer. 
+You can select each spline point and alter the tangents to produce the path you want your camera to follow. The camera that is attached to the rail can be rotated independently, and the camera's position on the rail can be adjusted with the **Current Position on Rail** property inside of the **Details** panel. This value can also be keyframed.
 
 [REGION:note]
 Please see [Shooting from a Camera Rig Rail](Engine\Sequencer\HowTo\CameraRigRail) for more information. 
@@ -189,7 +223,7 @@ Please see [Shooting from a Camera Rig Rail](Engine\Sequencer\HowTo\CameraRigRai
 
 ### Cine Camera Actor
 
-The Cine Camera Actor is similar to the default Camera Actor. However, additional Camera Settings are available. 
+The **Cine Camera Actor** is similar to the default Camera Actor. However, additional Camera Settings are available. 
 
 ![](CineCamera.png)(w:480)
 
@@ -197,7 +231,7 @@ The Cine Camera Actor is similar to the default Camera Actor. However, additiona
 Cine Camera Actor (black camera) and default Camera Actor (blue camera) pictured above. 
 [/REGION]
 
-The Cine Camera Actor includes settings for Look at Tracking (following an Actor), Filmback Settings (16:9 DSLR, Super 8mm, 35mm VistaVision, etc.), Lens and Focus Settings as well as Current Aperture and Focus Distances. While using a normal Camera Actor is perfectly fine, to achieve a more cinematic feel and expose more camera settings it is recommended that you shoot your scenes with a Cine Camera Actor. 
+The Cine Camera Actor includes settings for Look at Tracking (following an Actor), Filmback Settings (16:9 DSLR, Super 8mm, 35mm VistaVision, etc.), Lens and Focus Settings, as well as Current Aperture and Focus Distances. While using a normal Camera Actor is perfectly fine, to achieve a more cinematic feel and expose more camera settings, it is recommended that you shoot your scenes with a Cine Camera Actor.
 
 [REGION:note]
 Please see [Using Cine Camera Actors](Engine\Sequencer\HowTo\CineCameraActors) for more information. 
@@ -220,7 +254,8 @@ mvCPXslVb8Y
 [/PARAMLITERAL]
 [/OBJECT]
 
-Add a new recording, pick an Actor to record and click Record to start capturing. When stopping the recording, new Level Sequence assets will be created which you can then incorporate into your existing cinematic or you can use pieces from the recorded Level Sequence to create new sequences.
+Add a new recording, pick an Actor to record, and click Record to start capturing. When stopping the recording, new Level Sequence assets will be created. You can then incorporate them into your existing cinematic, or you can use pieces from the recorded Level Sequence to create new sequences.
+You can also record audio to go along with your recorded sequence from an external microphone, which will allow you to narrate your sequence, or provide instructions while recording level editor actions.
 
 [REGION:note]
 Please see [Recording Gameplay](Engine/Sequencer/HowTo/SequenceRecorder) for more information. 
@@ -228,11 +263,11 @@ Please see [Recording Gameplay](Engine/Sequencer/HowTo/SequenceRecorder) for mor
 
 ##Render Movie Settings
 
-Within Sequencer, you can choose to render out your cinematic to a Video or Image Sequence from the **Render Movie Settings** option. 
+Within Sequencer, you can choose to render out your cinematic to a Video, or Image Sequence, from the **Render Movie Settings** option. 
 
 ![](RenderMovie.png)
 
-This will display the **Render Movie Settings** window which you can use to define how your scene is rendered. 
+This will display the **Render Movie Settings** window, which you can use to define how your scene is rendered. 
 
 ![](CaptureMovie.png)
 
@@ -246,34 +281,35 @@ Click the **Capture Movie** button to start the capture process of recording you
 Please see [](Engine/Sequencer/HowTo/BurnIns) for a step-by-step guide on using burn-ins. 
 [/REGION]
 
-[PUBLISH:docs]
 ### Import/Export EDLs
 
 [INCLUDE:Engine/Sequencer/HowTo/ImportExportEDLs#intro]
 
+As part of exporting an EDL (Edit Decision List), you can also add **Frame Handles** automatically to each of your shots, and specify the number of frames you'd like to add. 
+
+ ![](RenderMovieSequence.png)
+ 
 [REGION:note]
 Please see [](Engine/Sequencer/HowTo/ImportExportEDLs) for a step-by-step guide on using EDLs.
 [/REGION]
-[/PUBLISH:docs]
 
-[PUBLISH:docs]
+
 ### Custom Render Passes
 
 If you are looking to export your cinematic in different render passes, you can do so from the **Render Movie Settings** window. 
-This will allow you to specify the render pass(es) to use when exporting your sequence. You can also export HDR data as .exr files and define compression and color gamut settings. 
+This will allow you to specify the render pass(es) to use when exporting your sequence. You can also export HDR (high dynamic-range) data as .exr files, and define compression and color gamut settings. 
 
 
 [REGION:note]
 Please see [](Engine\Sequencer\Workflow\CustomRenderPass) for more information.
 [/REGION]
-[/PUBLISH:docs]
 
 
 ## Workflow Considerations
 
-Once you have a basic understanding of how the Sequencer Editor works, you may want to consider how you will use the tool as there several different ways to author content. Whether you create all of your content within a single Level Sequence, embed sequences within one another and use a Mater sequence to control them all, or use a Master Sequence asset to get started, Sequencer offers multiple ways in which you can generate your cinematics. 
-
+Once you have a basic understanding of how the Sequencer Editor works, you may want to consider how you will use the tool (there are many ways to author content). Whether you create all of your content within a single Level Sequence, embed sequences within one another and use a Master sequence to control them all, or use a Master Sequence asset to get started, Sequencer offers multiple ways in which you can generate your cinematics.
 
 [REGION:note]
 Please see [Workflow Considerations](Engine\Sequencer\Workflow\ToolUsage) for more information. 
-[/REGION]
\ No newline at end of file
+[/REGION]
+
diff --git a/Engine/Documentation/Source/Engine/UI/LevelEditor/Details/LevelEditorDetails.INT.udn b/Engine/Documentation/Source/Engine/UI/LevelEditor/Details/LevelEditorDetails.INT.udn
--- a/Engine/Documentation/Source/Engine/UI/LevelEditor/Details/LevelEditorDetails.INT.udn
+++ b/Engine/Documentation/Source/Engine/UI/LevelEditor/Details/LevelEditorDetails.INT.udn
@@ -4,26 +4,6 @@ Crumbs: %ROOT%, Engine, Engine/UI, Engine/UI/LevelEditor
 Description:Panel providing information and customized editing tools for selected Actors in the Level Editor.
 Version: 4.12
 
-[VAR:playlistnav]
-[OBJECT:PlaylistNavigation]
-	[PARAMLITERAL:id]
-		LevelDesign_Beginner
-	[/PARAMLITERAL]
-	[PARAM:previous]
-		[Previous Topic](Engine/UI/LevelEditor/SceneOutliner "%Engine/UI/LevelEditor/SceneOutliner:title%")
-	[/PARAM]
-	[PARAM:next]
-		[Next Topic](Engine/UI/LevelEditor/Viewports/Basics "%Engine/UI/LevelEditor/Viewports/Basics:title%")
-	[/PARAM]
-	[PARAM:home]
-		[](Tracks/LevelDesign/Beginner)
-		%Tracks/LevelDesign/Beginner:menu%
-	[/PARAM]
-[/OBJECT]
-[/VAR]
-
-%playlistnav%
-
 [TOC(start:2 end:3)]
 
 [EXCERPT:DetailsOverview]
@@ -317,6 +297,3 @@ properties to be displayed.
 [/OBJECT]
 [/REGION]
 [/EXCERPT:CustomCategories]
-
-
-%playlistnav%
diff --git a/Engine/Documentation/Source/Engine/UI/LevelEditor/InEditorTesting/InEditorTesting.INT.udn b/Engine/Documentation/Source/Engine/UI/LevelEditor/InEditorTesting/InEditorTesting.INT.udn
--- a/Engine/Documentation/Source/Engine/UI/LevelEditor/InEditorTesting/InEditorTesting.INT.udn
+++ b/Engine/Documentation/Source/Engine/UI/LevelEditor/InEditorTesting/InEditorTesting.INT.udn
@@ -4,23 +4,6 @@ Crumbs: %ROOT%, Engine, Engine/UI, Engine/UI/LevelEditor
 Description: Guide to testing levels and gameplay directly in the editor.
 Version: 4.9
 
-[VAR:playlistnav]
-[OBJECT:PlaylistNavigation]
-	[PARAMLITERAL:id]
-		LevelDesign_Beginner
-	[/PARAMLITERAL]
-	[PARAM:previous]
-		[Previous Topic](Engine/Rendering/LightingAndShadows/Basics "%Engine/Rendering/LightingAndShadows/Basics:title%")
-	[/PARAM]
-	[PARAM:home]
-		[](Tracks/LevelDesign/Beginner)
-		%Tracks/LevelDesign/Beginner:menu%
-	[/PARAM]
-[/OBJECT]
-[/VAR]
-
-%playlistnav%
-
 [TOC(start:2 end:2)]
 
 [VAR:Topic]
@@ -232,6 +215,3 @@ However, you can change _Blueprint_ defaults while a preview session is active,
 For projects that contain C++ code, there is a **Compile** button on the **Level Editor Toolbar**.  This button recompiles and reloads game code on the fly.  If you change a property or a function in
 a C++ file in your project, clicking the **Compile** button will recompile and reload your game module, so that your code changes are reflected in the editor.  In some cases, you can even compile
 while using _Play In Editor_ or _Simulate In Editor_, and your code changes will update without you needing to stop gameplay or simulation.
-
-
-%playlistnav%
diff --git a/Engine/Documentation/Source/Engine/UI/LevelEditor/LevelEditor.INT.udn b/Engine/Documentation/Source/Engine/UI/LevelEditor/LevelEditor.INT.udn
--- a/Engine/Documentation/Source/Engine/UI/LevelEditor/LevelEditor.INT.udn
+++ b/Engine/Documentation/Source/Engine/UI/LevelEditor/LevelEditor.INT.udn
@@ -6,33 +6,6 @@ Version: 4.9
 parent:Engine/Editor
 order:3
 
-[VAR:playlistnav]
-[OBJECT:PlaylistNavigation]
-	[PARAMLITERAL:id]
-		LevelDesign_Beginner
-	[/PARAMLITERAL]
-	[PARAM:next]
-		[Next Topic](Engine/UI/LevelEditor/Toolbar "%Engine/UI/LevelEditor/Toolbar:title%")
-	[/PARAM]
-	[PARAM:home]
-		[](Tracks/LevelDesign/Beginner)
-		%Tracks/LevelDesign/Beginner:menu%
-	[/PARAM]
-[/OBJECT]
-[OBJECT:PlaylistNavigation]
-	[PARAMLITERAL:id]
-		Programming_Beginner
-	[/PARAMLITERAL]
-	[PARAM:next]
-		[Next Topic](Engine/Actors "%Engine/Actors:title%")
-	[/PARAM]
-	[PARAM:home]
-		[](Tracks/Programming/Beginner)
-	[/PARAM]
-[/OBJECT]
-[/VAR]
-
-%playlistnav%
 
 [VAR:Topic]
 [OBJECT:Topic]
@@ -213,5 +186,3 @@ See the [**World Outliner**](Engine/UI/LevelEditor/SceneOutliner) page for detai
 [INCLUDE:Engine/UI/LevelEditor/Layers#Overview]
 
 See the **[](Engine/UI/LevelEditor/Layers)** page for details on using the **Layers** panel.
-
-%playlistnav%
diff --git a/Engine/Documentation/Source/Engine/UI/LevelEditor/Modes/LevelEditorTools.INT.udn b/Engine/Documentation/Source/Engine/UI/LevelEditor/Modes/LevelEditorTools.INT.udn
--- a/Engine/Documentation/Source/Engine/UI/LevelEditor/Modes/LevelEditorTools.INT.udn
+++ b/Engine/Documentation/Source/Engine/UI/LevelEditor/Modes/LevelEditorTools.INT.udn
@@ -4,26 +4,6 @@ Crumbs: %ROOT%, Engine, Engine/UI, Engine/UI/LevelEditor
 Description:Set of tools used to control the mode the Level Editor is currently in, reshape the builder brush, and create new Brush geometry and volumes.
 Version: 4.9
 
-[VAR:playlistnav]
-[OBJECT:PlaylistNavigation]
-	[PARAMLITERAL:id]
-		LevelDesign_Beginner
-	[/PARAMLITERAL]
-	[PARAM:previous]
-		[Previous Topic](Engine/UI/LevelEditor/Toolbar "%Engine/UI/LevelEditor/Toolbar:title%")
-	[/PARAM]
-	[PARAM:next]
-		[Next Topic](Engine/UI/LevelEditor/SceneOutliner "%Engine/UI/LevelEditor/SceneOutliner:title%")
-	[/PARAM]
-	[PARAM:home]
-		[](Tracks/LevelDesign/Beginner)
-		%Tracks/LevelDesign/Beginner:menu%
-	[/PARAM]
-[/OBJECT]
-[/VAR]
-
-%playlistnav%
-
 [EXCERPT:main]
 [INCLUDE:Shared/Glossary/M#glossary_modes_panel]
 
@@ -41,5 +21,3 @@ The **Modes** represent different editing modes the Level Editor can be put into
 | ![](LE_Tools_Foliage.png) | Toggles [**Foliage** mode](Engine/Foliage) for painting instanced foliage. | **Shift + 4**|
 | ![](LE_Tools_Brush.png) | Toggles [**Geometry Editing** mode](Engine/Actors/Brushes) for modifying Brushes to geometry. | **Shift + 5**|
 [/EXCERPT:main]
-
-%playlistnav%
diff --git a/Engine/Documentation/Source/Engine/UI/LevelEditor/SceneOutliner/SceneOutliner.INT.udn b/Engine/Documentation/Source/Engine/UI/LevelEditor/SceneOutliner/SceneOutliner.INT.udn
--- a/Engine/Documentation/Source/Engine/UI/LevelEditor/SceneOutliner/SceneOutliner.INT.udn
+++ b/Engine/Documentation/Source/Engine/UI/LevelEditor/SceneOutliner/SceneOutliner.INT.udn
@@ -4,26 +4,6 @@ Crumbs: %ROOT%, Engine, Engine/UI, Engine/UI/LevelEditor
 Description:Hierarchical tree-view of all Actors within the current level used for selection, attachment, and more.
 Version: 4.9
 
-[VAR:playlistnav]
-[OBJECT:PlaylistNavigation]
-	[PARAMLITERAL:id]
-		LevelDesign_Beginner
-	[/PARAMLITERAL]
-	[PARAM:previous]
-		[Previous Topic](Engine/UI/LevelEditor/Modes "%Engine/UI/LevelEditor/Modes:title%")
-	[/PARAM]
-	[PARAM:next]
-		[Next Topic](Engine/UI/LevelEditor/Details "%Engine/UI/LevelEditor/Details:title%")
-	[/PARAM]
-	[PARAM:home]
-		[](Tracks/LevelDesign/Beginner)
-		%Tracks/LevelDesign/Beginner:menu%
-	[/PARAM]
-[/OBJECT]
-[/VAR]
-
-%playlistnav%
-
 [TOC(start:2)]
 
 
@@ -62,6 +42,3 @@ Double quotes can be placed around the search term to force an exact match of th
 ![Scene Outliner Search - Verbatim Match](outliner_search_verbatim.png)
 
 When the **Info** columns are shown, the data within them can be searched as well. This is really useful if you want to display Actors in a certain level; just turn on the Level display and type "+[LevelName]" in the search box.
-
-
-%playlistnav%
diff --git a/Engine/Documentation/Source/Engine/UI/LevelEditor/Toolbar/LevelEditorToolbar.INT.udn b/Engine/Documentation/Source/Engine/UI/LevelEditor/Toolbar/LevelEditorToolbar.INT.udn
--- a/Engine/Documentation/Source/Engine/UI/LevelEditor/Toolbar/LevelEditorToolbar.INT.udn
+++ b/Engine/Documentation/Source/Engine/UI/LevelEditor/Toolbar/LevelEditorToolbar.INT.udn
@@ -4,26 +4,6 @@ Crumbs: %ROOT%, Engine, Engine/UI, Engine/UI/LevelEditor
 Description:Collection of commands that provide quick access to commonly used tools and operations.
 Version: 4.9
 
-[VAR:playlistnav]
-[OBJECT:PlaylistNavigation]
-	[PARAMLITERAL:id]
-		LevelDesign_Beginner
-	[/PARAMLITERAL]
-	[PARAM:previous]
-		[Previous Topic](Engine/UI/LevelEditor "%Engine/UI/LevelEditor:title%")
-	[/PARAM]
-	[PARAM:next]
-		[Next Topic](Engine/UI/LevelEditor/Modes "%Engine/UI/LevelEditor/Modes:title%")
-	[/PARAM]
-	[PARAM:home]
-		[](Tracks/LevelDesign/Beginner)
-		%Tracks/LevelDesign/Beginner:menu%
-	[/PARAM]
-[/OBJECT]
-[/VAR]
-
-%playlistnav%
-
 The **Toolbar** panel, like in most applications, is a group of commands providing quick access to commonly used tools and operations.
 
 [REGION:fullwidth]
@@ -51,6 +31,3 @@ Starts the game in simulation mode. See the [Simulate In Editor](Engine/UI/Level
 Starts the game in normal play mode. Clicking the arrow displays the **Play Options** menu. See the [Play In Editor](Engine/UI/LevelEditor/InEditorTesting#PlayInEditor) section for more information.
 [/EXCERPT:playin]
 -->
-
-
-%playlistnav%
diff --git a/Engine/Documentation/Source/Engine/UI/LevelEditor/Viewports/Basics/ViewportBasics.INT.udn b/Engine/Documentation/Source/Engine/UI/LevelEditor/Viewports/Basics/ViewportBasics.INT.udn
--- a/Engine/Documentation/Source/Engine/UI/LevelEditor/Viewports/Basics/ViewportBasics.INT.udn
+++ b/Engine/Documentation/Source/Engine/UI/LevelEditor/Viewports/Basics/ViewportBasics.INT.udn
@@ -4,26 +4,8 @@ Crumbs: %ROOT%, Engine, Engine/UI, Engine/UI/LevelEditor/Viewports
 Description:Basic concepts and features of the Viewports in Unreal Editor.
 Version: 4.9
 SkillLevel: Beginner
+track:LevelDesign_Beginner|6
 
-[VAR:playlistnav]
-[OBJECT:PlaylistNavigation]
-	[PARAMLITERAL:id]
-		LevelDesign_Beginner
-	[/PARAMLITERAL]
-	[PARAM:previous]
-		[Previous Topic](Engine/UI/LevelEditor/Details "%Engine/UI/LevelEditor/Details:title%")
-	[/PARAM]
-	[PARAM:next]
-		[Next Topic](Engine/Actors "%Engine/Actors:title%")
-	[/PARAM]
-	[PARAM:home]
-		[](Tracks/LevelDesign/Beginner)
-		%Tracks/LevelDesign/Beginner:menu%
-	[/PARAM]
-[/OBJECT]
-[/VAR]
-
-%playlistnav%
 
 [VAR:Topic]
 [OBJECT:Topic]
@@ -206,5 +188,3 @@ button located at the top right of the Viewport when in Immersive Mode.
 	![](immersive_5.png)
 	[/PARAM]
 [/OBJECT]
-
-%playlistnav%
diff --git a/Engine/Documentation/Source/Engine/UI/MeshPreviewScene/MeshPreviewScene.INT.udn b/Engine/Documentation/Source/Engine/UI/MeshPreviewScene/MeshPreviewScene.INT.udn
--- a/Engine/Documentation/Source/Engine/UI/MeshPreviewScene/MeshPreviewScene.INT.udn
+++ b/Engine/Documentation/Source/Engine/UI/MeshPreviewScene/MeshPreviewScene.INT.udn
@@ -1,13 +1,14 @@
-Availability:docs
+Availability:Public
 Title: Mesh Preview Scenes
 Crumbs: %ROOT%, Engine
 Description: How to setup and use Mesh Preview Scene Profiles in Unreal Engine 4.
 Version: 4.13
+parent:Engine/UI
+type:overview
+tags:User Interface
 
 
-
-In the editor various editor windows used for your Static and Skeletal Meshes you can now setup preview scenes to be used when inspecting your assets. This allows you to quickly assess what your assets would look like in multiple different 
-environments and lighting scenarios without having to set up these scenes within the game viewport. 
+Across the various editor windows for your Static and Skeletal Meshes, you can now setup your own preview scenes to be used when inspecting your assets. This enables you to quickly assess what your assets would look like in multiple environments and lighting scenarios without having to set up these scenes within your level.
 
 ![](6.png)(w:950)
 
diff --git a/Engine/Documentation/Source/Engine/UMG/HowTo/Create3DWidgets/ActorSetup/ActorSetup.INT.udn b/Engine/Documentation/Source/Engine/UMG/HowTo/Create3DWidgets/ActorSetup/ActorSetup.INT.udn
--- a/Engine/Documentation/Source/Engine/UMG/HowTo/Create3DWidgets/ActorSetup/ActorSetup.INT.udn
+++ b/Engine/Documentation/Source/Engine/UMG/HowTo/Create3DWidgets/ActorSetup/ActorSetup.INT.udn
@@ -1,4 +1,4 @@
-Availability: Public
+Availability:Docs
 Title:2. Setting up the Actor Blueprint
 Crumbs: %ROOT%, Engine, Engine/UMG, Engine/UMG/HowTo, Engine\UMG\HowTo\Create3DWidgets
 Description:Here we use an Actor Blueprint with a Widget Component to display our Widget Blueprint in 3D world space. 
@@ -7,7 +7,8 @@ Related: Engine/UMG/UserGuide
 Related: Engine/Blueprints
 Version: 4.10
 SkillLevel: Intermediate
-checkpoint: editorqs
+checkpoint: umg_create_3d_howto
+tags:UMG UI Designer
 
 [VAR:Steps]
 [OBJECT:Navigation]
diff --git a/Engine/Documentation/Source/Engine/UMG/HowTo/Create3DWidgets/Create3DWidgets.INT.udn b/Engine/Documentation/Source/Engine/UMG/HowTo/Create3DWidgets/Create3DWidgets.INT.udn
--- a/Engine/Documentation/Source/Engine/UMG/HowTo/Create3DWidgets/Create3DWidgets.INT.udn
+++ b/Engine/Documentation/Source/Engine/UMG/HowTo/Create3DWidgets/Create3DWidgets.INT.udn
@@ -1,12 +1,13 @@
-Availability: Public
+Availability:Docs
 Title:Creating 3D Widgets
 Crumbs: %ROOT%, Engine, Engine/UMG, Engine/UMG/HowTo
 Description:Create an interactive 3D menu that appears in the game world instead of 2D screen space. 
 Version: 4.10
 SkillLevel: Intermediate
-parent:Engine/Editor
+parent:Engine/UMG/HowTo
 order:1
-checkpoint: editorqs
+checkpoint: umg_create_3d_howto
+tags:UMG UI Designer
 
 [REGION:warning]
 This is an experimental feature that is currently undergoing development. Some aspects may not work as expected or may change in future revisions.
diff --git a/Engine/Documentation/Source/Engine/UMG/HowTo/Create3DWidgets/WidgetSetup/WidgetSetup.INT.udn b/Engine/Documentation/Source/Engine/UMG/HowTo/Create3DWidgets/WidgetSetup/WidgetSetup.INT.udn
--- a/Engine/Documentation/Source/Engine/UMG/HowTo/Create3DWidgets/WidgetSetup/WidgetSetup.INT.udn
+++ b/Engine/Documentation/Source/Engine/UMG/HowTo/Create3DWidgets/WidgetSetup/WidgetSetup.INT.udn
@@ -1,4 +1,4 @@
-Availability: Public
+Availability:Docs
 Title:1. Setting up the Widget Blueprint
 Crumbs: %ROOT%, Engine, Engine/UMG, Engine/UMG/HowTo, Engine\UMG\HowTo\Create3DWidgets
 Description:Here we create the Widget Blueprint that we want to display in the 3D game world. 
@@ -7,7 +7,8 @@ Related: Engine/UMG/UserGuide
 Related: Engine/Blueprints
 Version: 4.10
 SkillLevel: Intermediate
-checkpoint: editorqs
+checkpoint: umg_create_3d_howto
+tags:UMG UI Designer
 
 [VAR:Steps]
 [OBJECT:Navigation]
diff --git a/Engine/Documentation/Source/Engine/UMG/HowTo/CreatingWidgets/CreatingWidgets.INT.udn b/Engine/Documentation/Source/Engine/UMG/HowTo/CreatingWidgets/CreatingWidgets.INT.udn
--- a/Engine/Documentation/Source/Engine/UMG/HowTo/CreatingWidgets/CreatingWidgets.INT.udn
+++ b/Engine/Documentation/Source/Engine/UMG/HowTo/CreatingWidgets/CreatingWidgets.INT.udn
@@ -6,6 +6,7 @@ Related: Engine/UMG
 Related: Engine/UMG/UserGuide
 Version: 4.12
 SkillLevel: Beginner
+tags:UMG UI Designer
 
 Typically in most games you will want to communicate some information to your players through your game's User Interface. 
 This can range from things like Main Menus, in-game Pause Menus, Health or Experience, Inventory Items, or help text to instruct players what to do in a given situation. 
diff --git a/Engine/Documentation/Source/Engine/UMG/HowTo/DragAndDrop/DragAndDrop.INT.udn b/Engine/Documentation/Source/Engine/UMG/HowTo/DragAndDrop/DragAndDrop.INT.udn
--- a/Engine/Documentation/Source/Engine/UMG/HowTo/DragAndDrop/DragAndDrop.INT.udn
+++ b/Engine/Documentation/Source/Engine/UMG/HowTo/DragAndDrop/DragAndDrop.INT.udn
@@ -1,13 +1,14 @@
-Availability: Docs
+Availability:Docs
 Title:Creating Drag and Drop UI
 Crumbs: %ROOT%, Engine, Engine/UMG, Engine/UMG/HowTo
-Description:Shows how you can create draggable and droppable UI widgets with UMG. 
-Version: 4.12
+Description:Shows how you can create drag and droppable UI widgets with UMG. 
+Version: 4.15
 SkillLevel: Intermediate
+tags:UMG UI Designer
 
 There may be instances in your project where you want your players to be able to interact with UI elements by dragging and dropping them onscreen. 
 This could be to customize interface layouts (positioning health bars or unit frames) or interacting with gameplay systems like inventory screens (adding/removing items around). 
-With UMG you can create these types of interactions with the **DragAndDropOperation** node and overriding some of the inherited Blueprint Functions based on your method of interaction. 
+With UMG you can create these types of interactions with the **DragAndDropOperation** node and override some of the inherited Blueprint Functions based on your method of interaction. 
 
 In this example, we take a look at how a player can drag and reposition a health bar on screen with the Left Mouse Button.
 
@@ -47,3 +48,4 @@ _At the end of this guide, you will have a UI widget that you can drag around in
 	[Click to Start](Engine\UMG\HowTo\DragAndDrop\ProjectSetup "%Engine\UMG\HowTo\DragAndDrop\ProjectSetup:title%")
 [/REGION]
 
+
diff --git a/Engine/Documentation/Source/Engine/UMG/HowTo/DragAndDrop/FinishingUp/FinishingUp.INT.udn b/Engine/Documentation/Source/Engine/UMG/HowTo/DragAndDrop/FinishingUp/FinishingUp.INT.udn
--- a/Engine/Documentation/Source/Engine/UMG/HowTo/DragAndDrop/FinishingUp/FinishingUp.INT.udn
+++ b/Engine/Documentation/Source/Engine/UMG/HowTo/DragAndDrop/FinishingUp/FinishingUp.INT.udn
@@ -4,8 +4,8 @@ Crumbs: %ROOT%, Engine, Engine/UMG, Engine/UMG/HowTo
 Description:We finish by configuring our drag widget, adding our Health Bar to our Main HUD for display and testing everything out.
 Related: Engine/UMG
 Related: Engine/UMG/UserGuide
-Version: 4.12
 SkillLevel: Intermediate
+tags:UMG UI Designer
 
 [VAR:Steps]
 [OBJECT:Navigation]
@@ -24,7 +24,7 @@ SkillLevel: Intermediate
 
 %Steps%
 
-In this final step we configure our **DragWidget** (which is the visual widget we are dragging). 
+In this final step, we configure our **DragWidget** (which is the visual widget we are dragging). 
 We then tell our Character Blueprint to add the **HUD** Widget Blueprint to the viewport as well as enable the Mouse Cursor so we can see where we are dragging. 
 
 ## Steps
@@ -41,14 +41,14 @@ We then tell our Character Blueprint to add the **HUD** Widget Blueprint to the
 
 	![](UpdateDragDrop4.png)
 
-	In our example we will have a silhouette of our Health Bar as the drag visual while keeping the actual Health Bar in the same location. 
+	In our example, we will have a silhouette of our Health Bar as the drag visual while keeping the actual Health Bar in the same location. 
 
-1.	On the **Graph**, drag in **Widget Reference** and cast to the **HealthBar** to get the **Desired Size** then **Break Vector**. 
+1.	On the **Graph**, drag in **Widget Reference** then off **Event Construct**, cast to the **HealthBar** to get the **Desired Size** and **Break Vector**. 
 
 	![](UpdateDragDrop5.png)
 
 	This will give us the size of our Health Bar which we can then set our Size Box to. 
-	We could have manually entered the overrides, however if we change the size of our Health Bar we would then have to also change it here. 
+	We could have manually entered the overrides, however, if we change the size of our Health Bar we would then have to also change it here. 
 
 1.	Drag in the **Widget Size** and off is **Set Height Override** and **Set Width Override** and connect as shown below. 
 
@@ -87,3 +87,4 @@ Additional checking may be needed to ensure that players do not drag the widget
 
 %Steps%
 
+
diff --git a/Engine/Documentation/Source/Engine/UMG/HowTo/DragAndDrop/OnDragDetected/OnDragDetected.INT.udn b/Engine/Documentation/Source/Engine/UMG/HowTo/DragAndDrop/OnDragDetected/OnDragDetected.INT.udn
--- a/Engine/Documentation/Source/Engine/UMG/HowTo/DragAndDrop/OnDragDetected/OnDragDetected.INT.udn
+++ b/Engine/Documentation/Source/Engine/UMG/HowTo/DragAndDrop/OnDragDetected/OnDragDetected.INT.udn
@@ -4,8 +4,8 @@ Crumbs: %ROOT%, Engine, Engine/UMG, Engine/UMG/HowTo
 Description:In this step we set up what happens when drag is detected.
 Related: Engine/UMG
 Related: Engine/UMG/UserGuide
-Version: 4.12
 SkillLevel: Intermediate
+tags:UMG UI Designer
 
 [VAR:Steps]
 [OBJECT:Navigation]
@@ -25,7 +25,7 @@ SkillLevel: Intermediate
 
 %Steps%
 
-In this step we determine what happens when the player is in fact dragging the HealthBar around on screen. 
+In this step, we determine what happens when the player is, in fact, dragging the HealthBar around on the screen. 
 
 ## Steps
 
@@ -34,7 +34,7 @@ In this step we determine what happens when the player is in fact dragging the H
 	![](UpdatedDragDrop1.png)
 
 	When drag is detected, we will create our DragWidget (which will hold our drag widget). 
-	Before we continue with our OnDragDetected function, we need to pass through the Health Bar widget as reference when the DragWidget is created. 
+	Before we continue with our OnDragDetected function, we need to pass through the Health Bar widget as a reference when the DragWidget is created. 
 
 1.	Open **DragWidget** and add a **User Widget** variable called **WidgetReference** set to **Expose on Spawn** and **Editable**. 
 
@@ -64,7 +64,7 @@ In this step we determine what happens when the player is in fact dragging the H
 	_Click for full view._
 
 	Above we indicated what Widget Blueprint is our reference, what the drag visual should be and provided an offset in which to start dragging. 
-	There is an offset pin on the node by default, however we are using our own in our own calculation to determine the position in which to start dragging when Mouse Button Down is pressed. 
+	There is an offset pin on the node by default, however, we are using our own in our own calculation to determine the position in which to start dragging when Mouse Button Down is pressed. 
 
 The last few things we need to do is to determine what happens when the player releases the Left Mouse Button and performs the OnDrop function. 
 We also need to add the HealthBar Widget to our Main HUD Widget Blueprint and tell that to be displayed before testing everything out.  
@@ -72,3 +72,4 @@ We also need to add the HealthBar Widget to our Main HUD Widget Blueprint and te
 %Steps%
 
 
+
diff --git a/Engine/Documentation/Source/Engine/UMG/HowTo/DragAndDrop/OnDrop/OnDrop.INT.udn b/Engine/Documentation/Source/Engine/UMG/HowTo/DragAndDrop/OnDrop/OnDrop.INT.udn
--- a/Engine/Documentation/Source/Engine/UMG/HowTo/DragAndDrop/OnDrop/OnDrop.INT.udn
+++ b/Engine/Documentation/Source/Engine/UMG/HowTo/DragAndDrop/OnDrop/OnDrop.INT.udn
@@ -1,11 +1,11 @@
 Availability: Docs
 Title:4. On Drop
 Crumbs: %ROOT%, Engine, Engine/UMG, Engine/UMG/HowTo
-Description:In this step we configure what happens when the player drops the dragged widget.
+Description:In this step, we configure what happens when the player drops the dragged widget.
 Related: Engine/UMG
 Related: Engine/UMG/UserGuide
-Version: 4.12
 SkillLevel: Intermediate
+tags:UMG UI Designer
 
 [VAR:Steps]
 [OBJECT:Navigation]
@@ -33,14 +33,14 @@ Here we set up our main **HUD** Widget Blueprint and override what happens when
 
 	![](DragDrop13.png)
 
-	We are doing our drag detection inside the Health Bar widget, however when we release the Left Mouse Button we want to drop it in our HUD Widget Blueprint.
+	We are doing our drag detection inside the Health Bar widget, however, when we release the Left Mouse Button we want to drop it in our HUD Widget Blueprint.
 	In order for our HUD to receive hit detection, we need to set its **Visibility** setting to **Visible**. 
 
 1.	From the **Palette** window, add the **Health Bar** widget to the graph in your desired location. 
 
 	![](DragDrop14.png)
 
-1.	Go to the **Graph** Mode, then add an **OnDrop** function override from the **My Blueprint** panel. 
+1.	Go to the **Graph** Mode, then add a **OnDrop** function override from the **My Blueprint** panel. 
 
 	![](DragDrop15.png)
 
@@ -73,10 +73,11 @@ Here we set up our main **HUD** Widget Blueprint and override what happens when
 	_Click for full view._
 	
 	We uncheck **Remove DPIScale** and check the **Return Value** checkboxes as we do not need to remove the [](Engine/UMG/UserGuide/DPIScaling) and we have handled this function so we return as true. 
-	With this script we first remove the existing Health Bar widget before re-adding it to the screen in the new location relative to our Drag Offset. 
+	With this script, we first remove the existing Health Bar widget before re-adding it to the screen in the new location relative to our Drag Offset. 
 
-Our HUD is setup to handle dropping our dragged widget and it displays our Health Bar. 
-In the next and final step in this tutorial we set up our Dragged Widget to mimic the shape of our Health Bar as well as instruct our Character Blueprint to add the HUD to the viewport. 
+Our HUD is set up to handle dropping our dragged widget and it displays our Health Bar. 
+In the next and final step in this tutorial, we set up our Dragged Widget to mimic the shape of our Health Bar as well as instruct our Character Blueprint to add the HUD to the viewport. 
 
 %Steps%
 
+
diff --git a/Engine/Documentation/Source/Engine/UMG/HowTo/DragAndDrop/ProjectSetup/ProjectSetup.INT.udn b/Engine/Documentation/Source/Engine/UMG/HowTo/DragAndDrop/ProjectSetup/ProjectSetup.INT.udn
--- a/Engine/Documentation/Source/Engine/UMG/HowTo/DragAndDrop/ProjectSetup/ProjectSetup.INT.udn
+++ b/Engine/Documentation/Source/Engine/UMG/HowTo/DragAndDrop/ProjectSetup/ProjectSetup.INT.udn
@@ -2,8 +2,8 @@ Availability: Docs
 Title:1. Project Setup
 Crumbs: %ROOT%, Engine, Engine/UMG, Engine/UMG/HowTo
 Description:Here we create the assets that we'll need and our DragAndDropOperation Class. 
-Version: 4.12
 SkillLevel: Intermediate
+tags:UMG UI Designer
 
 [VAR:Steps]
 [OBJECT:Navigation]
@@ -37,7 +37,7 @@ For this how-to, we are using the **Blueprint Third Person Template**. However,
 	In this example, the element that we want to be drag and droppable is our **HealthBar** widget that will be placed in our **HUD**. 
 	We are going to create and drag a widget called **DragWidget** in the same shape as our Health Bar rather than dragging the actual Health Bar around. 
 
-1.	**Right-click** in the **Content Browser** and create a new **Blueprint Class** of the **DragDropOpertion** class and call it **WidgetDrag**. 
+1.	**Right-click** in the **Content Browser** and create a new **Blueprint Class** of the **DragDropOperation** class and call it **WidgetDrag**. 
 
 	![](DragDrop2.png)
 
diff --git a/Engine/Documentation/Source/Engine/UMG/HowTo/DragAndDrop/WidgetSetup/WidgetSetup.INT.udn b/Engine/Documentation/Source/Engine/UMG/HowTo/DragAndDrop/WidgetSetup/WidgetSetup.INT.udn
--- a/Engine/Documentation/Source/Engine/UMG/HowTo/DragAndDrop/WidgetSetup/WidgetSetup.INT.udn
+++ b/Engine/Documentation/Source/Engine/UMG/HowTo/DragAndDrop/WidgetSetup/WidgetSetup.INT.udn
@@ -2,8 +2,8 @@ Availability: Docs
 Title:2. Widget Setup
 Crumbs: %ROOT%, Engine, Engine/UMG, Engine/UMG/HowTo
 Description:In this step we script what happens when the Left Mouse Button is pressed. 
-Version: 4.12
 SkillLevel: Intermediate
+tags:UMG UI Designer
 
 [VAR:Steps]
 [OBJECT:Navigation]
@@ -23,7 +23,7 @@ SkillLevel: Intermediate
 
 %Steps%
 
-In this step we determine when the Left Mouse Button is pressed and store the location where it was pressed in screen space. 
+In this step, we determine when the Left Mouse Button is pressed and store the location where it was pressed in screen space. 
 We also use the **DetectDragIfPressed** to determine if the player is dragging with a specified key. 
 
 ## Steps
@@ -38,7 +38,7 @@ We also use the **DetectDragIfPressed** to determine if the player is dragging w
 	
 	![](DragDrop4b.png)
 
-	In this tutorial we will not be creating a functional Health Bar and this is intended for demonstration purposes only. 
+	In this tutorial, we will not be creating a functional Health Bar and this is intended for demonstration purposes only. 
 
 1.	Switch to the **Graph** Mode, then in the **My Blueprint** panel add the **OnMouseButtonDown** and **OnDragDetected** overrides. 
 
@@ -71,6 +71,7 @@ We also use the **DetectDragIfPressed** to determine if the player is dragging w
 
 	Here we are determining if the player is dragging the Left Mouse Button before which will start the On Drag Detected script. 
 
-With the ability to detect if the player is dragging in place, next we will finish off what happens when drag is detected. 
+With the ability to detect if the player is dragging in place, next we will finish off what happens when the drag is detected. 
 
 %Steps%
+
diff --git a/Engine/Documentation/Source/Engine/UMG/HowTo/EventBasedUI/EventBasedUI.INT.udn b/Engine/Documentation/Source/Engine/UMG/HowTo/EventBasedUI/EventBasedUI.INT.udn
--- a/Engine/Documentation/Source/Engine/UMG/HowTo/EventBasedUI/EventBasedUI.INT.udn
+++ b/Engine/Documentation/Source/Engine/UMG/HowTo/EventBasedUI/EventBasedUI.INT.udn
@@ -8,7 +8,7 @@ Related: Engine/Blueprints
 Related: Resources\SampleGames\UnrealMatch3
 Version: 4.10
 SkillLevel: Advanced
-
+tags:UMG UI Designer
 
 
 [VAR:Topic]
diff --git a/Engine/Documentation/Source/Engine/UMG/HowTo/InWorldWidgetInteraction/InWorldWidgetInteraction.INT.udn b/Engine/Documentation/Source/Engine/UMG/HowTo/InWorldWidgetInteraction/InWorldWidgetInteraction.INT.udn
--- a/Engine/Documentation/Source/Engine/UMG/HowTo/InWorldWidgetInteraction/InWorldWidgetInteraction.INT.udn
+++ b/Engine/Documentation/Source/Engine/UMG/HowTo/InWorldWidgetInteraction/InWorldWidgetInteraction.INT.udn
@@ -5,14 +5,27 @@ Description:An example of how to create, display and remove UI elements from the
 Related: Engine\UMG\UserGuide\WidgetInteraction\
 Related: Engine/UMG
 Related: Engine/UMG/UserGuide
-Version: 4.13
+Version: 4.15
 SkillLevel: Intermediate
+tags:UMG UI Designer
 
-Traditional UI consists in 2D space on top of a 2D or 3D world (for example menus, health bars or score displays). 
-However there may be times when you want a piece of UI to exist within your 3D environment that players can interact with (a keypad, a virtual menu, an inventory or other interactive UI element). 
+Traditional UI consists of 2D space on top of a 2D or 3D world (for example menus, health bars or score displays). 
+However, there may be times when you want a piece of UI existing within your 3D environment that players can interact with (a keypad, a virtual menu, an inventory or other interactive UI element). 
 In VR projects, this is even more prevalent and UMG supports these types of interactions with the [Widget Interaction](Engine\UMG\UserGuide\WidgetInteraction) component. 
 
-In this how-to we create a widget containing a button that can be placed in the level and clicked on by a player. 
+In this how-to, we create a widget containing a button that can be placed in the level and clicked on by a player similar to the video below. 
+
+[OBJECT:EmbeddedVideo]
+[PARAMLITERAL:width]
+640
+[/PARAMLITERAL]
+[PARAMLITERAL:height]
+360
+[/PARAMLITERAL]
+[PARAMLITERAL:videoid]
+bOZy7H9IMOI
+[/PARAMLITERAL]
+[/OBJECT]
 
 
 ## Steps
@@ -21,12 +34,8 @@ In this how-to we create a widget containing a button that can be placed in the
 For this how-to, we are using the **Blueprint First Person Template** with **Starter Content** enabled. 
 [/REGION]
 
-1.	Inside the **Content/FirstPersonBP/Blueprints** folder, open the **FirstPersonCharacter** Blueprint. 
-
-	![](WidgetInteraction_01.png)
-
-1.	Click the **Add Component** button and search for and add the **Widget Interaction** component. 
-
+1.	Inside the **Content/FirstPersonBP/Blueprints** folder, open the **FirstPersonCharacter** Blueprint then add a component of the **Widget Interaction** type. 
+	
 	![](WidgetInteraction_02.png)
 
 1.	**Left-click** drag and drop the Widget Interaction component on top of the **Sphere** component attached to **FP_Gun**. 
@@ -39,17 +48,17 @@ For this how-to, we are using the **Blueprint First Person Template** with **Sta
 
 	![](WidgetInteraction_04.png)
 
-1.	Inside the **Event Graph**, add a **Right Mouse Button** Mouse Event node.
-
-1.	Drag the Widget Interaction component into the graph, then drag off it and add a **Press Pointer Key** node. 
+1.	Inside the **Event Graph**, add a **Right Mouse Button** Mouse Event, drag in the Widget Interaction component and add a **Press Pointer Key** node. 
 
 	![](WidgetInteraction_05.png)
 
-1.	Connect to the **Pressed** pin of the **Right Mouse Button** node and set the **Key** to **Left Mouse Button**. By default UMG registers click events as a result of a Left Mouse Button key press, using this node will allow us to simulate that when the Right Mouse Button is pressed. 
-	You could use any Key Input event to call this function (such as a trigger press or button press to simulate a Left Mouse Button click). 
+1.	Connect to the **Pressed** pin of the **Right Mouse Button** node and set the **Key** to **Left Mouse Button**. 
 
 	![](WidgetInteraction_06.png)
 
+	By default UMG registers click events as a result of a Left Mouse Button key press, using this node will allow us to simulate that when the Right Mouse Button is pressed. 
+	You could use any Key Input event to call this function (such as a trigger press or button press to simulate a Left Mouse Button click). 
+
 1.	Add a **Release Pointer Key** node (set to **Left Mouse Button**) and connect it to the **Released** pin of the Right Mouse Button node. 
 
 	![](WidgetInteraction_07.png)
@@ -58,7 +67,7 @@ For this how-to, we are using the **Blueprint First Person Template** with **Sta
 
 	![](WidgetInteraction_08.png)
 
-1.	In the **Visual Designer**, remove the **Canvas Panel** then add a **Button** with a **Text** widget on top of the button. For this example we are purely adding a button that can be clicked on. The text on the button we will indicate the number of times the button has been clicked by the player. You can create any type of interactive widget and interact with it using the Widget Interaction component (things like sliders, combo boxes, check boxes, etc. can be interacted with in world space).
+1.	In the **Visual Designer**, remove the **Canvas Panel** then add a **Button** with a **Text** widget on top of the button. 
 
 	[REGION:lightbox]
 	[![](WidgetInteraction_09.png)](WidgetInteraction_09.png)
@@ -68,6 +77,8 @@ For this how-to, we are using the **Blueprint First Person Template** with **Sta
 	Click for full view. 
 	[/REGION] 
 
+	For this example, we are purely adding a button that can be clicked on. The text on the button we will indicate the number of times the button has been clicked by the player. You can create any type of interactive widget and interact with it using the Widget Interaction component (things like sliders, combo boxes, check boxes, etc. can be interacted with in world space).
+
 1.	In the **Details** panel for the **Button**, under **Style/Hovered** change the **Tint** to a different color. 
 
 	![](WidgetInteraction_10.png)
@@ -76,32 +87,42 @@ For this how-to, we are using the **Blueprint First Person Template** with **Sta
 
 	![](WidgetInteraction_11.png)
 
-1.	In the **Graph**, create a **Text** Variable called **Value** then **Compile** and set **Value** to **0**. This is the variable that we will update when the player clicks on the button which we will bind to our Text widget. 
+1.	In the **Graph**, create a **Text** Variable called **Value** then **Compile** and set **Value** to **0**. 
 
 	![](WidgetInteraction_12.png)
 
+	This is the variable that we will update when the player clicks on the button which we will bind to our Text widget. 
+
 1.	On the **Designer** tab for the **Text** widget, click the **Bind** option next to **Text** and select the **Value** property.
 
 	![](WidgetInteraction_13.png)
 
-1.	In the **Details** panel for the **Button** widget under **Events**, click the **+** next to **On Clicked**. This will fire whenever the button is clicked on which is called from inside our **FirstPersonCharacter** Blueprint. 
+1.	In the **Details** panel for the **Button** widget under **Events**, click the **+** next to **On Clicked**. 
 
 	![](WidgetInteraction_14.png)
 
-1.	On the **Graph** tab, hold **Ctrl** and drag in the **Value** variable then off it use **To String** followed by **String To Int**. We are converting the Text to a String, then from a String to an Int so that we can increase the value when the button is clicked. 
+	This will fire whenever the button is clicked on which is called from inside our **FirstPersonCharacter** Blueprint. 
+
+1.	On the **Graph** tab, hold **Ctrl** and drag in the **Value** variable then off it use **To String** followed by **String To Int**. 
 	
 	![](WidgetInteraction_15.png)
 
+	We are converting the Text to a String, then from a String to an Int so that we can increase the value when the button is clicked. 
+
 1.	Off the **String To Int** pin, use an **Integer + Integer** node set to **+ 1**. 
 
-1.	Hold **Alt** and drag in the **Value** variable to set it, then connect the remaining pins to the **OnClicked** event as shown. The **ToText(int)** node will be automatically created when you connect the out pin of **Integer + Integer** to the text input pin on **Value**. 
+1.	Hold **Alt** and drag in the **Value** variable to set it, then connect the remaining pins to the **OnClicked** event as shown. 
 
 	![](WidgetInteraction_16.png)
 
-1.	In the **Content Browser**, create a new **Blueprint** based on **Actor** called **ExampleWidget** and add a **Widget** component to it. This will represent the 3D version of our **InteractiveWidget** Widget Blueprint. 
+	The **ToText(int)** node will be automatically created when you connect the out pin of **Integer + Integer** to the text input pin on **Value**. 
+
+1.	In the **Content Browser**, create a new **Blueprint** based on **Actor** called **ExampleWidget** and add a **Widget** component to it. 
 
 	![](WidgetInteraction_18.png)
 
+	This will represent the 3D version of our **InteractiveWidget** Widget Blueprint. 
+
 1.	In the **Details** panel of the **Widget** component under **User Interface**, change the **Widget Class** to **InteractiveWidget**. 
 
 	![](WidgetInteraction_19.png)
@@ -132,3 +153,4 @@ bOZy7H9IMOI
 As noted during the steps, we are using the Right Mouse Button to register click events however you can use gamepad button clicks to register Pointer Key Presses/Releases or (as shown below) Motion Controller Trigger presses. Here we point our Right Motion Controller (with an attached Widget Interaction component) at the button and press the Right Trigger to register a button click.
 
 ![](WidgetInteraction_21.png)
+
diff --git a/Engine/Documentation/Source/Engine/UMG/HowTo/ScalingUI/ScalingUI.INT.udn b/Engine/Documentation/Source/Engine/UMG/HowTo/ScalingUI/ScalingUI.INT.udn
--- a/Engine/Documentation/Source/Engine/UMG/HowTo/ScalingUI/ScalingUI.INT.udn
+++ b/Engine/Documentation/Source/Engine/UMG/HowTo/ScalingUI/ScalingUI.INT.udn
@@ -8,6 +8,7 @@ Related: Engine/UMG/UserGuide/DPIScaling
 Related: Resources\SampleGames\UnrealMatch3
 Version: 4.10
 SkillLevel: Intermediate
+tags:UMG UI Designer
 
 [VAR:Topic]
 [OBJECT:Topic]
diff --git a/Engine/Documentation/Source/Engine/UMG/HowTo/UMGHowTo.INT.udn b/Engine/Documentation/Source/Engine/UMG/HowTo/UMGHowTo.INT.udn
--- a/Engine/Documentation/Source/Engine/UMG/HowTo/UMGHowTo.INT.udn
+++ b/Engine/Documentation/Source/Engine/UMG/HowTo/UMGHowTo.INT.udn
@@ -4,8 +4,7 @@ Crumbs: %ROOT%, Engine, Engine\UMG
 Description:A How To Guide for creating UI elements with Unreal Motion Graphics UI Designer (UMG). 
 Related: Engine\UMG
 Related: Engine\UMG\UserGuide
-version: 4.9
-
+tags:UMG UI Designer
 
 The **Unreal Motion Graphics UI Designer (UMG) - How To...** pages will show you how to create various UI elements using UMG. Each example comes with step-by-step instructions and several images to help guide you through the process of creating the described content. Reference links are also included to provide more detailed documentation of the systems and content being outlined in each guide. 
 
@@ -20,18 +19,14 @@ Refer to the table below for more information and click each link to view its co
 [REGION:wrapper transparent]
 [](FANCY:Engine\UMG\HowTo\WidgetTemplates "%Engine\UMG\HowTo\WidgetTemplates:description%")
 [](FANCY:Engine\UMG\HowTo\Create3DWidgets "%Engine\UMG\HowTo\Create3DWidgets:description%")
-[PUBLISH:docs]
 [](FANCY:Engine\UMG\HowTo\DragAndDrop "%Engine\UMG\HowTo\DragAndDrop:description%")
-[/PUBLISH:docs]
 [](FANCY:Engine\UMG\HowTo\ScalingUI "%Engine\UMG\HowTo\ScalingUI:description%")
 [/REGION]
 
 ## Advanced
 [REGION:wrapper transparent]
 [](FANCY:Engine\UMG\HowTo\EventBasedUI "%Engine\UMG\HowTo\EventBasedUI:description%")
-[PUBLISH:docs]
 [](FANCY:Engine\UMG\HowTo\VirtualKeyboards "%Engine\UMG\HowTo\VirtualKeyboards:description%")
-[/PUBLISH:docs]
 [/REGION]
 
 
diff --git a/Engine/Documentation/Source/Engine/UMG/HowTo/VirtualKeyboards/ActorSetup/ActorSetup.INT.udn b/Engine/Documentation/Source/Engine/UMG/HowTo/VirtualKeyboards/ActorSetup/ActorSetup.INT.udn
--- a/Engine/Documentation/Source/Engine/UMG/HowTo/VirtualKeyboards/ActorSetup/ActorSetup.INT.udn
+++ b/Engine/Documentation/Source/Engine/UMG/HowTo/VirtualKeyboards/ActorSetup/ActorSetup.INT.udn
@@ -5,8 +5,8 @@ Description:Demonstrates how you can create and interaction with a 3D Virtual Ke
 Related: Engine\UMG\UserGuide\WidgetInteraction\
 Related: Engine/UMG
 Related: Engine/UMG/UserGuide
-Version: 4.13
 SkillLevel: Advanced
+tags:UMG UI Designer
 
 [VAR:Steps]
 [OBJECT:Navigation]
@@ -31,30 +31,38 @@ With everything in place, we give it a test to see the end result which prints *
 
 ## Steps
 
-1.	Open the **KeypadDisplay** Blueprint and add a **Widget** component with class set to **KeypadDisplayWB** and **Draw at Desired Size**. Draw at Desired Size is a quick way to cause the render target to automatically match the desired size inside the Widget Blueprint.
+1.	Open the **KeypadDisplay** Blueprint and add a **Widget** component with class set to **KeypadDisplayWB** and **Draw at Desired Size**. 
 
 	![](Keypad15.png)
 
+	Draw at Desired Size is a quick way to cause the render target to automatically match the desired size inside the Widget Blueprint.
+
 1.	Open the **Keypad** Blueprint and add a **Widget** component with class set to **KeypadWB** and **Draw at Desired Size**. 
 
-1.	Uncheck the **Window Focusable** option for this Widget Component. We don't want this component to take focus away from the Keypad Display widget, so this needs to be unchecked. 
+1.	Uncheck the **Window Focusable** option for this Widget Component. 
 
 	![](Keypad16.png)
 
-1.	Also inside **Keypad**, add a **WidgetInteraction** component with **Enable Hit Testing** set to false. We add this WidgetInteraction component so that we don't have to get a reference to the WidgetInteraction component in the player character and we disable Hit Testing as we don't want it to generate any hover effects. WidgetInteraction components by default are set to **Virtual User Index 0** when they are created and all of them communicate using the Virtual User Index ID. Since the WidgetInteraction component in the player character is also using Virtaul User Index 0, this is the same as getting a reference to it.
+	We don't want this component to take focus away from the Keypad Display widget, so this needs to be unchecked. 
+
+1.	Also inside **Keypad**, add a **WidgetInteraction** component with **Enable Hit Testing** set to false. 
 
 	![](Keypad17.png) 
 
+	We add this WidgetInteraction component so that we don't have to get a reference to the WidgetInteraction component in the player character and we disable Hit Testing as we don't want it to generate any hover effects. WidgetInteraction components by default are set to **Virtual User Index 0** when they are created and all of them communicate using the Virtual User Index ID. Since the WidgetInteraction component in the player character is also using Virtual User Index 0, this is the same as getting a reference to it.
+
 1.	On the Graph of **Keypad**, off **Event Begin Play** take the **Widget** component and **Get User Widget Object** to **Cast to KeypadWB**. 
 
 	![](Keypad18.png)
 
 1.	Off the **As Keypad WB** pin, **Set Widget Interaction** and pass through the **WidgetInteraction** component from Step 4. 
 
-1.	Also off the **As Keypad WB** pin, call the **Assign Enter Pressed** which will create an **Enter Pressed** Custom Event. Now when we click on the Enter button this custom event will be called which we can use to check if the correct code was entered. 
+1.	Also off the **As Keypad WB** pin, call the **Assign Enter Pressed** which will create an **Enter Pressed** Custom Event. 
 
 	![](Keypad19.png)
 
+	Now when we click on the Enter button this custom event will be called which we can use to check if the correct code was entered. 
+
 1.	Create a **Variable** of the **Keypad Display** Actor type called **TargetDisplay** and set it to **Editable**. 
 
 	![](Keypad20.png)
@@ -100,6 +108,7 @@ N2Fo4izRgh0
 [/PARAMLITERAL]
 [/OBJECT]
 
-While this example only prints to the screen if the correct code was entered or not, you could instead fire off some event that opens a door, spawns an enemy or any other form of gameplay releated event! 
+While this example only prints to the screen if the correct code was entered or not, you could instead fire off some event that opens a door, spawns an enemy or any other form of gameplay related event! 
+
 
 
diff --git a/Engine/Documentation/Source/Engine/UMG/HowTo/VirtualKeyboards/CharacterSetup/CharacterSetup.INT.udn b/Engine/Documentation/Source/Engine/UMG/HowTo/VirtualKeyboards/CharacterSetup/CharacterSetup.INT.udn
--- a/Engine/Documentation/Source/Engine/UMG/HowTo/VirtualKeyboards/CharacterSetup/CharacterSetup.INT.udn
+++ b/Engine/Documentation/Source/Engine/UMG/HowTo/VirtualKeyboards/CharacterSetup/CharacterSetup.INT.udn
@@ -2,8 +2,8 @@ Availability: Public
 Title:1. Character Setup
 Crumbs: %ROOT%, Engine, Engine/UMG, Engine/UMG/HowTo
 Description:Here we add the Widget Interaction component to our character and create the other assets we will need. 
-Version: 4.13
 SkillLevel: Advanced
+tags:UMG UI Designer
 
 [VAR:Steps]
 [OBJECT:Navigation]
@@ -32,20 +32,27 @@ For this how-to, we are using the **Blueprint First Person Template** with **Sta
 
 1.	Inside the **Content/FirstPersonBP/Blueprints** folder, open the **FirstPersonCharacter** Blueprint. 
 
-1.	Add a **WidgetInteraction** component attached to the **Sphere** then zero out the **Location** and **Show Debug** in the **Details** panel. Our Widget Interaction component will now point in the direction that player points the gun indicated in-game by a debug marker.
+1.	Add a **WidgetInteraction** component attached to the **Sphere** then zero out the **Location** and **Show Debug** in the **Details** panel. 
 
 	![](Keypad01.png)
 
+	Our Widget Interaction component will now point in the direction that player points the gun indicated in-game by a debug marker.
+
 1.	In the **Graph**, add a **Right Mouse Button** Event and drag in the **Widget Interaction** Component. 
 
-1.	Off **Widget Interaction** add **Press Pointer Key** and **Release Pointer Key** both set to **Left Mouse Button** and connect as shown. Widgets respond to mouse/touch interaction via Left Mouse Button click or touch controls and we can simulate that using the **Press Pointer Key** function. In our example we use the **Right Mouse Button** to simulate a Left Mouse Button click, however you could use a gamepad button press or even Motion Controller Trigger press to fire off the Press Pointer Key function.
+1.	Off **Widget Interaction** add **Press Pointer Key** and **Release Pointer Key** both set to **Left Mouse Button** and connect as shown. 
 
 	![](Keypad02.png) 
 
-1.	In the **Content Browser** create two **Class Blueprints** based on **Actor** and two **Widget Blueprints**. For the two Actor Blueprints, call one **Keypad** (to represent the keypad entry) and the other **KeypadDisplay** (to represent and display the values entered from the Keypad). For the Widget Blueprints use the same names as the Actor Blueprints however to distinguish them, suffix them with **WB** (for Widget Blueprint).
+	Widgets respond to mouse/touch interaction via Left Mouse Button click or touch controls and we can simulate that using the **Press Pointer Key** function. In our example we use the **Right Mouse Button** to simulate a Left Mouse Button click, however, you could use a gamepad button press or even Motion Controller Trigger press to fire off the Press Pointer Key function.
+
+1.	In the **Content Browser** create two **Class Blueprints** based on **Actor** and two **Widget Blueprints**. 
 
 	![](Keypad03.png) 
 
+	For the two Actor Blueprints, call one **Keypad** (to represent the keypad entry) and the other **KeypadDisplay** (to represent and display the values entered from the Keypad). 
+	For the Widget Blueprints use the same names as the Actor Blueprints, however, to distinguish them, suffix them with **WB** (for Widget Blueprint).
+
 
 ## End Result
 
@@ -53,3 +60,4 @@ Our project is ready to go and our character can interact with Widgets in the wo
 
 %Steps%
 
+
diff --git a/Engine/Documentation/Source/Engine/UMG/HowTo/VirtualKeyboards/CreatingVirtualKeyboards.INT.udn b/Engine/Documentation/Source/Engine/UMG/HowTo/VirtualKeyboards/CreatingVirtualKeyboards.INT.udn
--- a/Engine/Documentation/Source/Engine/UMG/HowTo/VirtualKeyboards/CreatingVirtualKeyboards.INT.udn
+++ b/Engine/Documentation/Source/Engine/UMG/HowTo/VirtualKeyboards/CreatingVirtualKeyboards.INT.udn
@@ -2,16 +2,17 @@ Availability: Public
 Title:Creating Virtual Keyboard Interactions
 Crumbs: %ROOT%, Engine, Engine/UMG, Engine/UMG/HowTo
 Description:Demonstrates how you can create and interaction with a 3D Virtual Keyboard with UMG. 
-Version: 4.13
+Version: 4.15
 SkillLevel: Advanced
 parent:Engine/UMG/HowTo
 order:4
 checkpoint: editorqs
+tags:UMG UI Designer
 
 If you are looking to create a 3D keypad or virtual keyboard that allows a player to enter specified characters, you can do so with a combination of the [Widget Interaction](Engine\UMG\UserGuide\WidgetInteraction) component and [Widget](Engine/Components/Widget) components. 
 The Widget Interaction component will help facilitate which "buttons" are clicked on and processes sending the characters through to a Widget component in an Actor that will handle the display.
 
-In this how-to we create a keypad that can be interacted with in and provides entries to a text box. 
+In this how-to, we create a keypad that can be interacted with in and provides entries to a text box. 
 
 [OBJECT:EmbeddedVideo]
 [PARAMLITERAL:width]
@@ -29,7 +30,7 @@ _At the end of this guide, you will have a keypad and display that appears insid
 
 One important aspect in setting this up will be to determine the Widgets that are (or are not) focusable. If the Text Box that displays the characters the player enters is set to focusable and the buttons that the player clicks on to send characters through is also focusable, the Text Box would lose player focus when a button (or any other widget) is clicked on and not receive the commands issued. 
 
-The first step however is to set up our character with the Widget Interaction component and create the assets to display our UI. 
+The first step, however, is to set up our character with the Widget Interaction component and create the assets to display our UI. 
 
 [OBJECT:TopicButtonList]
 	[PARAM:icon]
@@ -51,3 +52,4 @@ The first step however is to set up our character with the Widget Interaction co
 	[Click to Start](Engine\UMG\HowTo\VirtualKeyboards\CharacterSetup "%Engine\UMG\HowTo\VirtualKeyboards\CharacterSetup:title%")
 [/REGION]
 
+
diff --git a/Engine/Documentation/Source/Engine/UMG/HowTo/VirtualKeyboards/WidgetSetup/WidgetSetup.INT.udn b/Engine/Documentation/Source/Engine/UMG/HowTo/VirtualKeyboards/WidgetSetup/WidgetSetup.INT.udn
--- a/Engine/Documentation/Source/Engine/UMG/HowTo/VirtualKeyboards/WidgetSetup/WidgetSetup.INT.udn
+++ b/Engine/Documentation/Source/Engine/UMG/HowTo/VirtualKeyboards/WidgetSetup/WidgetSetup.INT.udn
@@ -2,8 +2,8 @@ Availability: Public
 Title:2. Widget Blueprint Setup
 Crumbs: %ROOT%, Engine, Engine/UMG, Engine/UMG/HowTo
 Description:Here we create the keypad and keypad display widgets that will appear in the level. 
-Version: 4.13
 SkillLevel: Advanced
+tags:UMG UI Designer
 
 [VAR:Steps]
 [OBJECT:Navigation]
@@ -23,7 +23,7 @@ SkillLevel: Advanced
 
 %Steps%
 
-In this step we create the layout of our interactive UI Widgets and script what happens when clicking a button. 
+In this step, we create the layout of our interactive UI Widgets and script what happens when clicking a button. 
 
 
 ## Steps
@@ -36,22 +36,30 @@ In this step we create the layout of our interactive UI Widgets and script what
 
 	![](Keypad04.png)
 
-1.	For the **Text Box**, set **Hint Text** to **1212** and adjust the **Font Size** to **28**. For example, the player must use the keypad to enter the text 1212 to signal a successful entry but you could use any value you wish. 
+1.	For the **Text Box**, set **Hint Text** to **1212** and adjust the **Font Size** to **28**. 
 
 	![](Keypad05.png)
 
-1.	Open the **KeypadWB** and set the **Hierarchy** up to create 3 Buttons with the text **0**, **1** and **2** on them as shown below.  You may also want to press **F2** on each of the Buttons in the Hierarchy to change their names to reflect which button they are. 
+	For example, the player must use the keypad to enter the text 1212 to signal a successful entry but you could use any value you wish. 
+
+1.	Open the **KeypadWB** and set the **Hierarchy** up to create 3 Buttons with the text **0**, **1** and **2** on them as shown below.  
 
 	![](Keypad06.png)
+
+	You may also want to press **F2** on each of the Buttons in the Hierarchy to change their names to reflect which button they are. 
 	
-1.	Add a second **Horizontal Box** containing two **Buttons** with text on them, one set to **Del** the other set to **Ent**. Our keypad will only have three numerical buttons and two others for deleting (**Del**) or committing (**Ent**) entries. 
+1.	Add a second **Horizontal Box** containing two **Buttons** with text on them, one set to **Del** the other set to **Ent**. 
 
 	![](Keypad07.png)
 
-1.	Hold **Ctrl** and select each of your Buttons, then set the **Is Focusable** option to **False** for each. In order for our setup to work, we need to ensure that no other widgets we interact with will take focus away from the Keypad Display widget. If we left these set to true, when we go to click on a button, focus will be taken away from the Keypad Display and shifted to the button preventing us from passing the key press and value over to the Keypad Display. 
+	Our keypad will only have three numerical buttons and two others for deleting (**Del**) or committing (**Ent**) entries. 
+
+1.	Hold **Ctrl** and select each of your Buttons, then set the **Is Focusable** option to **False** for each. 
 
 	![](Keypad08.png)
 
+	In order for our setup to work, we need to ensure that no other widgets we interact with will take focus away from the Keypad Display widget. If we left these set to true, when we go to click on a button, focus will be taken away from the Keypad Display and shifted to the button preventing us from passing the key press and value over to the Keypad Display. 
+
 1.	On the **Graph**, add **OnClicked** Events for each of the Buttons. 
 
 	![](Keypad09.png)
@@ -64,10 +72,12 @@ In this step we create the layout of our interactive UI Widgets and script what
 
 	![](Keypad11.png)
 
-1.	Hold **Ctrl** and drag in **WidgetInteraction** then off it **Send Key Char** and connect **CharToSend** and the **OnClicked** Events. We use the **Send Key Char** function here as Text Boxes in Slate do not get Keys sent directly to them, for example they don't recognize that the A key was pressed so send A through.Instead they listen for Send Char, that way the lower level system operating code can do things like determine if Shift was held down while pressing the A key or if Caps Lock was on while pressing A. 
+1.	Hold **Ctrl** and drag in **WidgetInteraction** then off it **Send Key Char** and connect **CharToSend** and the **OnClicked** Events. 
 
 	![](Keypad12.png)
 
+	We use the **Send Key Char** function here as Text Boxes in Slate do not get Keys sent directly to them, for example, they don't recognize that the A key was pressed so send A through.Instead, they listen for Send Char, that way the lower level system operating code can do things like determining if Shift was held down while pressing the A key or if Caps Lock was on while pressing A. 
+
 1.	Off the **Del** OnClicked Event, use the **Press and Release Key** set to **Backspace** from the **WidgetInteraction** variable. 
 
 	![](Keypad13.png)
@@ -80,7 +90,8 @@ In this step we create the layout of our interactive UI Widgets and script what
 ## End Result
 
 Our Widget Blueprints are set up to display our keypad and keypad display, next we will add them as Widget Components to our Actor Blueprints so that we can have them appear inside our level. 
-By the end of the next step we should be able to test out our system and interact with both elements during gameplay. 
+By the end of the next step, we should be able to test out our system and interact with both elements during gameplay. 
 
 %Steps%
 
+
