Copying //UE4/Dev-Documentation to //UE4/Samples-Main (Source: //UE4/Dev-Documentation @ 3372893)

diff --git a/Engine/Documentation/Source/Engine/UMG/UserGuide/Fonts/Fonts.INT.udn b/Engine/Documentation/Source/Engine/UMG/UserGuide/Fonts/Fonts.INT.udn
--- a/Engine/Documentation/Source/Engine/UMG/UserGuide/Fonts/Fonts.INT.udn
+++ b/Engine/Documentation/Source/Engine/UMG/UserGuide/Fonts/Fonts.INT.udn
@@ -2,22 +2,154 @@ Availability:Public
 Title: Fonts
 Crumbs:%ROOT%, Engine, Engine/UMG, Engine/UMG/UserGuide
 Description:Details the process of importing fonts, working with the Font Editor and using fonts in UMG. 
+Parent: Engine/UMG
 Related: Engine/UMG
 Related: Engine/UMG/UserGuide/WidgetBlueprints
 Related: Engine/UMG/UserGuide/WidgetTypeReference
-Version: 4.14
+Version: 4.15
+tags:UMG UI Designer
 
 [TOC(start:2 end:2)]
 
-This page covers how to import your own fonts, using the **Font Editor** and fonts in the **Unreal Motion Graphics (UMG) UI Designer**. 
+This page covers how to import your own fonts, using the **Font Editor**, and fonts in the **Unreal Motion Graphics (UMG) UI Designer**. 
 
 ## Font Assets
 
-Fonts in Unreal Engine 4 are categorized as a **Font** asset and use two caching methods, **Runtime** which is in the form of a Composite Font or **Offline** which is the older pre-computed Font Atlas method. You can switch between the two methods by opening up a Font asset in the Font Editor (this provides a simple way to convert existing Font assets from the older method to the new composite method without having to replace them). 
+Fonts in Unreal Engine 4 are categorized as a **Font** asset and use two caching methods, **Runtime** which is in the form of a Composite Font or **Offline** which is the older
+pre-computed Font Atlas method. You can switch between the two methods by opening up a Font asset in the Font Editor (this provides a simple way to convert existing Font assets
+from the older method to the new composite method without having to replace them). 
+
+## Font Face Assets
+
+The **Font Face** asset is created when you import a font and it stores the font data that can be referenced by the Font asset. This means that the same font data can be reused across 
+multiple Font assets or even with multiple typefaces within the asset, and ultimately reduces memory consumption. 
+
+![](FontFaceWindow.png)(w:400)
+
+From the Font Face editor window you have access to the settings for **Hinting** and **Loading Policy**.
+
+[REGION:simpletable]
+| Property | Description |
+| --- |--- |
+| Source File Name | The filename of the font face we were created from. This may not always exist on disk, as we may have previously loaded and cached the font data inside this asset. |
+| Hinting | The hinting algorithm to use with the font face. [INCLUDE:#hinting] |
+| Loading Policy | Enum controlling how this font face should be loaded at runtime. See the enum for more explanations of the options. [INCLUDE:#loadingpolicy]|
+[/REGION]
+
+<!---
+[EXCERPT:hinting]
+[REGION:simpletable]
+| Selection | Description  |
+| --- | --- |
+| Default | Use the default hinting specified in the font. |
+| Auto | Force the use of an automatic hinting algorithm. |
+| Auto Light | Force the use of an automatic hinting algorithm, optimized for non-monochrome displays. |
+| Monochrome | Force the use of an automatic hinting algorithm optimized for monochrome displays. |
+| None | Do not use hinting. |
+[/REGION] 
+[/EXCERPT:hinting]
+
+[EXCERPT:loadingpolicy]
+[REGION:simpletable]
+| Selection | Description |
+| --- | --- |
+| Lazy Load | Lazy load the entire font into memory. This will consume more memory than Streaming, however, there will be zero file-IO when rendering glyphs within the font, although the initial load may cause a hitch. |
+| Stream | Stream the font from disk. This will consume less memory than Lazy Load or Inline, however, there will be file-IO when rendering glyphs,  which may cause hitches under certain circumstances or on certain platforms. |
+| Inline | Embed the font data within the asset. This will consume more memory than Streaming, however, it is guaranteed to be hitch free (only valid for font data within a Font Face asset). |
+[/REGION] 
+[/EXCERPT:loadingpolicy]
+--->
+
+
+## Upgrading Font Assets to 4.15 and Later
+
+Starting with Unreal Engine 4.15, Font assets have now been split into two separate assets; Font and Font Faces. Existing Font assets will upgrade their internal font data into 
+embedded Font Face assets during load, but will not automatically split them. You will have to do this by opening the Font Editor and selecting **Upgrade Data** under the **Default Font Family** list.
+This will enable you to edit and share these assets with other Font assets that you've created. 
+
+![](FontUpgradeData.png)(w:175)
+
+After you click **Upgrade Data**, you will be prompted to save the new Font Face asset. 
+
+
+## Font Editor
+
+When you **Double-click** on a Font asset in the **Content Browser**, it will open up inside the **Font Editor** window. 
+
+![](FontEditorWindow.png)(w:820)
+
+A breakdown of the Font Editor Window is presented below: 
+
+#### Toolbar Menu
+
+![](ToolbarMenu.png)(w:550)
+
+From this menu, you can save any changes you make, find the asset in the **Content Browser**, change the Background Color of the preview window (4) or
+the Foreground Color (text color) in the preview window. There are options for Updating or Exporting changes made, however, these options are only available
+within the **Offline** cache mode. 
+
+#### Default Font Family
+
+![](DefaultFontFamily.png)(w:550)
+
+In this window, you can assign the Default Font Family for use with this Font asset. You can add versions of a particular Font style (for example Normal,
+Bold, Italics, Underline, etc.) or have a collection of different Font styles as one Composite Font. If you have created a blank Font asset, you can assign
+a font from inside this window as well. Once a Font has been added, you can also define the hinting algorithm used with the font. 
+
+#### Sub-Font Family
+
+![](AddSubFontFamily.png)(h:50)
+
+In this window, when you click the **Add Sub-Font Family* button, you can assign the Sub-Font Family for this Font asset to use. 
+
+[![](SubFontFamily.png)(w:550)](SubFontFamily.png)
+
+Here you can specify a Character Range and if a character entered falls within the range you can specify a different Font style
+to use instead of the Default. This is useful for when you want to use different Font types for different languages.
+
+
+#### Preview 
+
+![](Preview.png)(w:550)
+
+This window allows you to preview your fonts and provides a text entry box for entering sample texts. 
+
+##### Draw For Metrics
+
+[![](PreviewDrawFontMetrics.png)(w:550)](PreviewDrawFontMetrics.png)
+
+The **Draw Font Metrics** toggle will overlay the line height, glyph bounding boxes, and base-line as part of the preview.
+
+* **Base Line** - This is the line in which the text sits.
+* **Line Bounds** - This is the bounding box created for the length of the given text string. 
+* **Grapheme Cluster Bounds** - This is the bounding box drawn around what is considered a logical character in a given language, and may be comprised of several glyphs (eg. a base character and accent glyph).
+* **Glyph Bounds** - This is the bounding box drawn around the given glyph.
+
+#### Details
+
+![](Details.png)(w:500)
+
+In this window, you can change the Font Cache Type as well as change the Font Size and Font Name (for Runtime). 
+* If you are using the older method, you can still change the parameters for your Font while in Offline cache mode. 
+* You can also convert any existing Font assets from **Offline** to **Runtime** without having to replace them. 
+
+
+An example Font asset is provided below. 
+
+[REGION:lightbox]
+[![](ExampleFontLayout.png)(w:820)](ExampleFontLayout.png)
+[/REGION]
+
+_Click image for a full view._
+
+A Composite Font will always contain a Default Font Family, and may also contain any number of Sub-Font Families that should be used for a given range of characters. Each Font Family is itself made up of any number of Fonts that can be named based on their style. At runtime, the most suitable Font to use for each character based on the Fonts available in the Font Family for that character range is used. 
+
+As seen in the example image above, the Japanese text falls within the character ranges of the Japanese font family and so is drawn using Source Han Sans rather than the Default Font Family Roboto. Fonts in a Sub-Font Family are preferably chosen by name match, as in the case of Regular, Bold, and Light, however can also fallback to matching based on the attributes of the Default Font, as is the case of Bold Italic (it automatically chose the Bold Japanese font because the font contained the Bold attribute, and was the best match). 
+
 
 ## Creating & Importing Fonts
 
-This section illustrates how to create a new Font asset or import TTF or OTF files to generate a Font asset. 
+This section illustrates how to create a new Font asset and the ways you can import TTF or OTF files to generate your Font Face and Font assets. 
 
 ### Creating Font Assets
 
@@ -25,9 +157,13 @@ To create a Font asset, follow the steps below:
 
 1.	Click the **Add New** button inside the **Content Browser**, then under **User Interface** select the **Font** option. 
 
-	![](NewFont.png)
-
-1.	A new Composite Font asset will be created, prompting you to give it a name. 
+	[REGION:lightbox]
+	[![](NewFont.png)(h:500)](NewFont.png)
+	[/REGION]
+	[REGION:caption]
+	_Click image for full view._
+	[/REGION]
+1.	A new Composite Font asset will be created, prompting you to give it a **Name**. 
 
 	![](NameFont.png)
 
@@ -40,58 +176,73 @@ To create a Font asset, follow the steps below:
 Instead of clicking the **Add New** button, you can instead **Right-click** in an empty space in the **Content Browser** to access the context menu. From there, you can click the **User Interface** and **Font** option to create a Font asset. Do not forget to save your assets! 
 [/REGION]
 
-When you create a new Font asset inside the Editor, it will create an empty Composite Font asset which you can then define its parameters using the Font Editor. When you import a TTF or OTF file, a new Composite Font asset is generated using the font sample you provided (at which point you can define additional parameters by opening it up in the Font Editor). 
-
-### Importing Fonts
+When you create a new Font asset inside the Editor, it will create an empty Composite Font asset which you can then define its parameters using the Font Editor. When you import a TTF or OTF file, a new Font Face asset is created
+and you have the option to automatically create a Composite Font asset using the font data you provided (at which point you can define additional parameters by opening it up in the Font Editor).
 
-You can import either TTF or OTF files into the Editor, and you can do so in a number of different ways. 
 
-The first method is by clicking the **Import** button inside the **Content Browser**. 
-
-![](Import1.png)
+### Importing Fonts
 
-In the dialogue box that appears, navigate to the file to import then click on it and click **Open**. 
+You can import either TTF or OTF files into the Editor, and you can do so in a number of different ways demonstrated below.
 
-![](Import2.png)
+##### Importing Using the Content Browser
 
-The file will then be added to the Content Browser as a Composite Font asset, using the file you have selected. 
+You can use the Content Browser's **Import** button to select your TTF or OTF font file.
 
-![](Import3.png)
+1. In the Content Browser, click the **Import** button. 
 
-You can **Drag-and-drop** a TTF or OTF file into the **Content Browser** to create a Composite Font asset as well. 
+	![](Import1.png)
+1. In the dialogue box that appears, navigate to the TFF or OTF font file you want to import and select it. Then, click **Open**.
 
-![](Import4.png)
+	![](Import2.png)(w:500)
+1. After a moment, the **Font Face Import Options** will appear. 
 
-When using the Drag-and-drop method, a Composite Font asset will also automatically be created for you using the file you have selected. 
+	![](Import3.png)(w:500)
+		
+	Select **Yes** from the options listed to create your Font Face asset and your Composite Font asset in your Content Browser.
+	
+	![](Import3a.png)
+	
+	
+##### Importing Using Drag-and-drop
 
-## Font Editor
+You can **Drag-and-drop** a TTF or OTF file directly into the Content Browser to create your font assets.
 
-When you **Double-click** on a Font asset in the **Content Browser**, it will open up inside the **Font Editor** window. 
+1. Navigate to the folder where you've stored your TTF or OTF file(s). Select and hold to drag one over the **Content Browser** to start the import process. 
 
-![](FontEditorWindow.png)(w:820)
+	![](Import4.png)(w:650)
+1. After a moment, the **Font Import Options** will appear. 
+	
+	![](Import3.png)(w:500)
+		
+	Select **Yes** from the options listed to create your Font Face asset and your Composite Font asset in your Content Browser.
+	
+	![](Import3a.png)
 
-A breakdown of the Font Editor Window is presented below: 
 
-1.	**Toolbar Menu** - From this menu, you can save any changes you make, find the asset in the **Content Browser**, change the Background Color of the preview window (4) or the Foreground Color (text color) in the preview window. There are options for Updating or Exporting changes made, however these options are only available within the **Offline** cache mode. 
-1.	**Default Font Family** - In this window, you can assign the Default Font Family for use with this Font asset. You can add versions of a particular Font style (for example Normal, Bold, Italics, Underline, etc.) or have a collection of different Font styles as one Composite Font. If you have created a blank Font asset, you can assign a font from inside this window as well. Once a Font has been added, you can also define the hinting algorithm used with the font. 
-1.	**Sub-Font Family** - In this window, you assign the Sub-Font Family for this Font asset to use. Here you can specify a Character Range and if a character entered falls within the range, you can specify a different Font style to use instead of the Default (this is useful for when you want to use different Font types for different Languages). 
-1.	**Preview** - This window allow you to preview your fonts and provides a text entry box for entering sample texts. 
-1.	**Details** - In this window, you can change the Font Cache Type as well as change the Font Size and Font Name (for Runtime). 
-	* If you are using the older method, you can still change the parameters for your Font while in Offline cache mode. 
-	* You can also convert any existing Font assets from **Offline** to **Runtime** without having to replace them. 
+#### Importing Using the Font Editor
 
-An example Font asset is provided below. 
+You can import and create Font Face assets directly from the Default Font Family list in the **Font Editor**. 
 
-[REGION:lightbox]
-[![](ExampleFontLayout.png)(w:820)](ExampleFontLayout.png)
-[/REGION]
+1. Open an existing **Font** asset or create a new one using the **Add New** button in the Content Browser.
 
-_Click image for a full view._
+	![](Import5.png)(w:600)
+1. Click the **Add Font**  button to add a new font selection to the Default Font Family.
 
-A Composite Font will always contain a Default Font Family, and may also contain any number of Sub-Font Families that should be used for a given range of characters. Each Font Family is itself made up of any number of Fonts that can be named based on their style. At runtime, the most suitable Font to use for each character based on the Fonts available in the Font Family for that character range is used. 
+	![](Import6.png)(w:500) 
+1. Select the **folder** button next to the selection dropdown. 
+	
+	![](Import7.png)(w:500)
+1. In the dialogue box that appears, navigate to the TFF or OTF font file you want to import and select it. Then, click **Open**.
 
-As seen in the example image above, the Japanese text falls within the character ranges of the Japanese font family, and so is drawn using Source Han Sans rather than the Default Font Family Roboto. Fonts in a Sub-Font Family are preferably chosen by name match, as in the case of Regular, Bold, and Light, however can also fallback to matching based on the attributes of the Default Font, as is the case of Bold Italic (it automatically chose the Bold Japanese font because the font contained the Bold attribute, and was the best match). 
+	![](Import2.png)(w:500)
+1. Next, the **Save Font Face** window will appear. Give your font a name and choose the save location in your game folder's hierarchy. Then, click **Save**.
 
+	![](Import8.png)(w:500)
+	
+	The Font Face asset can now be found in your folder hierarchy.
+	
+	![](Import3a.png)
+	
 
 ## Fonts in UMG
 
@@ -99,7 +250,7 @@ Once you have created a [Widget Blueprint](Engine/UMG/UserGuide/WidgetBlueprints
 
 ![](FontWithNewOptions.png)
 
-By default the Engine uses **Roboto**, however if you click the drop-down menu, any Composite Font assets created can be selected and used instead. 
+By default the Engine uses **Roboto**, however, if you click the drop-down menu, any Composite Font assets created can be selected and used instead. 
 
 ![](NewFontWithNewOptions.png)
 
@@ -113,7 +264,7 @@ You can also specify the size of the font in the input box.
 
 ![](FontSizeNewOptions.png)
 
-Currently UMG only supports **Runtime** cached font assets. Also, if you have assigned fonts using the old method, none of your existing file based font settings will be lost, however going forward, you will need to create a Composite Font asset in order to use custom fonts with UMG. 
+Currently, UMG only supports **Runtime** cached font assets. Also, if you have assigned fonts using the old method, none of your existing file based font settings will be lost, however going forward, you will need to create a Composite Font asset in order to use custom fonts with UMG. 
 
 You can style your **Font** by [setting colors, Materials, and outline properties](Engine/UMG/UserGuide/Fonts/Materials), as well as by using shadows. 
 
@@ -121,3 +272,4 @@ You can style your **Font** by [setting colors, Materials, and outline propertie
 
 
 
+
diff --git a/Engine/Documentation/Source/Engine/UMG/UserGuide/WidgetTypeReference/BackgroundBlur/BackgroundBlur.INT.udn b/Engine/Documentation/Source/Engine/UMG/UserGuide/WidgetTypeReference/BackgroundBlur/BackgroundBlur.INT.udn
new file mode 100644
--- /dev/null
+++ b/Engine/Documentation/Source/Engine/UMG/UserGuide/WidgetTypeReference/BackgroundBlur/BackgroundBlur.INT.udn
@@ -0,0 +1,81 @@
+Availability:Public
+Title: Background Blur Widget
+Crumbs:%ROOT%
+Description: Describes how to use the Background Blur Widget to blur objects beneath a single child widget. 
+Related: Engine/UMG/UserGuide
+version: 4.15
+tags:UMG UI Designer
+Parent:Engine/UMG/UserGuide/WidgetTypeReference
+Order:
+Type:reference
+
+[TOC(start:2 end:2)]
+
+##Description
+
+[EXCERPT:desc]
+The **Background Blur Widget** can contain one child widget, providing an opportunity to surround it with adjustable padding and apply a post-process Gaussian blur to all content beneath the widget.
+[/EXCERPT:desc]
+
+## Details
+
+In the **Details** panel for a placed **Background Blur** widget, there are a couple of specific options that can be set, which pertain to the Widget: 
+
+![](BackgroundBlur.png)
+
+| Option | Description |
+| -------- | ----------- |
+| **Apply Alpha to Blur** | When true, this will modulate the strength of the blur based on the widget alpha. |
+| **Blur Strength** | How blurry the background is. Larger numbers increase blur, resulting in a larger runtime cost on the GPU. |
+| **Low-Quality Fallback Brush** | An image to draw (instead of applying a blur) when Low-Quality Override mode is enabled. You can enable Low-Quality Mode for background blurs by setting the cvar `Slate.ForceBackgroundBlurLowQualityOverride` to **1**. This is usually done in the project's scalability settings. |
+| **Blur Radius** | This is the number of pixels that will be weighted in each direction from any given pixel when computing the blur. A larger value is more costly but allows for stronger blurs. |
+
+Each of the aforementioned properties can also be set (or changed) at runtime through Blueprint Script. 
+
+![](BackgroundBlurBP.png)
+
+The ability to set other appearance settings (such as Horizontal and Vertical Alignment), as well as any padding around the Widget, can also be defined. 
+
+## Usage Example
+
+In the following example, we use the Background Blur Widget to highlight a menu when the game is paused, blurring out the background. 
+
+[OBJECT:EmbeddedVideo]
+[PARAMLITERAL:width]
+640
+[/PARAMLITERAL]
+[PARAMLITERAL:height]
+360
+[/PARAMLITERAL]
+[PARAMLITERAL:videoid]
+RUD047Z00bU
+[/PARAMLITERAL]
+[/OBJECT]
+
+For this, we add the Blur Widget with our simplified menu, using the **Blur Strength** value to determine the strength of background blur being applied.  
+
+![](BlurWidgetDesigner.png)
+
+On our Menu Widget Blueprint's Graph, we create a script to handle how our menu reacts to button clicks. 
+
+![](BlurWidgetGraph.png)
+
+Above, when we construct the widget, we are turning on the Mouse Cursor. When the **Resume** button is pressed, we hide the cursor, un-pause the game, and remove the menu from being displayed. 
+When the **Quit** button is pressed, we quit the game. Inside of our player character's Blueprint (see below), we add a script to create and display the menu when a key is pressed. In this case, whenever **P** is pressed, the game is paused when the menu is displayed.
+
+![](CharacterBlueprint.png)
+
+The result is the ability to pause the game and blur the background, keeping our menu intact for players to interact with. 
+
+![](50_Blur.png)
+
+We could also decrease the Blur Strength from 50 (above) down to 10 (below) to make the background slightly more visible. 
+
+![](10_Blur.png)
+
+
+
+
+
+
+
diff --git a/Engine/Documentation/Source/Engine/UMG/UserGuide/WidgetTypeReference/Invalidation/InvalidationBox.INT.udn b/Engine/Documentation/Source/Engine/UMG/UserGuide/WidgetTypeReference/Invalidation/InvalidationBox.INT.udn
--- a/Engine/Documentation/Source/Engine/UMG/UserGuide/WidgetTypeReference/Invalidation/InvalidationBox.INT.udn
+++ b/Engine/Documentation/Source/Engine/UMG/UserGuide/WidgetTypeReference/Invalidation/InvalidationBox.INT.udn
@@ -3,7 +3,7 @@ Title: Invalidiation Box
 Crumbs:%ROOT%, Engine, Engine/UMG, Engine/UMG/UserGuide, Engine/UMG/UserGuide/WidgetTypeReference
 Description:Overview of using the Invalidation Box for optimizing UI Widgets. 
 version: 4.9
-
+tags:UMG UI Designer
 
 [REGION:warning]
 This feature is still in development and is considered **Experimental**. As such, some features may or may not work as expected. 
diff --git a/Engine/Documentation/Source/Engine/UMG/UserGuide/WidgetTypeReference/WidgetTypeReference.INT.udn b/Engine/Documentation/Source/Engine/UMG/UserGuide/WidgetTypeReference/WidgetTypeReference.INT.udn
--- a/Engine/Documentation/Source/Engine/UMG/UserGuide/WidgetTypeReference/WidgetTypeReference.INT.udn
+++ b/Engine/Documentation/Source/Engine/UMG/UserGuide/WidgetTypeReference/WidgetTypeReference.INT.udn
@@ -1,12 +1,13 @@
 Availability:Public
 Title: Widget Type Reference
 Crumbs:%ROOT%, Engine, Engine/UMG, Engine/UMG/UserGuide
-Description:Reference for the types of Widgets available in Unreal Motion Graphics.
-version: 4.13
+Description:Reference for the types of Widgets available in Unreal Motion Graphics (UMG).
+version: 4.15
+tags:UMG UI Designer
 
 [TOC(start:2 end:2)]
 
-Inside the **Widget Blueprint Editor**, located under the **Palette** window, there are four categories of **Widgets**, each of which contain different Widgets Types that you can drag-and-drop into the **Visual Designer**. By mixing and matching these Widget Types, you can layout the look of your UI on the **Designer** tab and add the functionality behind your Widgets through the settings in the **Details** panel for each Widget and via the **Graph** tab. 
+Inside the **Widget Blueprint Editor**, located under the **Palette** window, there are several categories of **Widgets**, each of which contains different Widgets Types that you can drag-and-drop into the **Visual Designer**. By mixing and matching these Widget Types, you can layout the look of your UI on the **Designer** tab and add the functionality to your Widgets through the settings in the **Details** panel for each Widget and via the **Graph** tab. 
 
 Listed below are each of the Widget Types under the **Palette** window: 
 
@@ -15,7 +16,7 @@ Click the Widget Type (where applicable) for a usage guide along with more infor
 [/REGION]
 
 
-## Common Widgets
+## Common
 
 ![](Common.png)
 
@@ -24,17 +25,27 @@ The most frequently used Widgets will more than likely fall under this category.
 | Option | Description |
 | -------- | ----------- |
 | **Border** | A border is a container widget that can contain one child widget, providing an opportunity to surround it with a border image and adjustable padding. |
-| **Button** | The button is a single-child, click-able primitive widget which enables basic interaction. You can place any other widget inside a button to make a more complex and interesting click-able element in your UI.  |
-| **Check Box** | The check box widget allows you to display a toggled state of 'unchecked', 'checked', and 'indeterminable'. You can use the check box for a classic check box, as a toggle button, or as radio buttons.  |
-| [Named Slot](Engine\UMG\UserGuide\WidgetTypeReference\NamedSlot) | [INCLUDE:Engine\UMG\UserGuide\WidgetTypeReference\NamedSlot#desc] |
+| **Button** | The button is a single-child, click-able primitive widget which enables basic interaction. You can place any other widget inside a button to make a more complex and interesting clickable element in your UI.  |
+| **Check Box** | The Check Box widget allows you to display a toggled state of 'unchecked', 'checked', and 'indeterminable'. You can use the Check Box for a classic check box, as a toggle button, or as radio buttons.  |
+| **Image** | The image widget allows you to display a Slate Brush, Texture, Sprite, or Material in the UI. [REGION:note]Sprites that are part of the same texture atlas can be batched together in Slate, provided they all share the same layer when rendered. This means that platforms, where the draw call count budget is tight, will render UMG and Slate Widgets significantly more efficiently if Paper2D Sprites are used as Brush inputs.[/REGION] |
+| [**Named Slot**](Engine\UMG\UserGuide\WidgetTypeReference\NamedSlot) | [INCLUDE:Engine\UMG\UserGuide\WidgetTypeReference\NamedSlot#desc] |
 | **Progress Bar** | The progress bar widget is a simple bar that fills up and can be restyled to fit any number of uses such as Experience, Health, Points, etc. |
 | **Slider** | A simple widget that shows a sliding bar with a handle that allows you to control the value between 0-1. |
 | **Text** | The basic method for displaying text on screen and could be used for text descriptions of options or other UI elements. |
 | **Text Box** | Allows the user to type in custom text. Only permits a single line of text to be entered. |
-| **Image** | The image widget allows you to display a Slate Brush, Texture, Sprite, or Material in the UI. [REGION:note]Sprites that are part of the same texture atlas can be batched together in Slate, provided they all share the same layer when rendered. This means that platforms where the draw call count budget is tight will render UMG and Slate Widgets significantly more efficiently if Paper2D Sprites are used as Brush inputs.[/REGION] |
 
+## Extra
 
-## Input Widgets
+![](Extra.png)
+
+This section contains widgets that complement other widgets. 
+
+| Option | Description |
+| -------- | ----------- |
+| **Expandable Area** | Allows you to collapse or expand children widgets within a container. |
+
+
+## Input
 
 ![](Input.png)
 
@@ -47,17 +58,18 @@ There are a couple of options on how to allow Input from a user which are outlin
 | **Text Box (Multi-Line)** | Similar to the Text Box but allows a user to enter multiple lines of text instead of a single line. |
 
 
-## Optimization Widgets
+## Optimization
 
 ![](OptimizationWidgets.png)
 
-The widgets contained here are primiarly used for optimizating your UI for better performance. 
+The widgets contained here are primarily used for optimizing your UI for better performance. 
 
 | Option | Description |
 | -------- | ----------- |
 | [**Invalidation Box**](Engine\UMG\UserGuide\WidgetTypeReference\Invalidation) | [INCLUDE:Engine\UMG\UserGuide\WidgetTypeReference\Invalidation#desc]|
+| **Retainer Box** | Renders children widgets to a render target first before later rendering that render target to the screen. This allows both frequency and phase to be controlled so that the UI can actually render less often than the frequency of the main game render. It also has the side benefit of allowing materials to be applied to the render target after drawing the widgets to apply a simple post process. |
 
-## Panel Widgets
+## Panel
 
 ![](Panel.png)
 
@@ -65,20 +77,21 @@ The **Panel** category contains useful Widgets for controlling layouts and placi
 
 | Option | Description |
 | -------- | ----------- |
-| **Canvas Panel** | The canvas panel is a designer friendly panel that allows widgets to be laid out at arbitrary locations, anchored, and z-ordered with other children of the canvas. Note that, although z-order can be altered manually, widgets are rendered in list order, and it is considered preferable to order them correctly in the list rather than to rely on z-ordering. The canvas panel is a great widget for manual layout, but not as useful when you want to generate widgets procedurally and place them in a container (unless you want an absolute layout). For more information on Anchors, see the [Anchors](Engine\UMG\UserGuide\Anchors\) page. |
+| **Canvas Panel** | The canvas panel is a designer-friendly panel that allows widgets to be laid out at arbitrary locations, anchored, and z-ordered with other children of the canvas. Note that, although z-order can be altered manually, widgets are rendered in list order, and it is considered preferable to order them correctly in the list rather than to rely on z-ordering. The canvas panel is a great widget for manual layout, but not as useful when you want to generate widgets procedurally and place them in a container (unless you want an absolute layout). For more information on Anchors, see the [Anchors](Engine\UMG\UserGuide\Anchors\) page. |
 | **Grid Panel** | This is a panel that evenly divides up available space between all of its children widgets. |
 | **Horizontal Box** | Allows child widgets to be laid out in a flow horizontally. |
 | **Overlay** | Allows widgets to be stacked on top of each other and uses a simple flow layout for content on each layer.  |
 | **Safe Zone** | Pulls the platforms Safe Zone info and adds padding. |
 | **Scale Box** | Allows you to place content with a desired size and have it scale to meet the constraints placed on this box's allotted area. If you needed to have a background image scale to fill an area but not become distorted with different aspect rations, or if you need to auto fit some text to an area, this is the control for you. |
 | **Scroll Box** | An arbitrary scrollable collection of widgets. This is great for presenting 10-100 widgets in a list. This does not support virtualization. |
+| **Size Box** | Allows you to specify the size it reports to have and desire. Not all widgets report a desired size that you actually desire. Wrapping them in a Size Box lets you have the Size Box force them to be a particular size. |
 | **Uniform Grid Panel** | A panel that evenly divides up available space between all of its children. |
 | **Vertical Box** | A vertical box widget is a layout panel allowing child widgets to be automatically laid out vertically. This is great for stacking widgets on top of each other and keeping widgets aligned vertically.|
 | **Widget Switcher** | A widget switcher is like a tab control but without the tabs which you can create on your own and combine with this to get a tabbed effect. At most, one widget is visible at a time.  |
-| [Wrap Box](Engine\UMG\UserGuide\WidgetTypeReference\WrapBox) | [INCLUDE:Engine\UMG\UserGuide\WidgetTypeReference\WrapBox#desc] |
+| [**Wrap Box**](Engine\UMG\UserGuide\WidgetTypeReference\WrapBox) | [INCLUDE:Engine\UMG\UserGuide\WidgetTypeReference\WrapBox#desc] |
 
 
-## Primitive Widgets
+## Primitive
 
 ![](Primitive.png)
 
@@ -89,31 +102,53 @@ The **Widgets** contained in the **Primitive** category provide additional metho
 | **Circular Throbber** | A throbber widget that orients images in a spinning circle.  |
 | **Editable Text** | This a text field that allows for user input without a box background. This only supports a single line of editable text. |
 | **Editable Text (Multi-Line)** | Similar to Editable Text but supports multiple lines of text instead of a single line.|
-| [Menu Anchor](Engine\UMG\UserGuide\WidgetTypeReference\MenuAnchor) | [INCLUDE:Engine\UMG\UserGuide\WidgetTypeReference\MenuAnchor#desc] |
+| [**Menu Anchor**](Engine\UMG\UserGuide\WidgetTypeReference\MenuAnchor) | [INCLUDE:Engine\UMG\UserGuide\WidgetTypeReference\MenuAnchor#desc] |
 | **Native Widget Host** |This is a container widget that can contain one child Slate widget. This should be used when all you need is to nest a native widget inside a UMG widget. |
 | **Spacer** |A spacer widget provides custom padding between other widgets. The spacer does not have a visual representation and is invisible in game.|
 | **Throbber** | An animated throbber widget that shows several zooming circles in a row (could be used to signify loading for example). |
 
 
+## Special Effects
+
+![](SpecialEffects.png)
+
+These widgets are used to generate UI based special effects. 
+
+| Option | Description |
+| -------- | ----------- |
+| [**Background Blur**](Engine\UMG\UserGuide\WidgetTypeReference\BackgroundBlur) | [INCLUDE:Engine\UMG\UserGuide\WidgetTypeReference\BackgroundBlur#desc] |
+
+## Uncategorized
+
+![](Uncategorized.png)
+
+The widgets in this section are considered as a special case and do not fall in with other categories. 
+
+| Option | Description |
+| -------- | ----------- |
+| **Input Key Selector** | A widget for selecting a single key or a single key with a modifier. |
+| **Window Title Bar Area** | A panel for defining a region of the UI that should allow users to drag the window on desktop platforms. |
+
+
 ## User Created 
 
-![](UserCreatedWidgets.png)
+![](UserCreated.png)
 
 User Created Widgets are **Widget Blueprints** you have created that you can place inside another Widget Blueprint. This is useful for creating "parts" of UI elements as individual Widget Blueprints, then adding them together to construct your overall UI layout. 
 
 For example, you could create a _Player Health_ widget that contains display information for the player's health, an _Inventory_ widget that keeps track of the players collected items, an _Action Bar_ widget that allows the player to press a button and perform different actions or any number of other types of widgets, bringing them all together under one _GameHUD_ widget where those individual pieces are toggled on/off as needed (rather than constructing them all in a single Widget Blueprint). 
 
 
-
-## Experimental Widgets
+## Experimental 
 
 ![](ExperimentalPalette.png)
 
-These widgets are considered to be in the early stages of development and are **Experimental**. In order to access these widgets, you will need to enable them from the **Editor Preferences/Experimental** section or from the **Plugins** menu.  
+These widgets are considered to be in the early stages of development and in order to access them, you will need to enable them from the **Plugins** menu under **Widgets**. 
 
 | Option | Description |
 | -------- | ----------- |
-| [Web Browser](Engine\UMG\UserGuide\WidgetTypeReference\WebBrowser) | [INCLUDE:Engine\UMG\UserGuide\WidgetTypeReference\WebBrowser#desc] |
+| [**Web Browser**](Engine\UMG\UserGuide\WidgetTypeReference\WebBrowser) | [INCLUDE:Engine\UMG\UserGuide\WidgetTypeReference\WebBrowser#desc] |
+
 
 
 
diff --git a/Engine/Documentation/Source/Gameplay/ForceFeedback/ForceFeedback.INT.udn b/Engine/Documentation/Source/Gameplay/ForceFeedback/ForceFeedback.INT.udn
--- a/Engine/Documentation/Source/Gameplay/ForceFeedback/ForceFeedback.INT.udn
+++ b/Engine/Documentation/Source/Gameplay/ForceFeedback/ForceFeedback.INT.udn
@@ -1,8 +1,11 @@
 ﻿Availability:Public
-Title: Force Feedback
-Crumbs:
+Title: Force Feedback and the Force Feedback Component
+Crumbs:%ROOT%, Gameplay
 Description: Using the vibration functionality of mobile devices and controllers to convey a force occuring in the game to the player.
 Parent:Gameplay
+Related: Engine/Blueprints/UserGuide/Components
+version: 4.15
+
 
 [VAR:Topic]
 [OBJECT:Topic]
@@ -40,7 +43,9 @@ Parent:Gameplay
 
 [TOC]
 
+[EXCERPT:WhatIsForceFeedback]
 **Force Feedback** (sometimes referred to as rumble or haptic feedback) refers to the vibrating of a device or controller used in games to convey a force occuring in the game to the player. A simple example is when an explosion occurs in the game and force feedback is used to simulate the shockwave, which gives an additional dimension to the immersion.
+[/EXCERPT:WhatIsForceFeedback]
 
 Once implemented, Force Feedback will work on all platforms where it is supported. This includes iOS (as of 4.10), Android, and controllers.
 
@@ -86,6 +91,8 @@ Force Feedback Effect assets are created using the **Content Browser**:
 
 ## Playing Force Feedback
 
+### Directly To A Player
+
 Force Feedback is implemented in the base PlayerController class. You will need access to the local Player Controller in order to play the force feedback on the target device or controller.
 
 **Playing Force Feedback in Blueprints:**
@@ -116,6 +123,10 @@ Force Feedback is implemented in the base PlayerController class. You will need
 
 **Playing Force Feedback in C++:**
 
-Call [](API:APlayerController::ClientPlayForceFeedback) on the local Player Controller, passing it the Force Feedback Effect to use as well as whether or not the effect should loop, and an optional tag for the effect. The tag allows you to stop the effect; if an effect with the same tag is already playing, it will be stopped and the new effect will play instead.
+Call [](API:APlayerController::ClientPlayForceFeedback) on the local Player Controller, passing it the Force Feedback Effect to use as well as whether or not the effect should loop, and an optional tag for the effect. If a tag is provided, and another Force Feedback Effect with the same tag is played before the original effect ends, the original will be stopped immediately and the new effect will play instead.
+
+### At A World Location
 
+To play a Force Feedback Effect that changes intensity based on distance from the observing player, place a **Force Feedback Component** in the world at the intended source of the effect. A Force Feedback component plays a Force Feedback Effect on command, but also has a physical location in the world. Like sound or light, the intensity of the force experienced by the player will change with the player's distance from the source, according to a data-defined attenuation setting.
 
+Force Feedback Components can be attached to any actor via the usual code or Blueprint editor methods, but they can also be added as they are needed during live gameplay. This can be accomplished by calling the utility functions `UGameplayStatics::SpawnForceFeedbackAtLocation` (to spawn at a given world location), or `UGameplayStatics::SpawnForceFeedbackAttached` (to attach to a specific pre-existing component). These functions will return the spawned Force Feedback Component, so you can continue to manipulate it if you like. However, if you have no use for the component after it finishes playing the effect, you can use the "Auto Destroy" option to have it remove itself once the effect ends.
diff --git a/Engine/Documentation/Source/Gameplay/Input/Input.INT.udn b/Engine/Documentation/Source/Gameplay/Input/Input.INT.udn
--- a/Engine/Documentation/Source/Gameplay/Input/Input.INT.udn
+++ b/Engine/Documentation/Source/Gameplay/Input/Input.INT.udn
@@ -2,7 +2,10 @@ Availability: Public
 Title:Input
 Crumbs:%ROOT%, Gameplay
 Description:The Input object is responsible for converting input from the player into data in a form Actors can understand and make use of.
-version: 4.9
+version: 4.15
+skilllevel:Beginner
+Parent:Engine/Gameplay
+type:overview
 
 
 The **PlayerInput** object is responsible for converting input from the player into data in a form that Actors, like PlayerControllers or Pawns, 
@@ -15,7 +18,7 @@ For an example of setting up Input, refer to the [](Gameplay\HowTo\SetUpInput) d
 [TOC (start:2 end:4)]
 
 ## Hardware Input
-Hardware input from a player is very straightforward.  It most commonly includes key presses, mouse clicks or mouse movement, and controller button presses or joystick movement. 
+Hardware input from a player is very straightforward.  It most commonly includes key presses, mouse clicks or mouse movement, and controller button presses or joystick movement. Specialized input devices that don't conform to standard axis or button indices, or that have unusual input ranges, can be configured manually by using the [](Gameplay\Input\RawInput).
 
 
 ## PlayerInput
diff --git a/Engine/Documentation/Source/Gameplay/Input/RawInput/RawInput.INT.udn b/Engine/Documentation/Source/Gameplay/Input/RawInput/RawInput.INT.udn
new file mode 100644
--- /dev/null
+++ b/Engine/Documentation/Source/Gameplay/Input/RawInput/RawInput.INT.udn
@@ -0,0 +1,50 @@
+Availability:Docs
+Title:RawInput Plugin
+Crumbs:%ROOT%, Gameplay, Gameplay/Input
+Description: The RawInput plugin can read specific devices that don't conform to XInput standards.
+version: 4.15
+skilllevel:Intermediate
+Parent:Engine/Gameplay/Input
+type:overview
+
+## RawInput Plugin
+
+The ***RawInput** plugin provides support for specific, user-defined devices that aren't properly handled by XInput, usually flight sticks and steering wheels. The buttons and axes of these input devices can be custom-mapped to gameplay inputs, including new inputs created by the plugin itself.
+
+### Enabling RawInput
+
+By default, the plugin is not enabled. It can be enabled from the **Plugins** menu, in the **Input Devices** section.
+
+![RawInput plugin](RawInputPlugin.png)
+
+Once the plugin is enabled, the **Raw Input** section will appear in the **Project Settings** menu.
+
+![RawInput location in Project Settings menu](RawInput.png)
+
+### Configuring Devices
+
+Within this section, you can add new devices. Each device is identified by a vendor ID and a product ID, which are hexadecimal values that can be found in the device's driver properties. An array of axes (by default, eight) and buttons (by default, twelve) will be added to the controller's configuration. This provides the ability to map any of the controller's axes or buttons (indicated by the array index) to any **Unreal Engine** axis or key. The plugin provides eight new axes ("Generic USB Controller Axis 1" through "Generic USB Controller Axis 8") and twelve new buttons ("Generic USB Controller Button 1" through "Generic USB Controller Button 12"). Each axis and button can be enabled or disabled, and axes can also be inverted and offset.
+
+![Windows Device Settings](WindowsSettings.png)
+
+_Above, the Hardware ID information in Windows Control Panel (under Device Manager) for the Logitech G920. The Vendor ID is the four-character string following "VID\_" test (046D), and the Product ID is the four-character string following "PID\_" (C262). This information will be required when configuring the device in RawInput._
+
+### Example (Vehicle Templates)
+As an example, the Vehicle templates that ship with the **Unreal Engine** have been configured to work with the **Logitech G920**, a driving-game controller with a steering wheel and gas/brake pedals.
+
+#### Configuring the Device in RawInput
+The steering wheel is reported by the driver as Axis 1, the brake is reported as Axis 3, and the gas pedal is reported as Axis 4. In the RawInput plugin settings, each input must be entered into the array at the corresponding axis from the device driver. That is, the steering wheel's data must be in array element 1, the brake's data must be in array element 3, and the gas pedal's data must be in array element 4. All other entries may be removed or marked as disabled.
+
+[REGION:note]Note that the array index of an input does not need to correspond to the "Generic USB Controller Axis" or "Generic USB Controller Input" value. For example, the brake is bound to "Generic USB Controller Axis 2" even though it is in array position 3. This is necessary to ensure that different controllers can be configured to work the same way, even if they use different input axes.[/REGION]
+
+![RawInput settings](RawInputData.png)
+
+Because the device's output range doesn't match what we might expect from a standard controller, we will remap the input ranges. The wheel returns values from 0.0 (left) to 1.0 (right), but we want the input range to center on 0.0, so we will give it an offset of -0.5. The brake and gas pedals also return values from 0.0 to 1.0, but the device uses 0.0 to indicate that the pedal is pressed down, and 1.0 to indicate that the pedal is up, where our project expects the reverse. To make this adjustment, we set the axis values to be inverted, and then add a 1.0 offset.
+
+#### Mapping our RawInput Axes to Input Bindings
+
+![Input settings](InputData.png)
+
+We can then map the new axes to our game's input bindings. The wheel (on "Generic USB Controller Axis 1") is mapped to the range [-0.5, 0.5], but we want it to be more sensitive, so we can scale it up by a factor of 3.0. Similarly, the brake (on "Generic USB Controller Axis 2") needs to go in the negative direction, and should be more powerful than the gas pedal, so it will be scaled by -2.0. The gas pedal (on "Generic USB Controller Axis 3") requires no modification, and only needs to be added to the appropriate input binding. By using RawInput and adding these input bindings, we are able to support the Logitech G920 without needing any new input bindings or making any changes to our project's code or blueprints.
+
+[REGION:tip]This use of scaling provides a good argument for using the "Generic USB Controller" axes instead of configuring the device to the standard gamepad stick inputs.[/REGION]
diff --git a/Engine/Documentation/Source/Gameplay/Localization/Commandlets/Commandlets.INT.udn b/Engine/Documentation/Source/Gameplay/Localization/Commandlets/Commandlets.INT.udn
--- a/Engine/Documentation/Source/Gameplay/Localization/Commandlets/Commandlets.INT.udn
+++ b/Engine/Documentation/Source/Gameplay/Localization/Commandlets/Commandlets.INT.udn
@@ -1,10 +1,12 @@
 Availability:Public
-Title: Commandlets
+Title: Localization Commandlets
 Crumbs:%ROOT%, Gameplay, Gameplay/Localization
 Description:Descriptions and arguments for the commandlets used when localizing your project. 
 version: 4.9
 Parent:Gameplay/Localization
 Order:4
+Tags:Localization
+Type:Reference
 
 
 [TOC (start:2)]
diff --git a/Engine/Documentation/Source/Gameplay/Localization/Localization.INT.udn b/Engine/Documentation/Source/Gameplay/Localization/Localization.INT.udn
--- a/Engine/Documentation/Source/Gameplay/Localization/Localization.INT.udn
+++ b/Engine/Documentation/Source/Gameplay/Localization/Localization.INT.udn
@@ -11,7 +11,7 @@ version: 4.13
 
 ## Localizing Audio
 
-Spoken text is now gathered from [**Dialogue Wave**](Engine/Audio/Overview#dialoguevoiceanddialoguewave) assets and included in your PO files for translation along with your other asset based text.
+Spoken text is now gathered from [Dialogue Wave](Engine/Audio/Overview#dialoguevoiceanddialoguewave) assets and included in your PO files for translation along with your other asset based text.
 
 Several new [commandlets](Gameplay/Localization/Commandlets) have been added to facilitate the recording and importing of audio for your **Dialogue Wave** assets.
 * ExportDialogueScript - Exports a per-culture CSV file containing each line of dialogue and its associated filename.
diff --git a/Engine/Documentation/Source/Gameplay/Networking/Blueprints/Blueprints.INT.udn b/Engine/Documentation/Source/Gameplay/Networking/Blueprints/Blueprints.INT.udn
--- a/Engine/Documentation/Source/Gameplay/Networking/Blueprints/Blueprints.INT.udn
+++ b/Engine/Documentation/Source/Gameplay/Networking/Blueprints/Blueprints.INT.udn
@@ -23,11 +23,12 @@ To add multiplayer functionality to your game, it's important to understand the
 
  See the [](Gameplay/Framework) documentation for more information - but at least keep in mind the following tips when designing your game for multiplayer are:
 
-* The GameMode object only exists on the server.
-* The GameState exists on the server and the clients, so the server can use replicated variables on the GameState to keep all clients up-to-date on data about the game.
-* A PlayerController exists on the server for every player connected to the game. On clients, only the PlayerControllers for players owned by that client exist. This means that PlayerControllers are not the best place to store replicated properties that all connected clients are interested in. Instead, use the PlayerState.
-* A PlayerState will exist for every player connected to the game on both the server and the clients. This class can be used for replicated properties that all clients, not just the owning client, are interested in, such as the individual player's current score.
-* Pawns (including Characters) will also exist on the server and on all clients, and can also contain replicated variables and events. The decision of whether to use the PlayerState or the Pawn for a certain variable or event will depend on the situation, but the main thing to keep in mind is that the PlayerState will persist for the entire time the player is connected, but a `Pawn` may not. For example, if a player character dies during gameplay, the Pawn he was controlling may be destroyed and a new one created when the player respawns.
+* The GameInstance exists for the duration of the engine's session, meaning that it is created when the engine starts up and not destroyed or replaced until the engine shuts down. A separate GameInstance exists on the server and on each client, and these instances do not communicate with each other. Because the GameInstance exists outside of the game session and is the only game structure that exists across level loads, it is a good place to store certain types of persistent data, such as lifetime player statistics (e.g. total number of games won), account information (e.g. locked/unlocked status of special items), or even a list of maps to rotate through in a competitive game like Unreal Tournament.
+* The GameMode object only exists on the server. It generally stores information related to the game that clients do not need to know explicitly. For example, if a game has special rules like "rocket launchers only", the clients may not need to know this rule, but when randomly spawning weapons around the map, the server needs to know to pick only from the "rocket launcher" category.
+* The GameState exists on the server and the clients, so the server can use replicated variables on the GameState to keep all clients up-to-date on data about the game. Information that is of interest to all players and spectators, but isn't associated with any one specific player, is ideal for GameState replication. As an example, a baseball game could replicate each team's score and the current inning via the GameState.
+* One PlayerController exists on each client per player on that machine. These are replicated between the server and the associated client, but are not replicated to other clients, resulting in the server having PlayerControllers for every player, but local clients having only the PlayerControllers for their local players. PlayerControllers exist while the client is connected, and are associated with Pawns, but are not destroyed and respawned like Pawns often are. They are well-suited to communicating information between clients and servers without replicating this information to other clients, such as the server telling the client to ping its minimap in response to a game event that only that player detects.
+* A PlayerState will exist for every player connected to the game on both the server and the clients. This class can be used for replicated properties that all clients, not just the owning client, are interested in, such as the individual player's current score in a free-for-all game. Like the PlayerController, they are associated with individual Pawns, and are not destroyed and respawned when the Pawn is.
+* Pawns (including Characters) will also exist on the server and on all clients, and can contain replicated variables and events. The decision of whether to use the PlayerController, the PlayerState, or the Pawn for a certain variable or event will depend on the situation, but the main thing to keep in mind is that the PlayerController and PlayerState will persist as long as the owning player stays connected to the game and the game doesn't load a new level, whereas a `Pawn` may not. For example, if a Pawn dies during gameplay, it will usually be destroyed and replaced with a new Pawn, while the PlayerController and PlayerState will continue to exist and will be associated with the new Pawn once it finishes spawning. The Pawn's health, therefore, would be stored on the Pawn itself, since that is specific to the actual instance of the Pawn and should be reset when the Pawn is replaced with a new one.
 
 ## Actor replication
 
diff --git a/Engine/Documentation/Source/Gameplay/Networking/Overview/Overview.INT.udn b/Engine/Documentation/Source/Gameplay/Networking/Overview/Overview.INT.udn
--- a/Engine/Documentation/Source/Gameplay/Networking/Overview/Overview.INT.udn
+++ b/Engine/Documentation/Source/Gameplay/Networking/Overview/Overview.INT.udn
@@ -5,14 +5,24 @@ Description:Setting up networked games for multiplayer.
 Related: Gameplay/Networking/Blueprints
 Related: Gameplay/Networking/Example
 Related: Gameplay/Networking/Actors
+Related: Gameplay/Networking/Server
 version: 4.9
 
 The UE4 framework is built with multiplayer gaming in mind. As long as you follow the basic framework conventions, you generally don't have to do much to extend a single player experience to multiplayer.
 
-UE4 networking is built around the server/client model. This means that there will be one server that is authoritative (makes all the important decisions), and this server will then make sure all connected clients are continually updated so that they maintain the most up to date approximation of the server's world.
+UE4 networking is built around the server/client model. This means that there will be one server that is authoritative (makes all the important decisions), and this server will then make sure all connected clients are continually updated so that they maintain the most up to date approximation of the server's world. Even non-networked, single-player games have a server; the local machine acts as the server in these cases.
 
 ## Actors
 
 In a single player experience, actors are used to represent the game world. In multiplayer, this isn't any different, except for the fact that clients will have an approximation of each actor, while the server will maintain the authoritative version.
 
 Actors are the main workhorse that the server uses to keep clients up to date. When it is time for the server to update a particular client, the server will gather all relevant actors that it determines have changed since the last update, and then sends enough information to the client to catch those actors up to to the latest state.
+
+## Network Modes
+
+| Network Mode Type | Functionality / Ideal Use Case |
+| --- | --- |
+| `NM_Standalone` | This type indicates a server running on a local machine and not accepting clients from remote machines. This is best for single-player, or local multiplayer, games. |
+| `NM_DedicatedServer` | A dedicated server has no local players and can run more efficiently by discarding sound, graphics, user input, and other player-oriented features. This server type is used for multiplayer games hosted on a trusted server, such as competitive MOBA titles, MMO games, or online shooters where high-performing, reliable servers are needed. |
+| `NM_ListenServer` | A listen server is a server that hosts a local player, but is open to connections from remote players as well. This type of server is good for competitive or cooperative games where a dedicated server is unnecessary, and users are able to set up and play their own games without a third-party server. This network mode type may give the host some advantage due to the host's lack of network latency, and the game can be terminated without warning by the host. |
+| `NM_Client` | This is the only mode that is not a server. When in this mode, the local machine is a client to a dedicated or listen server, and therefore will not run server-side logic. |
diff --git a/Engine/Documentation/Source/Gameplay/Networking/Server/Server.INT.udn b/Engine/Documentation/Source/Gameplay/Networking/Server/Server.INT.udn
--- a/Engine/Documentation/Source/Gameplay/Networking/Server/Server.INT.udn
+++ b/Engine/Documentation/Source/Gameplay/Networking/Server/Server.INT.udn
@@ -32,7 +32,7 @@ Dedicated servers are headless by default. If you don't use -log, you won't see
 
 ## Server Gameplay Flow
 
-The server is in charge is driving the flow of gameplay. It is the server's job to notify clients that it is time to travel to a new map, when gameplay starts and ends, along with actor replication updates, etc.
+The server is in charge of driving the flow of gameplay. It is the server's job to notify clients that it is time to travel to a new map, when gameplay starts and ends, along with actor replication updates, etc.
 
 The [major framework portions](Gameplay/Framework) are mostly out of the scope of this document, but we can talk about some of the additional multiplayer nuances introduced when dealing working with certain classes.
 
diff --git a/Engine/Documentation/Source/Gameplay/Tags/Tags.INT.udn b/Engine/Documentation/Source/Gameplay/Tags/Tags.INT.udn
new file mode 100644
--- /dev/null
+++ b/Engine/Documentation/Source/Gameplay/Tags/Tags.INT.udn
@@ -0,0 +1,100 @@
+Availability:Docs
+Title:Gameplay Tags
+Crumbs: %ROOT%, Engine, Engine/Gameplay
+Description: Gameplay Tags can be used to identify, categorize, match, and filter objects.
+version: 4.15
+skilllevel:Intermediate
+Parent:Engine/Gameplay
+type:overview
+
+## Gameplay Tags
+**Gameplay Tags** are conceptual, hierarchical labels with user-defined names. These tags can have any number of hierarchical levels, separated by the "." character; for example, a Gameplay Tag with three levels would take the form of "Family.Genus.Species", with "Family" being the broadest identifier in the hierarchy, and "Species" being the most specific. Note that the existence of "Family.Genus.Species" implicitly means that "Family.Genus" and "Family" Gameplay Tags also exist. Individual tags are lightweight, and use the type `FGameplayTag`. Since game objects frequently have multiple tags, **Gameplay Tag Containers** (using the type `FGameplayTagContainer`) are provided, and can handle a variety of additional queries. Gameplay Tags (or Gameplay Tag Containers) can be added to any object to associate it with conceptual labels that can then be used to identify, match, categorize, or filter as needed for your project.
+
+### Creating Gameplay Tags
+Gameplay Tags must be added to the central tag dictionary for the engine to be aware of them. This can be done in three ways via editor functionality. Tags can be added manually in the **Project Settings** menu, with .ini files, through creating data table assets. All three methods are set up under **Project Settings**, by opening the **Gameplay Tags** tab under the **Project** section.
+
+![Gameplay Tags menu in Project Settings](ProjectSettings.png)
+
+#### Manually Adding Tags
+
+The easiest way to add gameplay tags is to use enter them manually in the **Project Settings** menu. By checking the **Import Tags From Config** option, `Config/DefaultGameplayTags.Ini` will be loaded, as well as all Gameplay Tags in any .Ini files found under the `Config/Tags` path. An option labeled **Add New Gameplay Tag** will appear, and this option can be used to add new Gameplay Tags to any existing Gameplay Tag .Ini files, including the default file. Tags added in this way will also have an optional field for entering a brief description, which will appear as a tooltip in the editor.
+
+![Manually Adding Tags in Project Settings](AddNewTag.png)
+
+#### Editing .Ini Files Directly
+To add tags using your own .Ini files, the **Import Tags From Config** option must be enabled. Gameplay Tags will then be loaded from each of the .Ini files found your project's `Config/Tags/` path, which are easy to hand-edit. If these edits are made while the editor is running, turning the **Import Tags From Config** option off and back on will reload the files. The format within these files looks like this:
+
+	[UserTags]
+	GameplayTags=Vehicle.Air.Helicopter
+	GameplayTags=Movement.Flying
+	GameplayTags=Movement.Hovering
+
+If a hand-written file is modified by the editor's **Add New Gameplay Tag** feature, all of the old tags, as well as the new tag, will be copied into a new section. The "[UserTags]" section will remain in the file, but will be ignored during load operations. For example, if "Movement.Landing" were added to our file above, the new file would look like this:
+
+	[UserTags]
+	GameplayTags=Vehicle.Air.Helicopter
+	GameplayTags=Movement.Flying
+	GameplayTags=Movement.Hovering
+
+	[/Script/GameplayTags.GameplayTagsList]
+	GameplayTagList=(Tag="Vehicle.Air.Helicopter",DevComment="")
+	GameplayTagList=(Tag="Movement.Flying",DevComment="")
+	GameplayTagList=(Tag="Movement.Hovering",DevComment="")
+	GameplayTagList=(Tag="Movement.Landing",DevComment="")
+
+[REGION:note]It is not recommended to hand-edit a tag .Ini file after using the **Project Settings** menu to add tags to that file.[/REGION]
+
+This method is very convenient as a way to enable members of your team to add their own tags, since .Ini files in `Config/Tags` are easily-edited text files that can be named with any file-naming convention you like.
+
+A "Developer Tag" config file can be used by going to the **Gameplay Tags Developer** menu within **Project Settings** and filling in the **Developer Config Name** with your desired filename. If set, this will save all tags you personally add to the default .Ini file to your specified .Ini file instead. Your .Ini file will reside in the `Config/Tags` folder, and can be passed to other developers or submitted to your source control system without further modification. This can be useful on large projects with thousands of tags, or simply as a way to separate tags created by specific developers or for use with specific game features. This is entirely optional, but it can be helpful for some teams and workflows to have certain tags stored in separate files.
+
+#### Data Table Assets
+Finally, tags can be added via [DataTable](Gameplay/DataDriven/#datatables) assets with the row type `GameplayTagTableRow`. This is especially useful for importing Gameplay Tags from an external source, such as an Excel spreadsheet file. Once your asset is created, add it to the **GameplayTagTableList** in the **Project Settings** menu and all tags contained within it will be added to the **Gameplay Tag Manager**. Note that this system supports multiple assets being listed, so you can separate out your project's tags into different spreadsheets if this helps your workflow or organization.
+
+![Gameplay Tags in a Data Table Asset](DataAsset.png)
+
+### Managing Gameplay Tags
+Once your Gameplay Tags have been added, you can manage them by searching for references to them, deleting them, or renaming them. All of these actions can be accessed with the caret dropdown next to a tag in the **Project Settings** menu. Tags can only be deleted in this menu if they were added via a .Ini file (either hand-edited or added via the **Project Settings** menu), and are not referenced by anything else. Renamed tags will cause additions to be made to the GameplayTagRedirects list, and any renamed tags will be silently fixed on load.
+
+#### Gameplay Tag Test Operations
+There are a number of test operations, all based on the concept of matching tags to other tags, that can be performed on Gameplay Tags and Gameplay Tag Containers. Each of these functions will be called on a Gameplay Tag or Gameplay Tag Container, and will take a single Gameplay Tag or Gameplay Tag Container as a parameter.
+
+The following table details the operations available for tags and tag containers. In this table, the format "A.1" represents a single Gameplay Tag, while "{A.1}" represents that same tag as part of a Gameplay Tag Container. A container with multiple tags will be of the format "{A.1, B.1}".
+
+| Source | Function/Operation | Input Parameter | Return Value |
+| --- | --- | --- | --- |
+| A.1 | MatchesTag | A | true |
+| A.1 | MatchesTagExact (or ==) | A | false |
+| A.1 | MatchesAny | {A, C} | true |
+| A.1 | MatchesAnyExact | {A, C} | false |
+| A.1 | MatchesAll | {A} | true |
+| A.1 | MatchesAllExact | {A} | false |
+| {A.1, B.1} | HasTag | A | true |
+| {A.1, B.1} | HasTagExact | A | false |
+| {A.1, B.1} | HasAny | {A, C} | true |
+| {A.1, B.1} | HasAnyExact | {A, C} | false |
+| {A.1, B.1} | HasAll | {A, B} | true |
+| {A.1, B.1} | HasAllExact | {A, B} | false |
+
+Note that providing empty or default Gameplay Tag Containers as the input parameter will cause all operations to return false, except for `HasAll`, `HasAllExact`, `MatchesAll`, and `MatchesAllExact`. The justification for this is that there are no tags in the parameter container that aren't in the source set, and this is consistent with set theory.
+
+### Gameplay Tag Queries
+Certain types of queries for Gameplay Tag Containers can be stored in a **Gameplay Tag Query** variable in order to make them data-driven, perform them in many places without having to edit more than one variable or piece of code, and speed them up. Gameplay Tag Queries are limited to three basic tests:
+**Any Tags Match** succeeds if at least one tag from the query is found in the container.
+**All Tags Match** succeeds if there are no tags in the query that are not also found in the container. (This includes the case that no tags are present in the query.)
+**No Tags Match** succeeds if there are no tags in the query that are also found in the container. (This includes the case that no tags are present in the query.)
+However, those tests can be combined together via three larger expressions, each of which contains an array of additional expressions or tests:
+**Any Expressions Match** succeeds when any of its sub-expressions or tests succeeds.
+**All Expressions Match** succeeds when none of its sub-expressions or tests fail. (This includes the case that no sub-expressions or tests are present in the query.)
+**No Expressions Match** succeeds when none of its sub-expressions or tests succeed. (This includes the case that no sub-expressions or tests are present in the query.)
+[REGION:note]Note that all three test types refer to "exact" matching as per the table of operations above.[/REGION]
+
+![Sample Gameplay Tag Query](TagQuery.png)
+
+In the above sample query, either TestTag.One or TestTag.Three (but not both) must be in the source container, and TestTag.Two must not be.
+
+### Using Gameplay Tags In Your Game
+To apply the tags you have created to objects in your game, add Gameplay Tags (C++ type `FGameplayTag`) or Gameplay Tag Containers (C++ type `FGameplayTagContainer`) properties to the objects. You can also add **Gameplay Tag Queries** (C++ type `FGameplayTagQuery`) to your objects as desired, or you can use the tag operations directly in Blueprints or C++ code. Once these properties have been added, you can edit them from the code or editor, like any other variable.
+
+### Accessing Gameplay Tags Quickly (Requires C++)
+The `IGameplayTagAssetInterface` can be implemented in C++. Doing so, and overriding the `GetOwnedGameplayTags` function, will create a Blueprint-accessible way to populate a Gameplay Tag Container with the tags associated with that object. In most cases, this just means copying the tags from the base class into a new container, but your implementation of the class could gather tags from multiple containers, call a Blueprint function to get access to Blueprint-declared tags, or whatever your object requires. The advantage of this interface is that it enables you to interact with a variety of disparate object types in a standardized way, so that you can get an object's tags without explicitly casting the object and then writing custom code for each possible type to get the correct tag variable (or variables). This is especially useful on larger projects, where it can help to avoid hard-to-reproduce bugs and make code cleaner, faster, and easier to maintain.
diff --git a/Engine/Documentation/Source/GettingStarted/DownloadingUnrealEngine/Perforce/Basics/PerforceBasics.INT.udn b/Engine/Documentation/Source/GettingStarted/DownloadingUnrealEngine/Perforce/Basics/PerforceBasics.INT.udn
--- a/Engine/Documentation/Source/GettingStarted/DownloadingUnrealEngine/Perforce/Basics/PerforceBasics.INT.udn
+++ b/Engine/Documentation/Source/GettingStarted/DownloadingUnrealEngine/Perforce/Basics/PerforceBasics.INT.udn
@@ -1,7 +1,8 @@
-﻿Availability:NoPublish
+﻿Availability:Licensee
 Title:Downloading Unreal Engine with Perforce
-Crumbs:%ROOT%, GettingStarted, GettingStarted/DownloadingUnrealEngine, GettingStarted/DownloadingUnrealEngine/Perforce
+Crumbs:
 Description:Instructions for downloading Unreal Engine source code with Perforce.
+Parent:GettingStarted/DownloadingUnrealEngine/Perforce
 
 <!--
 ## Connect to Epic's VPN
@@ -34,9 +35,20 @@ Connecting to Epic perforce requires utilizing the SSL feature, and you must be
 
 	* **User**: Perforce username provided by Epic Games.
 	* **Password**: Perforce password provided by Epic Games.
-	* **Workspace**: _UE4_Sample_TeamName_ (This is a recommendation. You can choose any workspace name you like as long as it ends with your team name.)
 
 1.	Press OK to connect to the Perforce Server.
+	
+1. In P4V, choose **Connection > New Workspace** to create a new workspace for the engine. Enter the information below and click **OK** to create the workspace:
+
+	* **Workspace name**: _UE4_TeamName_ (This is a recommendation. You can choose any workspace name you like as long as it ends with your team name.)
+	* **Stream**: Click **Browse** and select _//UE4/Release-Latest_ from the list of available streams.
+
+	![](new_workspace.png)
+
+1. In the Depot pane, click the **Filter Depot** button and select **Tree Restricted to Workspace View**.
+
+	![](filter_depot.png)
+
 
 [REGION:note]
 Epic maintains a consistent certificate fingerprint for all broker and proxy SSL endpoints, however Perforce requires a trust per IP address + fingerprint pair. When connecting to an endpoint for the first time, you must explicitly trust that endpoint. The Epic fingerprint is 45:0D:78:E2:0E:9E:E4:82:45:80:16:36:29:5E:54:4D:66:31:6C:43.
@@ -70,28 +82,22 @@ If you are running a local proxy, you  must connect through a broker instead of
 
 ## Download Unreal Engine 4
 
-We distribute Unreal Engine 4 to licensees via the _UE4-QA_ branch in the Perforce depot. This contains the entire 
+We distribute Unreal Engine 4 to licensees via the _//UE4/Release-Latest_ stream in the Perforce depot. This contains the entire 
 engine along with several additional projects in the form of example games, samples, and demos. You have the option 
 of downloading everything or picking and choosing only the parts you want or need. In order to get set up as quickly 
 as possible, we recommend you only download the bare minimum to start with and then download other parts on an as-needed 
 basis. This can dramatically reduce idle time spent waiting for the download to finish.
 
 [REGION:warning]
-	There are a large number of files in the _UE4-QA_ branch and the total download size is many gigabytes so the 
+	There are a large number of files in the _//UE4/Release-Latest_ stream and the total download size is many gigabytes so the 
 	download can take quite a while when syncing the entire branch.
 [/REGION]
 
-1.	**Right-click** on the _UE4-QA_ branch and choose **Get Revision**.
-
-	![Perforce - Get Revision Option](get_revision_menu.png)
-
-1.	Select the **Specify revision using** radio button and choose Label from the combobox.
-
-	![Perforce - Specify Label](specify_label.png)
+1.	**Right-click** on the _UE4_ stream and choose **Get Latest Revision**.
 
-1.	Press the **Browse** button and select either *QA_APPROVED_UE4_BUILD_MINIMUM* (minimum files needed to run the engine) or *QA_APPROVED_UE4_BUILD_CURRENT* (includes sample projects and is much larger) in the dialog that opens. Press **OK**.
+	![Perforce - Get Latest Revision](get_latest_menu.png)
 
-1.	Press **Get Revision** to begin downloading Unreal Engine 4. 
+1.	The latest version of all files will be downloaded.
 
 [/EXCERPT:main]
 
diff --git a/Engine/Documentation/Source/GettingStarted/FromUnity/UnrealEngineforUnityDevelopers.INT.udn b/Engine/Documentation/Source/GettingStarted/FromUnity/UnrealEngineforUnityDevelopers.INT.udn
--- a/Engine/Documentation/Source/GettingStarted/FromUnity/UnrealEngineforUnityDevelopers.INT.udn
+++ b/Engine/Documentation/Source/GettingStarted/FromUnity/UnrealEngineforUnityDevelopers.INT.udn
@@ -48,7 +48,7 @@ The following section contains common Unity terms on the left and their UE4 equi
 | **2D**             | Sprite Editor | [Paper2D](Engine/Paper2D) |
 | **Programming**    | C# | [C++](Programming) |
 |                    | Script | [Blueprint](Engine/Blueprints) |
-| **Physics**        | Raycast | [Line Trace, Shape Trace](Gameplay/HowTo/UseRaycasts/Blueprints) |
+| **Physics**        | Raycast | [Line Trace, Shape Trace](Engine/Physics/Tracing) |
 |                    | Rigid Body | [Collision, Physics](Engine/Physics) |
 | **Runtime Platforms** | iOS Player, Web Player | [Platforms](Platforms) |
 
diff --git a/Engine/Documentation/Source/GettingStarted/GettingStarted.INT.udn b/Engine/Documentation/Source/GettingStarted/GettingStarted.INT.udn
--- a/Engine/Documentation/Source/GettingStarted/GettingStarted.INT.udn
+++ b/Engine/Documentation/Source/GettingStarted/GettingStarted.INT.udn
@@ -281,6 +281,8 @@ gettingstarted
 		or artists to use.
 	[/PARAM]
 	[PARAM:links]
+		* [](Programming "%Programming:description%")
+		* [](Programming/Development/VisualStudioSetup "%Programming/Development/VisualStudioSetup:description%")
 		* [](Programming/QuickStart "%Programming/QuickStart:description%")
 		* [](Programming/Basics "%Programming/Basics:description%")
 		* [Learn About the Game Framework](Gameplay/Framework/QuickReference "%Gameplay/Framework/QuickReference:description%")
diff --git a/Engine/Documentation/Source/GettingStarted/HowTo/DetailsPanel/DetailsPanel.INT.udn b/Engine/Documentation/Source/GettingStarted/HowTo/DetailsPanel/DetailsPanel.INT.udn
--- a/Engine/Documentation/Source/GettingStarted/HowTo/DetailsPanel/DetailsPanel.INT.udn
+++ b/Engine/Documentation/Source/GettingStarted/HowTo/DetailsPanel/DetailsPanel.INT.udn
@@ -4,9 +4,8 @@ Crumbs: %ROOT%, Gameplay, GettingStarted\HowTo
 Description:Modifying properties of an Actor such as applying Materials, Physics, Scale through the **Details** panel. 
 Version: 4.10
 SkillLevel: Beginner
-parent:Engine/Editor
+parent:GettingStarted/HowTo
 order:1
-checkpoint: editorqs
 
 Once you place an [Actor](Engine/Actors) in your level, you may want to modify the way it looks or functions in your level. With the Actor selected in the level, inside the [Details](Engine/UI/LevelEditor/Details) panel you will find information about the Actor that can be modified. This How to will show you some sample use cases for accessing the Details panel and modifying an Actor's properties.
 
diff --git a/Engine/Documentation/Source/GettingStarted/Installation/InstallingUnrealEngine.INT.udn b/Engine/Documentation/Source/GettingStarted/Installation/InstallingUnrealEngine.INT.udn
--- a/Engine/Documentation/Source/GettingStarted/Installation/InstallingUnrealEngine.INT.udn
+++ b/Engine/Documentation/Source/GettingStarted/Installation/InstallingUnrealEngine.INT.udn
@@ -54,19 +54,25 @@ With a valid Epic Games account, you're able to download our installer (Epic Gam
 	active_button_content
 	[/PARAMLITERAL]
 	[PARAM:content]
-	1.	**Click** the **Download** button to get our installer.
+	[INCLUDE:#windows_install]
+	[/PARAM]
+[/OBJECT]
+
+<!--
+[EXCERPT:windows_install]
+1.	**Click** the **Download** button to get our installer.
 	
-		![](Download_Screen.png)
+	![](Download_Screen.png)
 	
-	1.	Once the download is complete, **Run** the installer.
+1.	Once the download is complete, **Run** the installer.
 	
-		![](RunInstaller.png)
+	![](RunInstaller.png)
 	
-	1.	**Click** the **Install** button when the **Epic Games Launcher Setup** dialog box appears.
+1.	**Click** the **Install** button when the **Epic Games Launcher Setup** dialog box appears.
 	
-		![](RunEpicLauncherSetup.png)
-	[/PARAM]
-[/OBJECT]
+	![](RunEpicLauncherSetup.png)
+[/EXCERPT:windows_install]
+-->
 
 [OBJECT:ToggleButtonContent]
 	[PARAMLITERAL:category]
diff --git a/Engine/Documentation/Source/GettingStarted/Terminology/Terminology.INT.udn b/Engine/Documentation/Source/GettingStarted/Terminology/Terminology.INT.udn
--- a/Engine/Documentation/Source/GettingStarted/Terminology/Terminology.INT.udn
+++ b/Engine/Documentation/Source/GettingStarted/Terminology/Terminology.INT.udn
@@ -166,17 +166,17 @@ For more information, see: [](Engine/Levels), [](Engine/UI/LevelEditor), [](Reso
 [EXCERPT:gamemode]
 **GameModes**
 [INCLUDE:Shared\Glossary\G#glossary_gamemode] 
-For more information, see: [](Gameplay/Framework/GameMode), [](Gameplay/HowTo/SettingUpAGameMode/Blueprints) 
+For more information, see: [GameMode](Gameplay/Framework/GameMode), [](Gameplay/HowTo/SettingUpAGameMode/Blueprints) 
 [/EXCERPT:gamemode]
 
 
 [EXCERPT:gamestate]
 **GameStates**
 [INCLUDE:Shared\Glossary\G#glossary_gamestate] 
-For more information, see: [](Gameplay/Framework/GameState)
+For more information, see: [GameState](Gameplay/Framework/GameMode#gamestate)
 [/EXCERPT:gamestate]
 
-
+ 
 [EXCERPT:playerstate]
 **PlayerStates**
 [INCLUDE:Shared\Glossary\P#glossary_playerstate] 
diff --git a/Engine/Documentation/Source/GettingStartedPDF/GettingStartedPDF.INT.udn b/Engine/Documentation/Source/GettingStartedPDF/GettingStartedPDF.INT.udn
--- a/Engine/Documentation/Source/GettingStartedPDF/GettingStartedPDF.INT.udn
+++ b/Engine/Documentation/Source/GettingStartedPDF/GettingStartedPDF.INT.udn
@@ -2,7 +2,7 @@ Availability:Docs
 Title: Getting Started with Unreal Engine
 Crumbs:
 DoIndex:false
-Description:Getting started with Unreal Engine, PDF Generator
+Description:Getting started with Unreal Engine
 Template:PDF.html
 
 [REGION:banner]
diff --git a/Engine/Documentation/Source/Modding/Modding.INT.udn b/Engine/Documentation/Source/Modding/Modding.INT.udn
new file mode 100644
--- /dev/null
+++ b/Engine/Documentation/Source/Modding/Modding.INT.udn
@@ -0,0 +1,18 @@
+Availability: Public
+Title:Making Mods
+Crumbs:%ROOT%
+Description:Guides and resources over modding games made with Unreal Engine 4
+Parent:%ROOT%
+type:landing
+order:11
+tags:Modding
+topic-image:modding_Topic.png
+social-image:modding_Social.png
+topic-icon:%ROOT%/mod_icon.png
+
+![](modding.png)
+
+Below you will find all of the official documentation for modding Unreal Engine 4 powered games, that support modding, specifically from Epic Games.
+
+[DIR(end:"1" output:"topic" org:"hierarchy" parent:"Modding")]
+
diff --git a/Engine/Documentation/Source/Modding/RoboRecall/HowTo/EnableMods/EnableAMod.INT.udn b/Engine/Documentation/Source/Modding/RoboRecall/HowTo/EnableMods/EnableAMod.INT.udn
new file mode 100644
--- /dev/null
+++ b/Engine/Documentation/Source/Modding/RoboRecall/HowTo/EnableMods/EnableAMod.INT.udn
@@ -0,0 +1,20 @@
+Availability: Public
+Title:Enabling a Mod
+Crumbs:%ROOT%
+Description:How to make the mods you have installed for Robo Recall active in the game
+Parent:Modding/RoboRecall
+type:how-to
+order:3
+tags:Modding
+tags:Robo Recall
+topic-image:mod_enabled_topic.png
+social-image:Modding/RoboRecall/RoboRecall_Social.png
+
+
+1. Teleport to the holostation to open up the menu and select **MODS**.
+
+	![](menu_mods.png)
+
+1. In the Mods menu, find the mod you want to enable and select it to enable it.
+
+	![](menu_mod_select.png)(h:280 convert:false) ![](menu_mod_enabled.png)(h:280 convert:false)
