Copying //UE4/Dev-Documentation to //UE4/Samples-Main (Source: //UE4/Dev-Documentation @ 3372893)

diff --git a/Engine/Documentation/Source/Modding/RoboRecall/HowTo/ExtendFunctionality/AddVisualEffects/AddVisualEffects.INT.udn b/Engine/Documentation/Source/Modding/RoboRecall/HowTo/ExtendFunctionality/AddVisualEffects/AddVisualEffects.INT.udn
new file mode 100644
--- /dev/null
+++ b/Engine/Documentation/Source/Modding/RoboRecall/HowTo/ExtendFunctionality/AddVisualEffects/AddVisualEffects.INT.udn
@@ -0,0 +1,352 @@
+Availability: Public
+Title:6. Adding Visual Effects
+Crumbs:%ROOT%
+Description:Robo Recall Modding Documentation
+Parent:Modding/RoboRecall/HowTo/ExtendFunctionality
+type:how-to
+order:6
+tags:Modding
+tags:Robo Recall
+topic-image:Modding/RoboRecall/RoboRecall_Topic.png
+social-image:Modding/RoboRecall/RoboRecall_Social.png
+checkpoint:mod_extend
+
+<!-- Location: Modding/RoboRecall/Quickstarts/WeaponMod/6 -->
+
+[NAV]
+
+<!-- Title Image if Needed Here -->
+
+In this step you will be finishing off the GrenadeGunPistol Blueprint by adding a particle system, sound, a radial force, and finally cleaning up the gun after the explosion.
+
+## Steps
+
+1. Using the techniques learned in previous sections, create a **Add Radial Force Component** node.
+
+	![image alt text](image_0.png)(w:916 h:472 convert:false)
+
+2. It needs a transform, so create a **Make Transform** node and connect its output to the **Relative Transform** pin on the **Add Radial Force Component**.
+
+	![image alt text](image_1.png)(w:934 h:354 convert:false)
+
+3. Select the **Add Radial Force Component** and change **Impulse Strength** to `650.0`. Enable **Impulse Vel Change**, Enable **Ignore Owning Actor**, and **Force Strength** to `0.0`.
+
+	![image alt text](image_2.png)(w:485 h:502 convert:false)
+
+4. Also, using the **+** sign on the **Object Types to Affect**, add a new entry and change it to **WorldDynamic**. This will allow the force to hit the bot ragdolls.
+
+	![image alt text](image_3.png)(w:485 h:528 convert:false)
+
+5. Next, pull off the **Return Value** on the **Add Radial Force Component**, search for and create a **Fire Impulse** node.
+
+	![image alt text](image_4.png)(w:728 h:461 convert:false)
+
+6. From the output execution pin of that node, add a **Spawn Emitter at Location** node.
+
+	![image alt text](image_5.png)(w:661 h:457 convert:false)
+
+7. Set its **Emitter Template** to `Boss_Rocket_Explo`.
+
+	![image alt text](image_6.png)(w:480 h:600 convert:false)
+
+8. From the output execution pin of that node, add a **Spawn Sound at Location** node.
+
+9. Set its **Sound** to `GrendaeExplosion_C_Cue`
+
+	![image alt text](image_7.png)(w:662 h:242 convert:false)
+
+10. From the output execution pin of that node, add a **Set Ammo Remaining** node, which should come in set to 0.
+
+11. From the output execution pin of that node, add a **Force Cleanup** node, which will dissolve and destroy the gun.
+
+	![image alt text](image_8.png)(w:765 h:201 convert:false)
+
+12. Create are a **Self** node (search for **Get a reference to self**). Pull off **Return Value** and create a a **Get Actor Location** node.
+
+	![image alt text](image_9.png)(w:934 h:293 convert:false)
+
+13. Connect the output of the **Get Actor Location** node to both the **Spawn Emitter** and **Spawn Sound at Location** nodes.
+
+	![image alt text](image_10.png)(w:934 h:289 convert:false)
+
+14. Below you can see the full graph.
+
+	[REGION:tip]
+	The graph below has been cleaned up using **re-route nodes**. These are essentially pass-through nodes that enable you to better organize the connections in a graph and make the graph easier to read.
+	[/REGION]
+
+	[OBJECT:Snippet]
+		[PARAMLITERAL:id]
+		result
+		[/PARAMLITERAL]
+		[PARAMLITERAL:type]
+		blueprint
+		[/PARAMLITERAL]
+		[PARAM:image]
+		[REGION:lightbox]
+		[![image alt text](image_11.png)(convert:false)](image_11.png)
+		[/REGION]
+		[/PARAM]
+		[PARAM:snippet]
+			[INCLUDE:#snippet]
+		[/PARAM]
+	[/OBJECT]
+
+	[REGION:caption]
+	Click the icon in the upper left corner of this image to copy the Blueprint Graph and paste it into your project.
+	[/REGION]
+
+## Results
+
+The GrenadeGunPistol is now complete. In the final step, you'll go over how to test the gun and ideas on where to take it next.
+
+[NAV]
+
+<!--
+[EXCERPT:snippet]
+	Begin Object Class=K2Node_Self Name="K2Node_Self_16"
+
+		NodePosX=2000
+
+		NodePosY=832
+
+		NodeGuid=97EA67D3438CF5F20BD2B7A9042C0980
+
+		CustomProperties Pin (PinId=100C9C8248AB8362E8F82792E9679624,PinName="self",Direction="EGPD_Output",PinType.PinCategory="object",PinType.PinSubCategory="self",PinType.PinSubCategoryObject=None,PinType.PinSubCategoryMemberReference=(),PinType.PinValueType=(),PinType.bIsMap=False,PinType.bIsSet=False,PinType.bIsArray=False,PinType.bIsReference=False,PinType.bIsConst=False,PinType.bIsWeakPointer=False,LinkedTo=(K2Node_CallFunction_116 6B6418B4475566B13E3FA9A11F29C747,),PersistentGuid=00000000000000000000000000000000,bHidden=False,bNotConnectable=False,bDefaultValueIsReadOnly=False,bDefaultValueIsIgnored=False,bAdvancedView=False,)
+
+	End Object
+
+	Begin Object Class=K2Node_CallFunction Name="K2Node_CallFunction_116"
+
+		bIsPureFunc=True
+
+		bIsConstFunc=True
+
+		FunctionReference=(MemberName="K2_GetActorLocation",bSelfContext=True)
+
+		NodePosX=2160
+
+		NodePosY=800
+
+		NodeGuid=B51CED234F6148D4F5EB1E88974DC6BC
+
+		CustomProperties Pin (PinId=6B6418B4475566B13E3FA9A11F29C747,PinName="self",PinFriendlyName=NSLOCTEXT("K2Node", "Target", "Target"),PinToolTip="Target\nActor Reference",PinType.PinCategory="object",PinType.PinSubCategory="",PinType.PinSubCategoryObject=Class'/Script/Engine.Actor',PinType.PinSubCategoryMemberReference=(),PinType.PinValueType=(),PinType.bIsMap=False,PinType.bIsSet=False,PinType.bIsArray=False,PinType.bIsReference=False,PinType.bIsConst=False,PinType.bIsWeakPointer=False,LinkedTo=(K2Node_Self_16 100C9C8248AB8362E8F82792E9679624,),PersistentGuid=00000000000000000000000000000000,bHidden=False,bNotConnectable=False,bDefaultValueIsReadOnly=False,bDefaultValueIsIgnored=False,bAdvancedView=False,)
+
+		CustomProperties Pin (PinId=C149106E4CC3C4B43D809AB6A6CC5465,PinName="ReturnValue",PinToolTip="Return Value\nVector ",Direction="EGPD_Output",PinType.PinCategory="struct",PinType.PinSubCategory="",PinType.PinSubCategoryObject=ScriptStruct'/Script/CoreUObject.Vector',PinType.PinSubCategoryMemberReference=(),PinType.PinValueType=(),PinType.bIsMap=False,PinType.bIsSet=False,PinType.bIsArray=False,PinType.bIsReference=False,PinType.bIsConst=False,PinType.bIsWeakPointer=False,AutogeneratedDefaultValue="0, 0, 0",LinkedTo=(K2Node_CallFunction_119 23AAB45947B8ACB24C4DDAAF5B1F2F44,K2Node_Knot_548 D8C5606C4E0103BD1B30C1909156FE54,),PersistentGuid=00000000000000000000000000000000,bHidden=False,bNotConnectable=False,bDefaultValueIsReadOnly=False,bDefaultValueIsIgnored=False,bAdvancedView=False,)
+
+	End Object
+
+	Begin Object Class=K2Node_CallFunction Name="K2Node_CallFunction_119"
+
+		FunctionReference=(MemberParent=Class'/Script/Engine.GameplayStatics',MemberName="SpawnEmitterAtLocation")
+
+		NodePosX=2448
+
+		NodePosY=640
+
+		NodeGuid=B43BDA3744EC576BD057ECB741F7BDA2
+
+		CustomProperties Pin (PinId=D1E0CD9A44C8258C668EAFA1A4486990,PinName="execute",PinToolTip="\nExec",PinType.PinCategory="exec",PinType.PinSubCategory="",PinType.PinSubCategoryObject=None,PinType.PinSubCategoryMemberReference=(),PinType.PinValueType=(),PinType.bIsMap=False,PinType.bIsSet=False,PinType.bIsArray=False,PinType.bIsReference=False,PinType.bIsConst=False,PinType.bIsWeakPointer=False,LinkedTo=(K2Node_CallFunction_77 AEC053B24B8500ADC9AF948F39A660AF,),PersistentGuid=00000000000000000000000000000000,bHidden=False,bNotConnectable=False,bDefaultValueIsReadOnly=False,bDefaultValueIsIgnored=False,bAdvancedView=False,)
+
+		CustomProperties Pin (PinId=08F736C8421ECE55DA8E58B0DE9216A2,PinName="then",PinToolTip="\nExec",Direction="EGPD_Output",PinType.PinCategory="exec",PinType.PinSubCategory="",PinType.PinSubCategoryObject=None,PinType.PinSubCategoryMemberReference=(),PinType.PinValueType=(),PinType.bIsMap=False,PinType.bIsSet=False,PinType.bIsArray=False,PinType.bIsReference=False,PinType.bIsConst=False,PinType.bIsWeakPointer=False,LinkedTo=(K2Node_CallFunction_855 C3A6A74B4E2B5249FF82DAB24058ED08,),PersistentGuid=00000000000000000000000000000000,bHidden=False,bNotConnectable=False,bDefaultValueIsReadOnly=False,bDefaultValueIsIgnored=False,bAdvancedView=False,)
+
+		CustomProperties Pin (PinId=376023054BD07D773E91A893E4AC9054,PinName="self",PinFriendlyName=NSLOCTEXT("K2Node", "Target", "Target"),PinToolTip="Target\nGameplay Statics Reference",PinType.PinCategory="object",PinType.PinSubCategory="",PinType.PinSubCategoryObject=Class'/Script/Engine.GameplayStatics',PinType.PinSubCategoryMemberReference=(),PinType.PinValueType=(),PinType.bIsMap=False,PinType.bIsSet=False,PinType.bIsArray=False,PinType.bIsReference=False,PinType.bIsConst=False,PinType.bIsWeakPointer=False,DefaultObject="/Script/Engine.Default__GameplayStatics",PersistentGuid=00000000000000000000000000000000,bHidden=True,bNotConnectable=False,bDefaultValueIsReadOnly=False,bDefaultValueIsIgnored=False,bAdvancedView=False,)
+
+		CustomProperties Pin (PinId=11F20FC747AF7A7B193F0396D81C4689,PinName="WorldContextObject",PinToolTip="World Context Object\nObject Reference\n\nObject that we can obtain a world context from",PinType.PinCategory="object",PinType.PinSubCategory="",PinType.PinSubCategoryObject=Class'/Script/CoreUObject.Object',PinType.PinSubCategoryMemberReference=(),PinType.PinValueType=(),PinType.bIsMap=False,PinType.bIsSet=False,PinType.bIsArray=False,PinType.bIsReference=False,PinType.bIsConst=True,PinType.bIsWeakPointer=False,PersistentGuid=00000000000000000000000000000000,bHidden=True,bNotConnectable=False,bDefaultValueIsReadOnly=False,bDefaultValueIsIgnored=False,bAdvancedView=False,)
+
+		CustomProperties Pin (PinId=99E48D0443F8B0C8BACA70B75AF2FF7A,PinName="EmitterTemplate",PinToolTip="Emitter Template\nParticle System Reference\n\nparticle system to create",PinType.PinCategory="object",PinType.PinSubCategory="",PinType.PinSubCategoryObject=Class'/Script/Engine.ParticleSystem',PinType.PinSubCategoryMemberReference=(),PinType.PinValueType=(),PinType.bIsMap=False,PinType.bIsSet=False,PinType.bIsArray=False,PinType.bIsReference=False,PinType.bIsConst=False,PinType.bIsWeakPointer=False,DefaultObject="/Game/FX/Character/Enemies/Boss/Boss_Rocket_Explo.Boss_Rocket_Explo",PersistentGuid=00000000000000000000000000000000,bHidden=False,bNotConnectable=False,bDefaultValueIsReadOnly=False,bDefaultValueIsIgnored=False,bAdvancedView=False,)
+
+		CustomProperties Pin (PinId=23AAB45947B8ACB24C4DDAAF5B1F2F44,PinName="Location",PinToolTip="Location\nVector \n\nlocation to place the effect in world space",PinType.PinCategory="struct",PinType.PinSubCategory="",PinType.PinSubCategoryObject=ScriptStruct'/Script/CoreUObject.Vector',PinType.PinSubCategoryMemberReference=(),PinType.PinValueType=(),PinType.bIsMap=False,PinType.bIsSet=False,PinType.bIsArray=False,PinType.bIsReference=False,PinType.bIsConst=False,PinType.bIsWeakPointer=False,DefaultValue="0, 0, 0",AutogeneratedDefaultValue="0, 0, 0",LinkedTo=(K2Node_CallFunction_116 C149106E4CC3C4B43D809AB6A6CC5465,),PersistentGuid=00000000000000000000000000000000,bHidden=False,bNotConnectable=False,bDefaultValueIsReadOnly=False,bDefaultValueIsIgnored=False,bAdvancedView=False,)
+
+		CustomProperties Pin (PinId=B94B5F734B70AB5FEAD10F861F1BEA44,PinName="Rotation",PinToolTip="Rotation\nRotator \n\nrotation to place the effect in world space",PinType.PinCategory="struct",PinType.PinSubCategory="",PinType.PinSubCategoryObject=ScriptStruct'/Script/CoreUObject.Rotator',PinType.PinSubCategoryMemberReference=(),PinType.PinValueType=(),PinType.bIsMap=False,PinType.bIsSet=False,PinType.bIsArray=False,PinType.bIsReference=False,PinType.bIsConst=False,PinType.bIsWeakPointer=False,DefaultValue="0, 0, 0",AutogeneratedDefaultValue="0, 0, 0",PersistentGuid=00000000000000000000000000000000,bHidden=False,bNotConnectable=False,bDefaultValueIsReadOnly=False,bDefaultValueIsIgnored=False,bAdvancedView=False,)
+
+		CustomProperties Pin (PinId=0F7AD8314643C10CF5D1F4B257BD4AD3,PinName="bAutoDestroy",PinToolTip="Auto Destroy\nBoolean\n\nWhether the component will automatically be destroyed when the particle system completes playing or whether it can be reactivated",PinType.PinCategory="bool",PinType.PinSubCategory="",PinType.PinSubCategoryObject=None,PinType.PinSubCategoryMemberReference=(),PinType.PinValueType=(),PinType.bIsMap=False,PinType.bIsSet=False,PinType.bIsArray=False,PinType.bIsReference=False,PinType.bIsConst=False,PinType.bIsWeakPointer=False,DefaultValue="true",AutogeneratedDefaultValue="true",PersistentGuid=00000000000000000000000000000000,bHidden=False,bNotConnectable=False,bDefaultValueIsReadOnly=False,bDefaultValueIsIgnored=False,bAdvancedView=False,)
+
+		CustomProperties Pin (PinId=FE279D714D8E7ADDE07C399A4C13D445,PinName="ReturnValue",PinToolTip="Return Value\nParticle System Component Reference",Direction="EGPD_Output",PinType.PinCategory="object",PinType.PinSubCategory="",PinType.PinSubCategoryObject=Class'/Script/Engine.ParticleSystemComponent',PinType.PinSubCategoryMemberReference=(),PinType.PinValueType=(),PinType.bIsMap=False,PinType.bIsSet=False,PinType.bIsArray=False,PinType.bIsReference=False,PinType.bIsConst=False,PinType.bIsWeakPointer=False,PersistentGuid=00000000000000000000000000000000,bHidden=False,bNotConnectable=False,bDefaultValueIsReadOnly=False,bDefaultValueIsIgnored=False,bAdvancedView=False,)
+
+	End Object
+
+	Begin Object Class=K2Node_VariableSet Name="K2Node_VariableSet_7"
+
+		VariableReference=(MemberName="AmmoRemaining",bSelfContext=True)
+
+		NodePosX=3136
+
+		NodePosY=656
+
+		NodeGuid=20A2DE8D4C3EA6C437754A8026BA22D8
+
+		CustomProperties Pin (PinId=AD945712496DF669C2B3C089FF4AEED2,PinName="execute",PinType.PinCategory="exec",PinType.PinSubCategory="",PinType.PinSubCategoryObject=None,PinType.PinSubCategoryMemberReference=(),PinType.PinValueType=(),PinType.bIsMap=False,PinType.bIsSet=False,PinType.bIsArray=False,PinType.bIsReference=False,PinType.bIsConst=False,PinType.bIsWeakPointer=False,LinkedTo=(K2Node_CallFunction_855 DA8628F44B19D85719E13A8E89CCEEFC,),PersistentGuid=00000000000000000000000000000000,bHidden=False,bNotConnectable=False,bDefaultValueIsReadOnly=False,bDefaultValueIsIgnored=False,bAdvancedView=False,)
+
+		CustomProperties Pin (PinId=F3FBF530454D488375093CB2D98EA58F,PinName="then",Direction="EGPD_Output",PinType.PinCategory="exec",PinType.PinSubCategory="",PinType.PinSubCategoryObject=None,PinType.PinSubCategoryMemberReference=(),PinType.PinValueType=(),PinType.bIsMap=False,PinType.bIsSet=False,PinType.bIsArray=False,PinType.bIsReference=False,PinType.bIsConst=False,PinType.bIsWeakPointer=False,LinkedTo=(K2Node_CallFunction_390 B99493564082E838045608A10AB4EAC8,),PersistentGuid=00000000000000000000000000000000,bHidden=False,bNotConnectable=False,bDefaultValueIsReadOnly=False,bDefaultValueIsIgnored=False,bAdvancedView=False,)
+
+		CustomProperties Pin (PinId=F91D5926437556C558DD198AF11A1778,PinName="AmmoRemaining",PinType.PinCategory="int",PinType.PinSubCategory="",PinType.PinSubCategoryObject=None,PinType.PinSubCategoryMemberReference=(),PinType.PinValueType=(),PinType.bIsMap=False,PinType.bIsSet=False,PinType.bIsArray=False,PinType.bIsReference=False,PinType.bIsConst=False,PinType.bIsWeakPointer=False,DefaultValue="0",AutogeneratedDefaultValue="0",PersistentGuid=00000000000000000000000000000000,bHidden=False,bNotConnectable=False,bDefaultValueIsReadOnly=False,bDefaultValueIsIgnored=False,bAdvancedView=False,)
+
+		CustomProperties Pin (PinId=B518031A4BBA575685F47588B853A38B,PinName="Output_Get",PinToolTip="Retrieves the value of the variable, can use instead of a separate Get node",Direction="EGPD_Output",PinType.PinCategory="int",PinType.PinSubCategory="",PinType.PinSubCategoryObject=None,PinType.PinSubCategoryMemberReference=(),PinType.PinValueType=(),PinType.bIsMap=False,PinType.bIsSet=False,PinType.bIsArray=False,PinType.bIsReference=False,PinType.bIsConst=False,PinType.bIsWeakPointer=False,DefaultValue="0",AutogeneratedDefaultValue="0",PersistentGuid=00000000000000000000000000000000,bHidden=False,bNotConnectable=False,bDefaultValueIsReadOnly=False,bDefaultValueIsIgnored=False,bAdvancedView=False,)
+
+		CustomProperties Pin (PinId=D190D47540AFF9F2E806E59D5032FFB7,PinName="self",PinFriendlyName=NSLOCTEXT("K2Node", "Target", "Target"),PinType.PinCategory="object",PinType.PinSubCategory="",PinType.PinSubCategoryObject=Class'/Script/Odin.OdinWeapon',PinType.PinSubCategoryMemberReference=(),PinType.PinValueType=(),PinType.bIsMap=False,PinType.bIsSet=False,PinType.bIsArray=False,PinType.bIsReference=False,PinType.bIsConst=False,PinType.bIsWeakPointer=False,PersistentGuid=00000000000000000000000000000000,bHidden=True,bNotConnectable=False,bDefaultValueIsReadOnly=False,bDefaultValueIsIgnored=False,bAdvancedView=False,)
+
+	End Object
+
+	Begin Object Class=K2Node_CallFunction Name="K2Node_CallFunction_390"
+
+		FunctionReference=(MemberName="ForceCleanup",MemberGuid=4DDCE77F4BB4E54FD0C2C19F85AE3622,bSelfContext=True)
+
+		NodePosX=3392
+
+		NodePosY=624
+
+		NodeGuid=B482ABF94923428747DD2D80B7DED0D4
+
+		CustomProperties Pin (PinId=B99493564082E838045608A10AB4EAC8,PinName="execute",PinType.PinCategory="exec",PinType.PinSubCategory="",PinType.PinSubCategoryObject=None,PinType.PinSubCategoryMemberReference=(),PinType.PinValueType=(),PinType.bIsMap=False,PinType.bIsSet=False,PinType.bIsArray=False,PinType.bIsReference=False,PinType.bIsConst=False,PinType.bIsWeakPointer=False,LinkedTo=(K2Node_VariableSet_7 F3FBF530454D488375093CB2D98EA58F,),PersistentGuid=00000000000000000000000000000000,bHidden=False,bNotConnectable=False,bDefaultValueIsReadOnly=False,bDefaultValueIsIgnored=False,bAdvancedView=False,)
+
+		CustomProperties Pin (PinId=5F96A5A74F6628D7F19304BD1C568D80,PinName="then",Direction="EGPD_Output",PinType.PinCategory="exec",PinType.PinSubCategory="",PinType.PinSubCategoryObject=None,PinType.PinSubCategoryMemberReference=(),PinType.PinValueType=(),PinType.bIsMap=False,PinType.bIsSet=False,PinType.bIsArray=False,PinType.bIsReference=False,PinType.bIsConst=False,PinType.bIsWeakPointer=False,PersistentGuid=00000000000000000000000000000000,bHidden=False,bNotConnectable=False,bDefaultValueIsReadOnly=False,bDefaultValueIsIgnored=False,bAdvancedView=False,)
+
+		CustomProperties Pin (PinId=634A5CF14141BDAC25E1D78251E66AD1,PinName="self",PinFriendlyName=NSLOCTEXT("K2Node", "Target", "Target"),PinType.PinCategory="object",PinType.PinSubCategory="",PinType.PinSubCategoryObject=BlueprintGeneratedClass'/Game/Prototype/Game/Gun_Base.Gun_Base_C',PinType.PinSubCategoryMemberReference=(),PinType.PinValueType=(),PinType.bIsMap=False,PinType.bIsSet=False,PinType.bIsArray=False,PinType.bIsReference=False,PinType.bIsConst=False,PinType.bIsWeakPointer=False,PersistentGuid=00000000000000000000000000000000,bHidden=False,bNotConnectable=False,bDefaultValueIsReadOnly=False,bDefaultValueIsIgnored=False,bAdvancedView=False,)
+
+	End Object
+
+	Begin Object Class=K2Node_AddComponent Name="K2Node_AddComponent_0"
+
+		TemplateBlueprint="/GrenadeGun/Blueprints/Weapons/GrenadeGunPistol.GrenadeGunPistol"
+
+		TemplateType=Class'/Script/Engine.RadialForceComponent'
+
+		FunctionReference=(MemberName="AddComponent",bSelfContext=True)
+
+		NodePosX=1840
+
+		NodePosY=624
+
+		NodeGuid=F0E9A3BB49521851B89C54AE2BD70612
+
+		CustomProperties Pin (PinId=E212548F4EB1DBC875747DB35AF1E319,PinName="execute",PinToolTip="\nExec",PinType.PinCategory="exec",PinType.PinSubCategory="",PinType.PinSubCategoryObject=None,PinType.PinSubCategoryMemberReference=(),PinType.PinValueType=(),PinType.bIsMap=False,PinType.bIsSet=False,PinType.bIsArray=False,PinType.bIsReference=False,PinType.bIsConst=False,PinType.bIsWeakPointer=False,LinkedTo=(K2Node_CallFunction_3 3A6AA71340DEF59A3541578052CDA34E,),PersistentGuid=00000000000000000000000000000000,bHidden=False,bNotConnectable=False,bDefaultValueIsReadOnly=False,bDefaultValueIsIgnored=False,bAdvancedView=False,)
+
+		CustomProperties Pin (PinId=02DBF18F43959E2CD4F3148EE7470FA3,PinName="then",PinToolTip="\nExec",Direction="EGPD_Output",PinType.PinCategory="exec",PinType.PinSubCategory="",PinType.PinSubCategoryObject=None,PinType.PinSubCategoryMemberReference=(),PinType.PinValueType=(),PinType.bIsMap=False,PinType.bIsSet=False,PinType.bIsArray=False,PinType.bIsReference=False,PinType.bIsConst=False,PinType.bIsWeakPointer=False,LinkedTo=(K2Node_CallFunction_77 9C4DCA9D48357011B95546B420FEE798,),PersistentGuid=00000000000000000000000000000000,bHidden=False,bNotConnectable=False,bDefaultValueIsReadOnly=False,bDefaultValueIsIgnored=False,bAdvancedView=False,)
+
+		CustomProperties Pin (PinId=3EC4E8F546212ABAEB8F6CA736697C0B,PinName="self",PinFriendlyName=NSLOCTEXT("K2Node", "Target", "Target"),PinToolTip="Target\nActor Reference",PinType.PinCategory="object",PinType.PinSubCategory="",PinType.PinSubCategoryObject=Class'/Script/Engine.Actor',PinType.PinSubCategoryMemberReference=(),PinType.PinValueType=(),PinType.bIsMap=False,PinType.bIsSet=False,PinType.bIsArray=False,PinType.bIsReference=False,PinType.bIsConst=False,PinType.bIsWeakPointer=False,PersistentGuid=00000000000000000000000000000000,bHidden=False,bNotConnectable=False,bDefaultValueIsReadOnly=False,bDefaultValueIsIgnored=False,bAdvancedView=False,)
+
+		CustomProperties Pin (PinId=F89E8D154A1CF3F6935E14922F72985C,PinName="TemplateName",PinToolTip="Template Name\nName\n\nThe name of the Component Template to use.",PinType.PinCategory="name",PinType.PinSubCategory="",PinType.PinSubCategoryObject=None,PinType.PinSubCategoryMemberReference=(),PinType.PinValueType=(),PinType.bIsMap=False,PinType.bIsSet=False,PinType.bIsArray=False,PinType.bIsReference=False,PinType.bIsConst=False,PinType.bIsWeakPointer=False,DefaultValue="NODE_AddRadialForceComponent-0",AutogeneratedDefaultValue="None",PersistentGuid=00000000000000000000000000000000,bHidden=True,bNotConnectable=True,bDefaultValueIsReadOnly=True,bDefaultValueIsIgnored=False,bAdvancedView=False,)
+
+		CustomProperties Pin (PinId=A415A70F43EF1D6C120EB6BF865B1970,PinName="bManualAttachment",PinToolTip="Manual Attachment\nBoolean\n\nWhether manual or automatic attachment is to be used",PinType.PinCategory="bool",PinType.PinSubCategory="",PinType.PinSubCategoryObject=None,PinType.PinSubCategoryMemberReference=(),PinType.PinValueType=(),PinType.bIsMap=False,PinType.bIsSet=False,PinType.bIsArray=False,PinType.bIsReference=False,PinType.bIsConst=False,PinType.bIsWeakPointer=False,DefaultValue="false",AutogeneratedDefaultValue="false",PersistentGuid=00000000000000000000000000000000,bHidden=False,bNotConnectable=False,bDefaultValueIsReadOnly=False,bDefaultValueIsIgnored=False,bAdvancedView=False,)
+
+		CustomProperties Pin (PinId=4483A3B74C12C1530E84C49F92B991A8,PinName="RelativeTransform",PinToolTip="Relative Transform\nTransform \n\nThe relative transform between the new component and its attach parent (automatic only)",PinType.PinCategory="struct",PinType.PinSubCategory="",PinType.PinSubCategoryObject=ScriptStruct'/Script/CoreUObject.Transform',PinType.PinSubCategoryMemberReference=(),PinType.PinValueType=(),PinType.bIsMap=False,PinType.bIsSet=False,PinType.bIsArray=False,PinType.bIsReference=False,PinType.bIsConst=True,PinType.bIsWeakPointer=False,LinkedTo=(K2Node_CallFunction_76 AF359A7549B7FC5151F55B9B5C7306DD,),PersistentGuid=00000000000000000000000000000000,bHidden=False,bNotConnectable=False,bDefaultValueIsReadOnly=False,bDefaultValueIsIgnored=True,bAdvancedView=False,)
+
+		CustomProperties Pin (PinId=8E7920984BBFAE694B4C5F952CC4BAC6,PinName="ComponentTemplateContext",PinToolTip="Component Template Context\nObject Reference\n\nOptional UBlueprintGeneratedClass reference to use to find the template in. If null (or not a BPGC), component is sought in this Actor\'s class",PinType.PinCategory="object",PinType.PinSubCategory="",PinType.PinSubCategoryObject=Class'/Script/CoreUObject.Object',PinType.PinSubCategoryMemberReference=(),PinType.PinValueType=(),PinType.bIsMap=False,PinType.bIsSet=False,PinType.bIsArray=False,PinType.bIsReference=False,PinType.bIsConst=True,PinType.bIsWeakPointer=False,PersistentGuid=00000000000000000000000000000000,bHidden=True,bNotConnectable=True,bDefaultValueIsReadOnly=False,bDefaultValueIsIgnored=False,bAdvancedView=False,)
+
+		CustomProperties Pin (PinId=86731E5E4566B2DBFB777890155E4768,PinName="ReturnValue",PinToolTip="Return Value\nRadial Force Component Reference",Direction="EGPD_Output",PinType.PinCategory="object",PinType.PinSubCategory="",PinType.PinSubCategoryObject=Class'/Script/Engine.RadialForceComponent',PinType.PinSubCategoryMemberReference=(),PinType.PinValueType=(),PinType.bIsMap=False,PinType.bIsSet=False,PinType.bIsArray=False,PinType.bIsReference=False,PinType.bIsConst=False,PinType.bIsWeakPointer=False,LinkedTo=(K2Node_CallFunction_77 3DA9F1F1472BF28C89AAE083A5111D9B,),PersistentGuid=00000000000000000000000000000000,bHidden=False,bNotConnectable=False,bDefaultValueIsReadOnly=False,bDefaultValueIsIgnored=False,bAdvancedView=False,)
+
+	End Object
+
+	Begin Object Class=K2Node_CallFunction Name="K2Node_CallFunction_76"
+
+		bIsPureFunc=True
+
+		FunctionReference=(MemberParent=Class'/Script/Engine.KismetMathLibrary',MemberName="MakeTransform")
+
+		NodePosX=1504
+
+		NodePosY=736
+
+		NodeGuid=3DB3C7EF4294D4C1609FA1BB84ACC98F
+
+		CustomProperties Pin (PinId=3709AF75400A1BB13E481CAC286A148C,PinName="self",PinFriendlyName=NSLOCTEXT("K2Node", "Target", "Target"),PinToolTip="Target\nKismet Math Library Reference",PinType.PinCategory="object",PinType.PinSubCategory="",PinType.PinSubCategoryObject=Class'/Script/Engine.KismetMathLibrary',PinType.PinSubCategoryMemberReference=(),PinType.PinValueType=(),PinType.bIsMap=False,PinType.bIsSet=False,PinType.bIsArray=False,PinType.bIsReference=False,PinType.bIsConst=False,PinType.bIsWeakPointer=False,DefaultObject="/Script/Engine.Default__KismetMathLibrary",PersistentGuid=00000000000000000000000000000000,bHidden=True,bNotConnectable=False,bDefaultValueIsReadOnly=False,bDefaultValueIsIgnored=False,bAdvancedView=False,)
+
+		CustomProperties Pin (PinId=3A4A2A21477AF2FE29AD609B9D3E1843,PinName="Location",PinToolTip="Location\nVector ",PinType.PinCategory="struct",PinType.PinSubCategory="",PinType.PinSubCategoryObject=ScriptStruct'/Script/CoreUObject.Vector',PinType.PinSubCategoryMemberReference=(),PinType.PinValueType=(),PinType.bIsMap=False,PinType.bIsSet=False,PinType.bIsArray=False,PinType.bIsReference=False,PinType.bIsConst=False,PinType.bIsWeakPointer=False,DefaultValue="0, 0, 0",AutogeneratedDefaultValue="0, 0, 0",PersistentGuid=00000000000000000000000000000000,bHidden=False,bNotConnectable=False,bDefaultValueIsReadOnly=False,bDefaultValueIsIgnored=False,bAdvancedView=False,)
+
+		CustomProperties Pin (PinId=CC2C5AC745E2D12D91F3759DC657F36A,PinName="Rotation",PinToolTip="Rotation\nRotator ",PinType.PinCategory="struct",PinType.PinSubCategory="",PinType.PinSubCategoryObject=ScriptStruct'/Script/CoreUObject.Rotator',PinType.PinSubCategoryMemberReference=(),PinType.PinValueType=(),PinType.bIsMap=False,PinType.bIsSet=False,PinType.bIsArray=False,PinType.bIsReference=False,PinType.bIsConst=False,PinType.bIsWeakPointer=False,DefaultValue="0, 0, 0",AutogeneratedDefaultValue="0, 0, 0",PersistentGuid=00000000000000000000000000000000,bHidden=False,bNotConnectable=False,bDefaultValueIsReadOnly=False,bDefaultValueIsIgnored=False,bAdvancedView=False,)
+
+		CustomProperties Pin (PinId=E6667EDA466F0B6CDF67C3B073158169,PinName="Scale",PinToolTip="Scale\nVector ",PinType.PinCategory="struct",PinType.PinSubCategory="",PinType.PinSubCategoryObject=ScriptStruct'/Script/CoreUObject.Vector',PinType.PinSubCategoryMemberReference=(),PinType.PinValueType=(),PinType.bIsMap=False,PinType.bIsSet=False,PinType.bIsArray=False,PinType.bIsReference=False,PinType.bIsConst=False,PinType.bIsWeakPointer=False,DefaultValue="1,1,1",AutogeneratedDefaultValue="1,1,1",PersistentGuid=00000000000000000000000000000000,bHidden=False,bNotConnectable=False,bDefaultValueIsReadOnly=False,bDefaultValueIsIgnored=False,bAdvancedView=False,)
+
+		CustomProperties Pin (PinId=AF359A7549B7FC5151F55B9B5C7306DD,PinName="ReturnValue",PinToolTip="Return Value\nTransform ",Direction="EGPD_Output",PinType.PinCategory="struct",PinType.PinSubCategory="",PinType.PinSubCategoryObject=ScriptStruct'/Script/CoreUObject.Transform',PinType.PinSubCategoryMemberReference=(),PinType.PinValueType=(),PinType.bIsMap=False,PinType.bIsSet=False,PinType.bIsArray=False,PinType.bIsReference=False,PinType.bIsConst=False,PinType.bIsWeakPointer=False,LinkedTo=(K2Node_AddComponent_0 4483A3B74C12C1530E84C49F92B991A8,),PersistentGuid=00000000000000000000000000000000,bHidden=False,bNotConnectable=False,bDefaultValueIsReadOnly=False,bDefaultValueIsIgnored=False,bAdvancedView=False,)
+
+	End Object
+
+	Begin Object Class=K2Node_CallFunction Name="K2Node_CallFunction_77"
+
+		FunctionReference=(MemberParent=Class'/Script/Engine.RadialForceComponent',MemberName="FireImpulse")
+
+		NodePosX=2160
+
+		NodePosY=624
+
+		NodeGuid=C3D60D2C4B023B605B68C6BBFFAF922D
+
+		CustomProperties Pin (PinId=9C4DCA9D48357011B95546B420FEE798,PinName="execute",PinType.PinCategory="exec",PinType.PinSubCategory="",PinType.PinSubCategoryObject=None,PinType.PinSubCategoryMemberReference=(),PinType.PinValueType=(),PinType.bIsMap=False,PinType.bIsSet=False,PinType.bIsArray=False,PinType.bIsReference=False,PinType.bIsConst=False,PinType.bIsWeakPointer=False,LinkedTo=(K2Node_AddComponent_0 02DBF18F43959E2CD4F3148EE7470FA3,),PersistentGuid=00000000000000000000000000000000,bHidden=False,bNotConnectable=False,bDefaultValueIsReadOnly=False,bDefaultValueIsIgnored=False,bAdvancedView=False,)
+
+		CustomProperties Pin (PinId=AEC053B24B8500ADC9AF948F39A660AF,PinName="then",Direction="EGPD_Output",PinType.PinCategory="exec",PinType.PinSubCategory="",PinType.PinSubCategoryObject=None,PinType.PinSubCategoryMemberReference=(),PinType.PinValueType=(),PinType.bIsMap=False,PinType.bIsSet=False,PinType.bIsArray=False,PinType.bIsReference=False,PinType.bIsConst=False,PinType.bIsWeakPointer=False,LinkedTo=(K2Node_CallFunction_119 D1E0CD9A44C8258C668EAFA1A4486990,),PersistentGuid=00000000000000000000000000000000,bHidden=False,bNotConnectable=False,bDefaultValueIsReadOnly=False,bDefaultValueIsIgnored=False,bAdvancedView=False,)
+
+		CustomProperties Pin (PinId=3DA9F1F1472BF28C89AAE083A5111D9B,PinName="self",PinFriendlyName=NSLOCTEXT("K2Node", "Target", "Target"),PinType.PinCategory="object",PinType.PinSubCategory="",PinType.PinSubCategoryObject=Class'/Script/Engine.RadialForceComponent',PinType.PinSubCategoryMemberReference=(),PinType.PinValueType=(),PinType.bIsMap=False,PinType.bIsSet=False,PinType.bIsArray=False,PinType.bIsReference=False,PinType.bIsConst=False,PinType.bIsWeakPointer=False,LinkedTo=(K2Node_AddComponent_0 86731E5E4566B2DBFB777890155E4768,),PersistentGuid=00000000000000000000000000000000,bHidden=False,bNotConnectable=False,bDefaultValueIsReadOnly=False,bDefaultValueIsIgnored=False,bAdvancedView=False,)
+
+	End Object
+
+	Begin Object Class=K2Node_CallFunction Name="K2Node_CallFunction_855"
+
+		FunctionReference=(MemberParent=Class'/Script/Engine.GameplayStatics',MemberName="SpawnSoundAtLocation")
+
+		NodePosX=2800
+
+		NodePosY=640
+
+		NodeGuid=3D5F99DA4701089AAE931A844F708FE0
+
+		AdvancedPinDisplay=Hidden
+
+		CustomProperties Pin (PinId=C3A6A74B4E2B5249FF82DAB24058ED08,PinName="execute",PinType.PinCategory="exec",PinType.PinSubCategory="",PinType.PinSubCategoryObject=None,PinType.PinSubCategoryMemberReference=(),PinType.PinValueType=(),PinType.bIsMap=False,PinType.bIsSet=False,PinType.bIsArray=False,PinType.bIsReference=False,PinType.bIsConst=False,PinType.bIsWeakPointer=False,LinkedTo=(K2Node_CallFunction_119 08F736C8421ECE55DA8E58B0DE9216A2,),PersistentGuid=00000000000000000000000000000000,bHidden=False,bNotConnectable=False,bDefaultValueIsReadOnly=False,bDefaultValueIsIgnored=False,bAdvancedView=False,)
+
+		CustomProperties Pin (PinId=DA8628F44B19D85719E13A8E89CCEEFC,PinName="then",Direction="EGPD_Output",PinType.PinCategory="exec",PinType.PinSubCategory="",PinType.PinSubCategoryObject=None,PinType.PinSubCategoryMemberReference=(),PinType.PinValueType=(),PinType.bIsMap=False,PinType.bIsSet=False,PinType.bIsArray=False,PinType.bIsReference=False,PinType.bIsConst=False,PinType.bIsWeakPointer=False,LinkedTo=(K2Node_VariableSet_7 AD945712496DF669C2B3C089FF4AEED2,),PersistentGuid=00000000000000000000000000000000,bHidden=False,bNotConnectable=False,bDefaultValueIsReadOnly=False,bDefaultValueIsIgnored=False,bAdvancedView=False,)
+
+		CustomProperties Pin (PinId=8C90F28542563BCB742F2BADF030C188,PinName="self",PinFriendlyName=NSLOCTEXT("K2Node", "Target", "Target"),PinType.PinCategory="object",PinType.PinSubCategory="",PinType.PinSubCategoryObject=Class'/Script/Engine.GameplayStatics',PinType.PinSubCategoryMemberReference=(),PinType.PinValueType=(),PinType.bIsMap=False,PinType.bIsSet=False,PinType.bIsArray=False,PinType.bIsReference=False,PinType.bIsConst=False,PinType.bIsWeakPointer=False,DefaultObject="/Script/Engine.Default__GameplayStatics",PersistentGuid=00000000000000000000000000000000,bHidden=True,bNotConnectable=False,bDefaultValueIsReadOnly=False,bDefaultValueIsIgnored=False,bAdvancedView=False,)
+
+		CustomProperties Pin (PinId=FACE851341C0E1B0CF1B30973D4A3710,PinName="WorldContextObject",PinType.PinCategory="object",PinType.PinSubCategory="",PinType.PinSubCategoryObject=Class'/Script/CoreUObject.Object',PinType.PinSubCategoryMemberReference=(),PinType.PinValueType=(),PinType.bIsMap=False,PinType.bIsSet=False,PinType.bIsArray=False,PinType.bIsReference=False,PinType.bIsConst=True,PinType.bIsWeakPointer=False,PersistentGuid=00000000000000000000000000000000,bHidden=True,bNotConnectable=False,bDefaultValueIsReadOnly=False,bDefaultValueIsIgnored=False,bAdvancedView=False,)
+
+		CustomProperties Pin (PinId=7F0F67224C161114DBB96C9A70CD0823,PinName="Sound",PinType.PinCategory="object",PinType.PinSubCategory="",PinType.PinSubCategoryObject=Class'/Script/Engine.SoundBase',PinType.PinSubCategoryMemberReference=(),PinType.PinValueType=(),PinType.bIsMap=False,PinType.bIsSet=False,PinType.bIsArray=False,PinType.bIsReference=False,PinType.bIsConst=False,PinType.bIsWeakPointer=False,DefaultObject="/Game/Audio/Grenade/GrenadeExplosion_C_Cue.GrenadeExplosion_C_Cue",PersistentGuid=00000000000000000000000000000000,bHidden=False,bNotConnectable=False,bDefaultValueIsReadOnly=False,bDefaultValueIsIgnored=False,bAdvancedView=False,)
+
+		CustomProperties Pin (PinId=61B25F85499B909E42E828BB47F0C279,PinName="Location",PinType.PinCategory="struct",PinType.PinSubCategory="",PinType.PinSubCategoryObject=ScriptStruct'/Script/CoreUObject.Vector',PinType.PinSubCategoryMemberReference=(),PinType.PinValueType=(),PinType.bIsMap=False,PinType.bIsSet=False,PinType.bIsArray=False,PinType.bIsReference=False,PinType.bIsConst=False,PinType.bIsWeakPointer=False,DefaultValue="0, 0, 0",AutogeneratedDefaultValue="0, 0, 0",LinkedTo=(K2Node_Knot_547 EF8E7FAC49D1E792731CD98F0A2754DA,),PersistentGuid=00000000000000000000000000000000,bHidden=False,bNotConnectable=False,bDefaultValueIsReadOnly=False,bDefaultValueIsIgnored=False,bAdvancedView=False,)
+
+		CustomProperties Pin (PinId=301FF2814987B335A5A749A2C98983B8,PinName="Rotation",PinType.PinCategory="struct",PinType.PinSubCategory="",PinType.PinSubCategoryObject=ScriptStruct'/Script/CoreUObject.Rotator',PinType.PinSubCategoryMemberReference=(),PinType.PinValueType=(),PinType.bIsMap=False,PinType.bIsSet=False,PinType.bIsArray=False,PinType.bIsReference=False,PinType.bIsConst=False,PinType.bIsWeakPointer=False,DefaultValue="0, 0, 0",AutogeneratedDefaultValue="0, 0, 0",PersistentGuid=00000000000000000000000000000000,bHidden=False,bNotConnectable=False,bDefaultValueIsReadOnly=False,bDefaultValueIsIgnored=False,bAdvancedView=True,)
+
+		CustomProperties Pin (PinId=A065258948E8997A28A64D8EE0A73882,PinName="VolumeMultiplier",PinType.PinCategory="float",PinType.PinSubCategory="",PinType.PinSubCategoryObject=None,PinType.PinSubCategoryMemberReference=(),PinType.PinValueType=(),PinType.bIsMap=False,PinType.bIsSet=False,PinType.bIsArray=False,PinType.bIsReference=False,PinType.bIsConst=False,PinType.bIsWeakPointer=False,DefaultValue="1.000000",AutogeneratedDefaultValue="1.000000",PersistentGuid=00000000000000000000000000000000,bHidden=False,bNotConnectable=False,bDefaultValueIsReadOnly=False,bDefaultValueIsIgnored=False,bAdvancedView=True,)
+
+		CustomProperties Pin (PinId=F86791D74B0F657F11BF288CF7F24EA6,PinName="PitchMultiplier",PinType.PinCategory="float",PinType.PinSubCategory="",PinType.PinSubCategoryObject=None,PinType.PinSubCategoryMemberReference=(),PinType.PinValueType=(),PinType.bIsMap=False,PinType.bIsSet=False,PinType.bIsArray=False,PinType.bIsReference=False,PinType.bIsConst=False,PinType.bIsWeakPointer=False,DefaultValue="1.000000",AutogeneratedDefaultValue="1.000000",PersistentGuid=00000000000000000000000000000000,bHidden=False,bNotConnectable=False,bDefaultValueIsReadOnly=False,bDefaultValueIsIgnored=False,bAdvancedView=True,)
+
+		CustomProperties Pin (PinId=B4E444E2439EC48B82163E8184484752,PinName="StartTime",PinType.PinCategory="float",PinType.PinSubCategory="",PinType.PinSubCategoryObject=None,PinType.PinSubCategoryMemberReference=(),PinType.PinValueType=(),PinType.bIsMap=False,PinType.bIsSet=False,PinType.bIsArray=False,PinType.bIsReference=False,PinType.bIsConst=False,PinType.bIsWeakPointer=False,DefaultValue="0.000000",AutogeneratedDefaultValue="0.000000",PersistentGuid=00000000000000000000000000000000,bHidden=False,bNotConnectable=False,bDefaultValueIsReadOnly=False,bDefaultValueIsIgnored=False,bAdvancedView=True,)
+
+		CustomProperties Pin (PinId=212C680D4368045887ECC5ACE22018DF,PinName="AttenuationSettings",PinType.PinCategory="object",PinType.PinSubCategory="",PinType.PinSubCategoryObject=Class'/Script/Engine.SoundAttenuation',PinType.PinSubCategoryMemberReference=(),PinType.PinValueType=(),PinType.bIsMap=False,PinType.bIsSet=False,PinType.bIsArray=False,PinType.bIsReference=False,PinType.bIsConst=False,PinType.bIsWeakPointer=False,PersistentGuid=00000000000000000000000000000000,bHidden=False,bNotConnectable=False,bDefaultValueIsReadOnly=False,bDefaultValueIsIgnored=False,bAdvancedView=True,)
+
+		CustomProperties Pin (PinId=815CB039420FE49D87FBF6B179746652,PinName="ConcurrencySettings",PinType.PinCategory="object",PinType.PinSubCategory="",PinType.PinSubCategoryObject=Class'/Script/Engine.SoundConcurrency',PinType.PinSubCategoryMemberReference=(),PinType.PinValueType=(),PinType.bIsMap=False,PinType.bIsSet=False,PinType.bIsArray=False,PinType.bIsReference=False,PinType.bIsConst=False,PinType.bIsWeakPointer=False,PersistentGuid=00000000000000000000000000000000,bHidden=False,bNotConnectable=False,bDefaultValueIsReadOnly=False,bDefaultValueIsIgnored=False,bAdvancedView=True,)
+
+		CustomProperties Pin (PinId=9016ED934882ADFF5A17378EE4490B90,PinName="ReturnValue",Direction="EGPD_Output",PinType.PinCategory="object",PinType.PinSubCategory="",PinType.PinSubCategoryObject=Class'/Script/Engine.AudioComponent',PinType.PinSubCategoryMemberReference=(),PinType.PinValueType=(),PinType.bIsMap=False,PinType.bIsSet=False,PinType.bIsArray=False,PinType.bIsReference=False,PinType.bIsConst=False,PinType.bIsWeakPointer=False,PersistentGuid=00000000000000000000000000000000,bHidden=False,bNotConnectable=False,bDefaultValueIsReadOnly=False,bDefaultValueIsIgnored=False,bAdvancedView=False,)
+
+	End Object
+
+	Begin Object Class=K2Node_Knot Name="K2Node_Knot_547"
+
+		NodePosX=2688
+
+		NodePosY=896
+
+		NodeGuid=318EF95D4B750464B88915B0EFADF64C
+
+		CustomProperties Pin (PinId=00AD35A545734CB7CF62A1A741FBEA71,PinName="InputPin",PinType.PinCategory="struct",PinType.PinSubCategory="",PinType.PinSubCategoryObject=ScriptStruct'/Script/CoreUObject.Vector',PinType.PinSubCategoryMemberReference=(),PinType.PinValueType=(),PinType.bIsMap=False,PinType.bIsSet=False,PinType.bIsArray=False,PinType.bIsReference=False,PinType.bIsConst=False,PinType.bIsWeakPointer=False,LinkedTo=(K2Node_Knot_548 B674F14C4D8682DFC747E1B926EA534B,),PersistentGuid=00000000000000000000000000000000,bHidden=False,bNotConnectable=False,bDefaultValueIsReadOnly=False,bDefaultValueIsIgnored=True,bAdvancedView=False,)
+
+		CustomProperties Pin (PinId=EF8E7FAC49D1E792731CD98F0A2754DA,PinName="OutputPin",Direction="EGPD_Output",PinType.PinCategory="struct",PinType.PinSubCategory="",PinType.PinSubCategoryObject=ScriptStruct'/Script/CoreUObject.Vector',PinType.PinSubCategoryMemberReference=(),PinType.PinValueType=(),PinType.bIsMap=False,PinType.bIsSet=False,PinType.bIsArray=False,PinType.bIsReference=False,PinType.bIsConst=False,PinType.bIsWeakPointer=False,LinkedTo=(K2Node_CallFunction_855 61B25F85499B909E42E828BB47F0C279,),PersistentGuid=00000000000000000000000000000000,bHidden=False,bNotConnectable=False,bDefaultValueIsReadOnly=False,bDefaultValueIsIgnored=False,bAdvancedView=False,)
+
+	End Object
+
+	Begin Object Class=K2Node_Knot Name="K2Node_Knot_548"
+
+		NodePosX=2416
+
+		NodePosY=896
+
+		NodeGuid=511290E24A466C676B1DF681EFAB6D07
+
+		CustomProperties Pin (PinId=D8C5606C4E0103BD1B30C1909156FE54,PinName="InputPin",PinType.PinCategory="struct",PinType.PinSubCategory="",PinType.PinSubCategoryObject=ScriptStruct'/Script/CoreUObject.Vector',PinType.PinSubCategoryMemberReference=(),PinType.PinValueType=(),PinType.bIsMap=False,PinType.bIsSet=False,PinType.bIsArray=False,PinType.bIsReference=False,PinType.bIsConst=False,PinType.bIsWeakPointer=False,LinkedTo=(K2Node_CallFunction_116 C149106E4CC3C4B43D809AB6A6CC5465,),PersistentGuid=00000000000000000000000000000000,bHidden=False,bNotConnectable=False,bDefaultValueIsReadOnly=False,bDefaultValueIsIgnored=True,bAdvancedView=False,)
+
+		CustomProperties Pin (PinId=B674F14C4D8682DFC747E1B926EA534B,PinName="OutputPin",Direction="EGPD_Output",PinType.PinCategory="struct",PinType.PinSubCategory="",PinType.PinSubCategoryObject=ScriptStruct'/Script/CoreUObject.Vector',PinType.PinSubCategoryMemberReference=(),PinType.PinValueType=(),PinType.bIsMap=False,PinType.bIsSet=False,PinType.bIsArray=False,PinType.bIsReference=False,PinType.bIsConst=False,PinType.bIsWeakPointer=False,LinkedTo=(K2Node_Knot_547 00AD35A545734CB7CF62A1A741FBEA71,),PersistentGuid=00000000000000000000000000000000,bHidden=False,bNotConnectable=False,bDefaultValueIsReadOnly=False,bDefaultValueIsIgnored=False,bAdvancedView=False,)
+
+	End Object
+[/EXCERPT:snippet]
+-->
+
diff --git a/Engine/Documentation/Source/Modding/RoboRecall/HowTo/ExtendFunctionality/AssignMaterialInstance/AssignMaterialInstance.INT.udn b/Engine/Documentation/Source/Modding/RoboRecall/HowTo/ExtendFunctionality/AssignMaterialInstance/AssignMaterialInstance.INT.udn
new file mode 100644
--- /dev/null
+++ b/Engine/Documentation/Source/Modding/RoboRecall/HowTo/ExtendFunctionality/AssignMaterialInstance/AssignMaterialInstance.INT.udn
@@ -0,0 +1,45 @@
+Availability: Public
+Title:3. Assigning the Material Instance
+Crumbs:%ROOT%
+Description:Robo Recall Modding Documentation
+Parent:Modding/RoboRecall/HowTo/ExtendFunctionality
+type:how-to
+order:3
+tags:Modding
+tags:Robo Recall
+topic-image:Modding/RoboRecall/RoboRecall_Topic.png
+social-image:Modding/RoboRecall/RoboRecall_Social.png
+checkpoint:mod_extend
+
+<!-- Location: Modding/RoboRecall/Quickstarts/WeaponMod/3 -->
+
+[NAV]
+
+<!-- Title Image if Needed Here -->
+
+Your custom pistol needs to know about the new dissolve Material in order to make use of it. You can assign it to your custom pistol to force it to be used.
+
+## Steps
+
+1. Double Click the **GrenadeGunPistol** Blueprint to open it in the [Blueprint Editor](Engine/Blueprints/Editor).
+
+2. Select **GrenadeGunPistol(self)** in the **Components** panel
+
+	![image alt text](image_0.png)(w:302 h:253 convert:false)
+
+3. In the Details panel search for **Dissolve** to easily find the **Dissolve Material** property.
+
+4. Click the dropdown menu for the  **Dissolve Material** property to show the Asset Picker and select **GrenadeGun_Dissolve_MI**. You can use the search box to quickly filter the list.
+
+	[REGION:raw]
+	![image alt text](image_1.png)(w:933 h:494 convert:false)
+	[/REGION]
+
+5. In the Blueprint Editor toolbar, click **Compile** to update the Blueprint and click **Save** to preserve your changes.
+
+## Results
+
+That's it, the GrenadeGunPistol, when it is destroyed, will now dissolve to red instead of blue.
+
+[NAV]
+
diff --git a/Engine/Documentation/Source/Modding/RoboRecall/HowTo/ExtendFunctionality/BaseDamageBlueprint/BaseDamageBlueprint.INT.udn b/Engine/Documentation/Source/Modding/RoboRecall/HowTo/ExtendFunctionality/BaseDamageBlueprint/BaseDamageBlueprint.INT.udn
new file mode 100644
--- /dev/null
+++ b/Engine/Documentation/Source/Modding/RoboRecall/HowTo/ExtendFunctionality/BaseDamageBlueprint/BaseDamageBlueprint.INT.udn
@@ -0,0 +1,394 @@
+Availability: Public
+Title:5. Creating the Base Damage Blueprint
+Crumbs:%ROOT%
+Description:Robo Recall Modding Documentation
+Parent:Modding/RoboRecall/HowTo/ExtendFunctionality
+type:how-to
+order:5
+tags:Modding
+tags:Robo Recall
+topic-image:Modding/RoboRecall/RoboRecall_Topic.png
+social-image:Modding/RoboRecall/RoboRecall_Social.png
+checkpoint:mod_extend
+
+<!-- Location: Modding/RoboRecall/Quickstarts/WeaponMod/5 -->
+
+[NAV]
+
+This step will cover building the basic Blueprint for having the GrenadeGunPistol deal damage once it collides with the something. 
+
+[REGION:note]
+Since you've learned how to make nodes in the previous steps, this section will start to accelerate by grouping creating nodes into chunks. Just remember you can **Right-click** in the **Event Graph** to find every node available to the Blueprint, and pulling off a pin will narrow down the context sensitivity of the node search to things that operate on that pin type.
+[/REGION]
+
+## Steps
+
+1. Create an **Event Hit** node.
+
+	![image alt text](image_0.png)(w:140 h:280 convert:false)
+
+2. Pull off the **Other** pin and search for **!=** to create a **Not Equals** node.
+
+	![image alt text](image_1.png)(w:594 h:440 convert:false)
+
+3. Set the second input pin on the **Not Equals** node to `shell`, you'll need to scroll down a ways to find the Blueprint simply called `shell`.
+
+	![image alt text](image_2.png)(w:870 h:538 convert:false)
+
+4. Next up you'll need to setup a bit of logic using an **And** node, a** Branch** node, and the **bCanExplode** variable.
+
+	![image alt text](image_3.png)(w:746 h:278 convert:false)
+
+	[REGION:tip]
+	There are a number of keyboard shortcuts in the **Event Graph**, once of which is **Hold B + Left Click** to create a **Branch** node, arguably one of the most used logic nodes in **Blueprints**.
+	[/REGION]
+
+5. Now you're going to set up how much damage the GrenadeGun will do by getting the **Ammo Remaining **in the gun, multiplying it by 1000 and adding 1000 to it.
+
+	![image alt text](image_4.png)(w:546 h:76 convert:false)
+
+	**Ammo Remaining** is a variable that is apart of the base pistol class, so you can find it by searching for **Get Ammo Remaining**. The **Multiply** node is actually an **Int * Float** node, and can be found by searching for **Int * Float**. And finally, the **Addition** node can be found by searching for **Float + Float**.
+
+6. You'll only want the gun to explode once, so you'll need a **Do Once** node connected to the **True** pin of the **Branch** node.
+	
+	![image alt text](image_5.png)(w:934 h:236 convert:false)
+
+7. You're going to need a **Self** node and a **Get Player Controller** node.
+
+	![image alt text](image_6.png)(w:264 h:126 convert:false)
+
+8. Finally, to do some damage you'll need a **Apply Radial Damage** node, that will take a bunch of inputs from what you've created and will require some settings changed. Connect it up as shown below. You'll need to change 3 settings on **Apply Radial Damage** node: **Damage Radius**, **Damage Type Class**, and **Do Full Damage**.
+
+	* Set **Damage Radius** to `400`
+	* Set **Damage Type Class** to `Odin Damage Type`
+	* Set **Do Full Damage** to `True`**
+
+	[OBJECT:Snippet]
+		[PARAMLITERAL:id]
+		result
+		[/PARAMLITERAL]
+		[PARAMLITERAL:type]
+		blueprint
+		[/PARAMLITERAL]
+		[PARAM:image]
+		[REGION:lightbox]
+		[![image alt text](image_7.png)(convert:false)](image_7.png)
+		[/REGION]
+		[/PARAM]
+		[PARAM:snippet]
+			[INCLUDE:#snippet]
+		[/PARAM]
+	[/OBJECT]
+
+	[REGION:caption]
+	Click the icon in the upper left corner of this image to copy the Blueprint Graph and paste it into your project.
+	[/REGION]
+
+	[REGION:tip]
+	You can create **Reroute Nodes** to clean up your graphs by **Double Clicking** on any connecting line. As seen above on the yellow line that leads from **Hit Location** to **Origin**.
+	[/REGION]
+
+9. In the Blueprint Editor toolbar, click **Compile** to update the Blueprint and click **Save** to preserve your changes.
+
+## Results
+
+After this step, you will have a functional weapon that will explode when thrown into the world. It will do more damage based on the ammo left in the magazine, and a fully loaded pistol will do enough damage to kill just about any non-specialized bot. In the next step, you will be adding in sound, VFX, a physics impulse, and clean up the gun after it explodes.
+
+[NAV]
+
+<!--
+[EXCERPT:snippet]
+	Begin Object Class=K2Node_Event Name="K2Node_Event_3"
+
+	   EventReference=(MemberParent=Class'/Script/Engine.Actor',MemberName="ReceiveHit")
+
+	   bOverrideFunction=True
+
+	   NodePosY=640
+
+	   NodeGuid=E100ABD14FC0884AB0900FB8E5E6F008
+
+	   CustomProperties Pin (PinId=7A104AB44216A8AD0F2414A1658563AA,PinName="OutputDelegate",Direction="EGPD_Output",PinType.PinCategory="delegate",PinType.PinSubCategory="",PinType.PinSubCategoryObject=None,PinType.PinSubCategoryMemberReference=(MemberParent=Class'/Script/Engine.Actor',MemberName="ReceiveHit"),PinType.PinValueType=(),PinType.bIsMap=False,PinType.bIsSet=False,PinType.bIsArray=False,PinType.bIsReference=False,PinType.bIsConst=False,PinType.bIsWeakPointer=False,PersistentGuid=00000000000000000000000000000000,bHidden=False,bNotConnectable=False,bDefaultValueIsReadOnly=False,bDefaultValueIsIgnored=False,bAdvancedView=False,)
+
+	   CustomProperties Pin (PinId=E143C00B442A411DCCCF33BB48310E3A,PinName="then",Direction="EGPD_Output",PinType.PinCategory="exec",PinType.PinSubCategory="",PinType.PinSubCategoryObject=None,PinType.PinSubCategoryMemberReference=(),PinType.PinValueType=(),PinType.bIsMap=False,PinType.bIsSet=False,PinType.bIsArray=False,PinType.bIsReference=False,PinType.bIsConst=False,PinType.bIsWeakPointer=False,LinkedTo=(K2Node_IfThenElse_0 0B9C43A44127F21B33ECF99B1D75F64A,),PersistentGuid=00000000000000000000000000000000,bHidden=False,bNotConnectable=False,bDefaultValueIsReadOnly=False,bDefaultValueIsIgnored=False,bAdvancedView=False,)
+
+	   CustomProperties Pin (PinId=4C58F0BE4FFB6304C47CAF8181D3CA46,PinName="MyComp",PinToolTip="My Comp\nPrimitive Component Reference",Direction="EGPD_Output",PinType.PinCategory="object",PinType.PinSubCategory="",PinType.PinSubCategoryObject=Class'/Script/Engine.PrimitiveComponent',PinType.PinSubCategoryMemberReference=(),PinType.PinValueType=(),PinType.bIsMap=False,PinType.bIsSet=False,PinType.bIsArray=False,PinType.bIsReference=False,PinType.bIsConst=False,PinType.bIsWeakPointer=False,PersistentGuid=00000000000000000000000000000000,bHidden=False,bNotConnectable=False,bDefaultValueIsReadOnly=False,bDefaultValueIsIgnored=False,bAdvancedView=False,)
+
+	   CustomProperties Pin (PinId=0907949C45B9012B8F8758BB8D7C1FE6,PinName="Other",PinToolTip="Other\nActor Reference",Direction="EGPD_Output",PinType.PinCategory="object",PinType.PinSubCategory="",PinType.PinSubCategoryObject=Class'/Script/Engine.Actor',PinType.PinSubCategoryMemberReference=(),PinType.PinValueType=(),PinType.bIsMap=False,PinType.bIsSet=False,PinType.bIsArray=False,PinType.bIsReference=False,PinType.bIsConst=False,PinType.bIsWeakPointer=False,LinkedTo=(K2Node_CallFunction_1 DE9CF168407388BF6E83AC9FC4CD7444,),PersistentGuid=00000000000000000000000000000000,bHidden=False,bNotConnectable=False,bDefaultValueIsReadOnly=False,bDefaultValueIsIgnored=False,bAdvancedView=False,)
+
+	   CustomProperties Pin (PinId=F7FAB48E4AF60278ED841B96A969FA73,PinName="OtherComp",PinToolTip="Other Comp\nPrimitive Component Reference",Direction="EGPD_Output",PinType.PinCategory="object",PinType.PinSubCategory="",PinType.PinSubCategoryObject=Class'/Script/Engine.PrimitiveComponent',PinType.PinSubCategoryMemberReference=(),PinType.PinValueType=(),PinType.bIsMap=False,PinType.bIsSet=False,PinType.bIsArray=False,PinType.bIsReference=False,PinType.bIsConst=False,PinType.bIsWeakPointer=False,PersistentGuid=00000000000000000000000000000000,bHidden=False,bNotConnectable=False,bDefaultValueIsReadOnly=False,bDefaultValueIsIgnored=False,bAdvancedView=False,)
+
+	   CustomProperties Pin (PinId=B367ABEC43BDD822BE465DA5ECDCE31C,PinName="bSelfMoved",PinToolTip="Self Moved\nBoolean",Direction="EGPD_Output",PinType.PinCategory="bool",PinType.PinSubCategory="",PinType.PinSubCategoryObject=None,PinType.PinSubCategoryMemberReference=(),PinType.PinValueType=(),PinType.bIsMap=False,PinType.bIsSet=False,PinType.bIsArray=False,PinType.bIsReference=False,PinType.bIsConst=False,PinType.bIsWeakPointer=False,DefaultValue="false",AutogeneratedDefaultValue="false",PersistentGuid=00000000000000000000000000000000,bHidden=False,bNotConnectable=False,bDefaultValueIsReadOnly=False,bDefaultValueIsIgnored=False,bAdvancedView=False,)
+
+	   CustomProperties Pin (PinId=3739097840F8ACAB5630D8AEA7CADC9B,PinName="HitLocation",PinToolTip="Hit Location\nVector ",Direction="EGPD_Output",PinType.PinCategory="struct",PinType.PinSubCategory="",PinType.PinSubCategoryObject=ScriptStruct'/Script/CoreUObject.Vector',PinType.PinSubCategoryMemberReference=(),PinType.PinValueType=(),PinType.bIsMap=False,PinType.bIsSet=False,PinType.bIsArray=False,PinType.bIsReference=False,PinType.bIsConst=False,PinType.bIsWeakPointer=False,AutogeneratedDefaultValue="0, 0, 0",LinkedTo=(K2Node_Knot_546 6EEFCD4849DD5591D5B1CB9B8FB76133,),PersistentGuid=00000000000000000000000000000000,bHidden=False,bNotConnectable=False,bDefaultValueIsReadOnly=False,bDefaultValueIsIgnored=False,bAdvancedView=False,)
+
+	   CustomProperties Pin (PinId=6E1FB5A3485F4F722DA1A3B2384A17B4,PinName="HitNormal",PinToolTip="Hit Normal\nVector ",Direction="EGPD_Output",PinType.PinCategory="struct",PinType.PinSubCategory="",PinType.PinSubCategoryObject=ScriptStruct'/Script/CoreUObject.Vector',PinType.PinSubCategoryMemberReference=(),PinType.PinValueType=(),PinType.bIsMap=False,PinType.bIsSet=False,PinType.bIsArray=False,PinType.bIsReference=False,PinType.bIsConst=False,PinType.bIsWeakPointer=False,DefaultValue="0, 0, 0",AutogeneratedDefaultValue="0, 0, 0",PersistentGuid=00000000000000000000000000000000,bHidden=False,bNotConnectable=False,bDefaultValueIsReadOnly=False,bDefaultValueIsIgnored=False,bAdvancedView=False,)
+
+	   CustomProperties Pin (PinId=834813F0496D475CB83A10A9695F089E,PinName="NormalImpulse",PinToolTip="Normal Impulse\nVector ",Direction="EGPD_Output",PinType.PinCategory="struct",PinType.PinSubCategory="",PinType.PinSubCategoryObject=ScriptStruct'/Script/CoreUObject.Vector',PinType.PinSubCategoryMemberReference=(),PinType.PinValueType=(),PinType.bIsMap=False,PinType.bIsSet=False,PinType.bIsArray=False,PinType.bIsReference=False,PinType.bIsConst=False,PinType.bIsWeakPointer=False,DefaultValue="0, 0, 0",AutogeneratedDefaultValue="0, 0, 0",PersistentGuid=00000000000000000000000000000000,bHidden=False,bNotConnectable=False,bDefaultValueIsReadOnly=False,bDefaultValueIsIgnored=False,bAdvancedView=False,)
+
+	   CustomProperties Pin (PinId=D13B1D8F483A1166FD266D9270B6D8F1,PinName="Hit",PinToolTip="Hit\nHit Result Structure (by ref)",Direction="EGPD_Output",PinType.PinCategory="struct",PinType.PinSubCategory="",PinType.PinSubCategoryObject=ScriptStruct'/Script/Engine.HitResult',PinType.PinSubCategoryMemberReference=(),PinType.PinValueType=(),PinType.bIsMap=False,PinType.bIsSet=False,PinType.bIsArray=False,PinType.bIsReference=True,PinType.bIsConst=True,PinType.bIsWeakPointer=False,PersistentGuid=00000000000000000000000000000000,bHidden=False,bNotConnectable=False,bDefaultValueIsReadOnly=False,bDefaultValueIsIgnored=False,bAdvancedView=False,)
+
+	End Object
+
+	Begin Object Class=K2Node_CallFunction Name="K2Node_CallFunction_1"
+
+	   bIsPureFunc=True
+
+	   FunctionReference=(MemberParent=Class'/Script/Engine.KismetMathLibrary',MemberName="NotEqual_ObjectObject")
+
+	   NodePosX=160
+
+	   NodePosY=688
+
+	   NodeGuid=6BE9ED03489886A614CEE9A92CE30E08
+
+	   CustomProperties Pin (PinId=CA8201D74A791A3C024064A21191A14A,PinName="self",PinFriendlyName=NSLOCTEXT("K2Node", "Target", "Target"),PinType.PinCategory="object",PinType.PinSubCategory="",PinType.PinSubCategoryObject=Class'/Script/Engine.KismetMathLibrary',PinType.PinSubCategoryMemberReference=(),PinType.PinValueType=(),PinType.bIsMap=False,PinType.bIsSet=False,PinType.bIsArray=False,PinType.bIsReference=False,PinType.bIsConst=False,PinType.bIsWeakPointer=False,DefaultObject="/Script/Engine.Default__KismetMathLibrary",PersistentGuid=00000000000000000000000000000000,bHidden=True,bNotConnectable=False,bDefaultValueIsReadOnly=False,bDefaultValueIsIgnored=False,bAdvancedView=False,)
+
+	   CustomProperties Pin (PinId=DE9CF168407388BF6E83AC9FC4CD7444,PinName="A",PinType.PinCategory="object",PinType.PinSubCategory="",PinType.PinSubCategoryObject=Class'/Script/CoreUObject.Object',PinType.PinSubCategoryMemberReference=(),PinType.PinValueType=(),PinType.bIsMap=False,PinType.bIsSet=False,PinType.bIsArray=False,PinType.bIsReference=False,PinType.bIsConst=False,PinType.bIsWeakPointer=False,LinkedTo=(K2Node_Event_3 0907949C45B9012B8F8758BB8D7C1FE6,),PersistentGuid=00000000000000000000000000000000,bHidden=False,bNotConnectable=False,bDefaultValueIsReadOnly=False,bDefaultValueIsIgnored=False,bAdvancedView=False,)
+
+	   CustomProperties Pin (PinId=19554E954BF591BA2B063FA23DB6918A,PinName="B",PinType.PinCategory="object",PinType.PinSubCategory="",PinType.PinSubCategoryObject=Class'/Script/CoreUObject.Object',PinType.PinSubCategoryMemberReference=(),PinType.PinValueType=(),PinType.bIsMap=False,PinType.bIsSet=False,PinType.bIsArray=False,PinType.bIsReference=False,PinType.bIsConst=False,PinType.bIsWeakPointer=False,DefaultObject="/Game/Prototype/Game/Shell.Shell",PersistentGuid=00000000000000000000000000000000,bHidden=False,bNotConnectable=False,bDefaultValueIsReadOnly=False,bDefaultValueIsIgnored=False,bAdvancedView=False,)
+
+	   CustomProperties Pin (PinId=C57D3581488D11C825293B9D27152EAA,PinName="ReturnValue",Direction="EGPD_Output",PinType.PinCategory="bool",PinType.PinSubCategory="",PinType.PinSubCategoryObject=None,PinType.PinSubCategoryMemberReference=(),PinType.PinValueType=(),PinType.bIsMap=False,PinType.bIsSet=False,PinType.bIsArray=False,PinType.bIsReference=False,PinType.bIsConst=False,PinType.bIsWeakPointer=False,AutogeneratedDefaultValue="false",LinkedTo=(K2Node_CommutativeAssociativeBinaryOperator_1 5DB3157C48E4D9D9B2A4F5B237BB4B05,),PersistentGuid=00000000000000000000000000000000,bHidden=False,bNotConnectable=False,bDefaultValueIsReadOnly=False,bDefaultValueIsIgnored=False,bAdvancedView=False,)
+
+	End Object
+
+	Begin Object Class=K2Node_IfThenElse Name="K2Node_IfThenElse_0"
+
+	   NodePosX=560
+
+	   NodePosY=640
+
+	   NodeGuid=74BA66474F5FD85C30E3C09D9D842B16
+
+	   CustomProperties Pin (PinId=0B9C43A44127F21B33ECF99B1D75F64A,PinName="execute",PinType.PinCategory="exec",PinType.PinSubCategory="",PinType.PinSubCategoryObject=None,PinType.PinSubCategoryMemberReference=(),PinType.PinValueType=(),PinType.bIsMap=False,PinType.bIsSet=False,PinType.bIsArray=False,PinType.bIsReference=False,PinType.bIsConst=False,PinType.bIsWeakPointer=False,LinkedTo=(K2Node_Event_3 E143C00B442A411DCCCF33BB48310E3A,),PersistentGuid=00000000000000000000000000000000,bHidden=False,bNotConnectable=False,bDefaultValueIsReadOnly=False,bDefaultValueIsIgnored=False,bAdvancedView=False,)
+
+	   CustomProperties Pin (PinId=78E5DA334930E09580EF89A92BF81356,PinName="Condition",PinType.PinCategory="bool",PinType.PinSubCategory="",PinType.PinSubCategoryObject=None,PinType.PinSubCategoryMemberReference=(),PinType.PinValueType=(),PinType.bIsMap=False,PinType.bIsSet=False,PinType.bIsArray=False,PinType.bIsReference=False,PinType.bIsConst=False,PinType.bIsWeakPointer=False,DefaultValue="false",AutogeneratedDefaultValue="true",LinkedTo=(K2Node_CommutativeAssociativeBinaryOperator_1 7B8DC675463DABC748E5C1BC81E5CA8F,),PersistentGuid=00000000000000000000000000000000,bHidden=False,bNotConnectable=False,bDefaultValueIsReadOnly=False,bDefaultValueIsIgnored=False,bAdvancedView=False,)
+
+	   CustomProperties Pin (PinId=899DE7DA4D10D721FAF47DB62C728E8C,PinName="then",PinFriendlyName=NSLOCTEXT("K2Node", "true", "true"),Direction="EGPD_Output",PinType.PinCategory="exec",PinType.PinSubCategory="",PinType.PinSubCategoryObject=None,PinType.PinSubCategoryMemberReference=(),PinType.PinValueType=(),PinType.bIsMap=False,PinType.bIsSet=False,PinType.bIsArray=False,PinType.bIsReference=False,PinType.bIsConst=False,PinType.bIsWeakPointer=False,LinkedTo=(K2Node_MacroInstance_0 2B445179486EFD71E4EB229C0D980F97,),PersistentGuid=00000000000000000000000000000000,bHidden=False,bNotConnectable=False,bDefaultValueIsReadOnly=False,bDefaultValueIsIgnored=False,bAdvancedView=False,)
+
+	   CustomProperties Pin (PinId=2A420E25462FCBB3BB26FC802B478C83,PinName="else",PinFriendlyName=NSLOCTEXT("K2Node", "false", "false"),Direction="EGPD_Output",PinType.PinCategory="exec",PinType.PinSubCategory="",PinType.PinSubCategoryObject=None,PinType.PinSubCategoryMemberReference=(),PinType.PinValueType=(),PinType.bIsMap=False,PinType.bIsSet=False,PinType.bIsArray=False,PinType.bIsReference=False,PinType.bIsConst=False,PinType.bIsWeakPointer=False,PersistentGuid=00000000000000000000000000000000,bHidden=False,bNotConnectable=False,bDefaultValueIsReadOnly=False,bDefaultValueIsIgnored=False,bAdvancedView=False,)
+
+	End Object
+
+	Begin Object Class=K2Node_CallFunction Name="K2Node_CallFunction_3"
+
+	   FunctionReference=(MemberParent=Class'/Script/Engine.GameplayStatics',MemberName="ApplyRadialDamage")
+
+	   NodePosX=1136
+
+	   NodePosY=640
+
+	   NodeGuid=36557B3748666D6F59AFC3B260C90ADB
+
+	   CustomProperties Pin (PinId=22BE44D544FB103B193288B3A615D3BC,PinName="execute",PinToolTip="\nExec",PinType.PinCategory="exec",PinType.PinSubCategory="",PinType.PinSubCategoryObject=None,PinType.PinSubCategoryMemberReference=(),PinType.PinValueType=(),PinType.bIsMap=False,PinType.bIsSet=False,PinType.bIsArray=False,PinType.bIsReference=False,PinType.bIsConst=False,PinType.bIsWeakPointer=False,LinkedTo=(K2Node_MacroInstance_0 7B80FE2849D77B7B2400F69815D9369F,),PersistentGuid=00000000000000000000000000000000,bHidden=False,bNotConnectable=False,bDefaultValueIsReadOnly=False,bDefaultValueIsIgnored=False,bAdvancedView=False,)
+
+	   CustomProperties Pin (PinId=3A6AA71340DEF59A3541578052CDA34E,PinName="then",PinToolTip="\nExec",Direction="EGPD_Output",PinType.PinCategory="exec",PinType.PinSubCategory="",PinType.PinSubCategoryObject=None,PinType.PinSubCategoryMemberReference=(),PinType.PinValueType=(),PinType.bIsMap=False,PinType.bIsSet=False,PinType.bIsArray=False,PinType.bIsReference=False,PinType.bIsConst=False,PinType.bIsWeakPointer=False,LinkedTo=(K2Node_AddComponent_0 E212548F4EB1DBC875747DB35AF1E319,),PersistentGuid=00000000000000000000000000000000,bHidden=False,bNotConnectable=False,bDefaultValueIsReadOnly=False,bDefaultValueIsIgnored=False,bAdvancedView=False,)
+
+	   CustomProperties Pin (PinId=3792DCC64775908063366E82290D5A78,PinName="self",PinFriendlyName=NSLOCTEXT("K2Node", "Target", "Target"),PinToolTip="Target\nGameplay Statics Reference",PinType.PinCategory="object",PinType.PinSubCategory="",PinType.PinSubCategoryObject=Class'/Script/Engine.GameplayStatics',PinType.PinSubCategoryMemberReference=(),PinType.PinValueType=(),PinType.bIsMap=False,PinType.bIsSet=False,PinType.bIsArray=False,PinType.bIsReference=False,PinType.bIsConst=False,PinType.bIsWeakPointer=False,DefaultObject="/Script/Engine.Default__GameplayStatics",PersistentGuid=00000000000000000000000000000000,bHidden=True,bNotConnectable=False,bDefaultValueIsReadOnly=False,bDefaultValueIsIgnored=False,bAdvancedView=False,)
+
+	   CustomProperties Pin (PinId=42A550E540E1B2020BF8608B44F9498A,PinName="WorldContextObject",PinToolTip="World Context Object\nObject Reference",PinType.PinCategory="object",PinType.PinSubCategory="",PinType.PinSubCategoryObject=Class'/Script/CoreUObject.Object',PinType.PinSubCategoryMemberReference=(),PinType.PinValueType=(),PinType.bIsMap=False,PinType.bIsSet=False,PinType.bIsArray=False,PinType.bIsReference=False,PinType.bIsConst=True,PinType.bIsWeakPointer=False,PersistentGuid=00000000000000000000000000000000,bHidden=True,bNotConnectable=False,bDefaultValueIsReadOnly=False,bDefaultValueIsIgnored=False,bAdvancedView=False,)
+
+	   CustomProperties Pin (PinId=EAD1F906420BA584197D0A838C95F09C,PinName="BaseDamage",PinToolTip="Base Damage\nFloat\n\nThe base damage to apply, i.e. the damage at the origin.",PinType.PinCategory="float",PinType.PinSubCategory="",PinType.PinSubCategoryObject=None,PinType.PinSubCategoryMemberReference=(),PinType.PinValueType=(),PinType.bIsMap=False,PinType.bIsSet=False,PinType.bIsArray=False,PinType.bIsReference=False,PinType.bIsConst=False,PinType.bIsWeakPointer=False,DefaultValue="0.0",AutogeneratedDefaultValue="0.0",LinkedTo=(K2Node_CommutativeAssociativeBinaryOperator_10 0DA0468E49FFB654F7C5D4841F83ADAC,),PersistentGuid=00000000000000000000000000000000,bHidden=False,bNotConnectable=False,bDefaultValueIsReadOnly=False,bDefaultValueIsIgnored=False,bAdvancedView=False,)
+
+	   CustomProperties Pin (PinId=2951FACB4839E3137D9FD09BFE3CC47F,PinName="Origin",PinToolTip="Origin\nVector  (by ref)\n\nEpicenter of the damage area.",PinType.PinCategory="struct",PinType.PinSubCategory="",PinType.PinSubCategoryObject=ScriptStruct'/Script/CoreUObject.Vector',PinType.PinSubCategoryMemberReference=(),PinType.PinValueType=(),PinType.bIsMap=False,PinType.bIsSet=False,PinType.bIsArray=False,PinType.bIsReference=True,PinType.bIsConst=True,PinType.bIsWeakPointer=False,DefaultValue="0, 0, 0",AutogeneratedDefaultValue="0, 0, 0",LinkedTo=(K2Node_Knot_545 38F57ED3452FC48D3E3F858071F90765,),PersistentGuid=00000000000000000000000000000000,bHidden=False,bNotConnectable=False,bDefaultValueIsReadOnly=False,bDefaultValueIsIgnored=False,bAdvancedView=False,)
+
+	   CustomProperties Pin (PinId=AE4A5E3D45C26736F66D8E9EB908E94B,PinName="DamageRadius",PinToolTip="Damage Radius\nFloat\n\nRadius of the damage area, from Origin",PinType.PinCategory="float",PinType.PinSubCategory="",PinType.PinSubCategoryObject=None,PinType.PinSubCategoryMemberReference=(),PinType.PinValueType=(),PinType.bIsMap=False,PinType.bIsSet=False,PinType.bIsArray=False,PinType.bIsReference=False,PinType.bIsConst=False,PinType.bIsWeakPointer=False,DefaultValue="400",AutogeneratedDefaultValue="0.0",PersistentGuid=00000000000000000000000000000000,bHidden=False,bNotConnectable=False,bDefaultValueIsReadOnly=False,bDefaultValueIsIgnored=False,bAdvancedView=False,)
+
+	   CustomProperties Pin (PinId=46F0ACE24F439D2B6D101E8209961EAB,PinName="DamageTypeClass",PinToolTip="Damage Type Class\nDamage Type Class\n\nClass that describes the damage that was done.",PinType.PinCategory="class",PinType.PinSubCategory="",PinType.PinSubCategoryObject=Class'/Script/Engine.DamageType',PinType.PinSubCategoryMemberReference=(),PinType.PinValueType=(),PinType.bIsMap=False,PinType.bIsSet=False,PinType.bIsArray=False,PinType.bIsReference=False,PinType.bIsConst=False,PinType.bIsWeakPointer=False,DefaultObject="/Script/Odin.OdinDamageType",PersistentGuid=00000000000000000000000000000000,bHidden=False,bNotConnectable=False,bDefaultValueIsReadOnly=False,bDefaultValueIsIgnored=False,bAdvancedView=False,)
+
+	   CustomProperties Pin (PinId=7C1D1AA1492CD24C99658280C756C5DC,PinName="IgnoreActors",PinToolTip="Ignore Actors\nArray of Actor References",PinType.PinCategory="object",PinType.PinSubCategory="",PinType.PinSubCategoryObject=Class'/Script/Engine.Actor',PinType.PinSubCategoryMemberReference=(),PinType.PinValueType=(),PinType.bIsMap=False,PinType.bIsSet=False,PinType.bIsArray=True,PinType.bIsReference=True,PinType.bIsConst=True,PinType.bIsWeakPointer=False,PersistentGuid=00000000000000000000000000000000,bHidden=False,bNotConnectable=False,bDefaultValueIsReadOnly=False,bDefaultValueIsIgnored=True,bAdvancedView=False,)
+
+	   CustomProperties Pin (PinId=A406A2894F64492E5253A6B4A6995988,PinName="DamageCauser",PinToolTip="Damage Causer\nActor Reference\n\nActor that actually caused the damage (e.g. the grenade that exploded).  This actor will not be damaged and it will not block damage.",PinType.PinCategory="object",PinType.PinSubCategory="",PinType.PinSubCategoryObject=Class'/Script/Engine.Actor',PinType.PinSubCategoryMemberReference=(),PinType.PinValueType=(),PinType.bIsMap=False,PinType.bIsSet=False,PinType.bIsArray=False,PinType.bIsReference=False,PinType.bIsConst=False,PinType.bIsWeakPointer=False,LinkedTo=(K2Node_Self_0 C455D38B417906467DD289B1BA690A5A,),PersistentGuid=00000000000000000000000000000000,bHidden=False,bNotConnectable=False,bDefaultValueIsReadOnly=False,bDefaultValueIsIgnored=False,bAdvancedView=False,)
+
+	   CustomProperties Pin (PinId=1D1FED694F2543FC8EBDC6BE50EF1389,PinName="InstigatedByController",PinToolTip="Instigated by Controller\nController Reference\n\nController that was responsible for causing this damage (e.g. player who threw the grenade)",PinType.PinCategory="object",PinType.PinSubCategory="",PinType.PinSubCategoryObject=Class'/Script/Engine.Controller',PinType.PinSubCategoryMemberReference=(),PinType.PinValueType=(),PinType.bIsMap=False,PinType.bIsSet=False,PinType.bIsArray=False,PinType.bIsReference=False,PinType.bIsConst=False,PinType.bIsWeakPointer=False,LinkedTo=(K2Node_CallFunction_5 5D766A03432CFC38C6EAB9AF111871AC,),PersistentGuid=00000000000000000000000000000000,bHidden=False,bNotConnectable=False,bDefaultValueIsReadOnly=False,bDefaultValueIsIgnored=False,bAdvancedView=False,)
+
+	   CustomProperties Pin (PinId=0ED47B1D486F44A1FBB0D68A0836263D,PinName="bDoFullDamage",PinToolTip="Do Full Damage\nBoolean",PinType.PinCategory="bool",PinType.PinSubCategory="",PinType.PinSubCategoryObject=None,PinType.PinSubCategoryMemberReference=(),PinType.PinValueType=(),PinType.bIsMap=False,PinType.bIsSet=False,PinType.bIsArray=False,PinType.bIsReference=False,PinType.bIsConst=False,PinType.bIsWeakPointer=False,DefaultValue="true",AutogeneratedDefaultValue="false",PersistentGuid=00000000000000000000000000000000,bHidden=False,bNotConnectable=False,bDefaultValueIsReadOnly=False,bDefaultValueIsIgnored=False,bAdvancedView=False,)
+
+	   CustomProperties Pin (PinId=F24CAF574F3C041F4804D8840CDBB403,PinName="DamagePreventionChannel",PinToolTip="Damage Prevention Channel\nECollisionChannel Enum\n\nDamage will not be applied to victim if there is something between the origin and the victim which blocks traces on this channel",PinType.PinCategory="byte",PinType.PinSubCategory="",PinType.PinSubCategoryObject=Enum'/Script/Engine.ECollisionChannel',PinType.PinSubCategoryMemberReference=(),PinType.PinValueType=(),PinType.bIsMap=False,PinType.bIsSet=False,PinType.bIsArray=False,PinType.bIsReference=False,PinType.bIsConst=False,PinType.bIsWeakPointer=False,DefaultValue="ECC_Visibility",AutogeneratedDefaultValue="ECC_Visibility",PersistentGuid=00000000000000000000000000000000,bHidden=False,bNotConnectable=False,bDefaultValueIsReadOnly=False,bDefaultValueIsIgnored=False,bAdvancedView=False,)
+
+	   CustomProperties Pin (PinId=FF5ECB02491A3EDBC24D3B99A4A81C3A,PinName="ReturnValue",PinToolTip="Return Value\nBoolean\n\ntrue if damage was applied to at least one actor.",Direction="EGPD_Output",PinType.PinCategory="bool",PinType.PinSubCategory="",PinType.PinSubCategoryObject=None,PinType.PinSubCategoryMemberReference=(),PinType.PinValueType=(),PinType.bIsMap=False,PinType.bIsSet=False,PinType.bIsArray=False,PinType.bIsReference=False,PinType.bIsConst=False,PinType.bIsWeakPointer=False,DefaultValue="false",AutogeneratedDefaultValue="false",PersistentGuid=00000000000000000000000000000000,bHidden=False,bNotConnectable=False,bDefaultValueIsReadOnly=False,bDefaultValueIsIgnored=False,bAdvancedView=False,)
+
+	End Object
+
+	Begin Object Class=K2Node_Self Name="K2Node_Self_0"
+
+	   NodePosX=896
+
+	   NodePosY=880
+
+	   NodeGuid=E649F6B849A0E6C29106DAB80BBFB420
+
+	   CustomProperties Pin (PinId=C455D38B417906467DD289B1BA690A5A,PinName="self",Direction="EGPD_Output",PinType.PinCategory="object",PinType.PinSubCategory="self",PinType.PinSubCategoryObject=None,PinType.PinSubCategoryMemberReference=(),PinType.PinValueType=(),PinType.bIsMap=False,PinType.bIsSet=False,PinType.bIsArray=False,PinType.bIsReference=False,PinType.bIsConst=False,PinType.bIsWeakPointer=False,LinkedTo=(K2Node_CallFunction_3 A406A2894F64492E5253A6B4A6995988,),PersistentGuid=00000000000000000000000000000000,bHidden=False,bNotConnectable=False,bDefaultValueIsReadOnly=False,bDefaultValueIsIgnored=False,bAdvancedView=False,)
+
+	End Object
+
+	Begin Object Class=K2Node_CallFunction Name="K2Node_CallFunction_5"
+
+	   bIsPureFunc=True
+
+	   FunctionReference=(MemberParent=Class'/Script/Engine.GameplayStatics',MemberName="GetPlayerController")
+
+	   NodePosX=784
+
+	   NodePosY=928
+
+	   NodeGuid=52DA79524EC3D3EA20BF3D910C3A0BCD
+
+	   CustomProperties Pin (PinId=0131FC9344346AEAF4934E912760A869,PinName="self",PinFriendlyName=NSLOCTEXT("K2Node", "Target", "Target"),PinType.PinCategory="object",PinType.PinSubCategory="",PinType.PinSubCategoryObject=Class'/Script/Engine.GameplayStatics',PinType.PinSubCategoryMemberReference=(),PinType.PinValueType=(),PinType.bIsMap=False,PinType.bIsSet=False,PinType.bIsArray=False,PinType.bIsReference=False,PinType.bIsConst=False,PinType.bIsWeakPointer=False,DefaultObject="/Script/Engine.Default__GameplayStatics",PersistentGuid=00000000000000000000000000000000,bHidden=True,bNotConnectable=False,bDefaultValueIsReadOnly=False,bDefaultValueIsIgnored=False,bAdvancedView=False,)
+
+	   CustomProperties Pin (PinId=3ABC149F4E2E97843A58A2B430877955,PinName="WorldContextObject",PinType.PinCategory="object",PinType.PinSubCategory="",PinType.PinSubCategoryObject=Class'/Script/CoreUObject.Object',PinType.PinSubCategoryMemberReference=(),PinType.PinValueType=(),PinType.bIsMap=False,PinType.bIsSet=False,PinType.bIsArray=False,PinType.bIsReference=False,PinType.bIsConst=True,PinType.bIsWeakPointer=False,PersistentGuid=00000000000000000000000000000000,bHidden=True,bNotConnectable=False,bDefaultValueIsReadOnly=False,bDefaultValueIsIgnored=False,bAdvancedView=False,)
+
+	   CustomProperties Pin (PinId=D261DA134C596EC8F6AF5E983AEA907F,PinName="PlayerIndex",PinType.PinCategory="int",PinType.PinSubCategory="",PinType.PinSubCategoryObject=None,PinType.PinSubCategoryMemberReference=(),PinType.PinValueType=(),PinType.bIsMap=False,PinType.bIsSet=False,PinType.bIsArray=False,PinType.bIsReference=False,PinType.bIsConst=False,PinType.bIsWeakPointer=False,DefaultValue="0",AutogeneratedDefaultValue="0",PersistentGuid=00000000000000000000000000000000,bHidden=False,bNotConnectable=False,bDefaultValueIsReadOnly=False,bDefaultValueIsIgnored=False,bAdvancedView=False,)
+
+	   CustomProperties Pin (PinId=5D766A03432CFC38C6EAB9AF111871AC,PinName="ReturnValue",Direction="EGPD_Output",PinType.PinCategory="object",PinType.PinSubCategory="",PinType.PinSubCategoryObject=Class'/Script/Engine.PlayerController',PinType.PinSubCategoryMemberReference=(),PinType.PinValueType=(),PinType.bIsMap=False,PinType.bIsSet=False,PinType.bIsArray=False,PinType.bIsReference=False,PinType.bIsConst=False,PinType.bIsWeakPointer=False,LinkedTo=(K2Node_CallFunction_3 1D1FED694F2543FC8EBDC6BE50EF1389,),PersistentGuid=00000000000000000000000000000000,bHidden=False,bNotConnectable=False,bDefaultValueIsReadOnly=False,bDefaultValueIsIgnored=False,bAdvancedView=False,)
+
+	End Object
+
+	Begin Object Class=K2Node_Knot Name="K2Node_Knot_545"
+
+	   NodePosX=832
+
+	   NodePosY=864
+
+	   NodeGuid=17DE826548331E66FC59AD9F8AE8C477
+
+	   CustomProperties Pin (PinId=3112B9B54E93377E8BF8F7A5C78DB120,PinName="InputPin",PinType.PinCategory="struct",PinType.PinSubCategory="",PinType.PinSubCategoryObject=ScriptStruct'/Script/CoreUObject.Vector',PinType.PinSubCategoryMemberReference=(),PinType.PinValueType=(),PinType.bIsMap=False,PinType.bIsSet=False,PinType.bIsArray=False,PinType.bIsReference=False,PinType.bIsConst=False,PinType.bIsWeakPointer=False,LinkedTo=(K2Node_Knot_546 1D7421E14F9B31B78E2CFDAACAC148EB,),PersistentGuid=00000000000000000000000000000000,bHidden=False,bNotConnectable=False,bDefaultValueIsReadOnly=False,bDefaultValueIsIgnored=True,bAdvancedView=False,)
+
+	   CustomProperties Pin (PinId=38F57ED3452FC48D3E3F858071F90765,PinName="OutputPin",Direction="EGPD_Output",PinType.PinCategory="struct",PinType.PinSubCategory="",PinType.PinSubCategoryObject=ScriptStruct'/Script/CoreUObject.Vector',PinType.PinSubCategoryMemberReference=(),PinType.PinValueType=(),PinType.bIsMap=False,PinType.bIsSet=False,PinType.bIsArray=False,PinType.bIsReference=False,PinType.bIsConst=False,PinType.bIsWeakPointer=False,LinkedTo=(K2Node_CallFunction_3 2951FACB4839E3137D9FD09BFE3CC47F,),PersistentGuid=00000000000000000000000000000000,bHidden=False,bNotConnectable=False,bDefaultValueIsReadOnly=False,bDefaultValueIsIgnored=False,bAdvancedView=False,)
+
+	End Object
+
+	Begin Object Class=K2Node_CommutativeAssociativeBinaryOperator Name="K2Node_CommutativeAssociativeBinaryOperator_1"
+
+	   bIsPureFunc=True
+
+	   FunctionReference=(MemberParent=Class'/Script/Engine.KismetMathLibrary',MemberName="BooleanAND")
+
+	   NodePosX=384
+
+	   NodePosY=720
+
+	   NodeGuid=5F1C067B48DB870E641E1A91108626D1
+
+	   CustomProperties Pin (PinId=7EBD044A4E6145D3EB3365A9E0B6B1FC,PinName="self",PinFriendlyName=NSLOCTEXT("K2Node", "Target", "Target"),PinType.PinCategory="object",PinType.PinSubCategory="",PinType.PinSubCategoryObject=Class'/Script/Engine.KismetMathLibrary',PinType.PinSubCategoryMemberReference=(),PinType.PinValueType=(),PinType.bIsMap=False,PinType.bIsSet=False,PinType.bIsArray=False,PinType.bIsReference=False,PinType.bIsConst=False,PinType.bIsWeakPointer=False,DefaultObject="/Script/Engine.Default__KismetMathLibrary",PersistentGuid=00000000000000000000000000000000,bHidden=True,bNotConnectable=False,bDefaultValueIsReadOnly=False,bDefaultValueIsIgnored=False,bAdvancedView=False,)
+
+	   CustomProperties Pin (PinId=5DB3157C48E4D9D9B2A4F5B237BB4B05,PinName="A",PinType.PinCategory="bool",PinType.PinSubCategory="",PinType.PinSubCategoryObject=None,PinType.PinSubCategoryMemberReference=(),PinType.PinValueType=(),PinType.bIsMap=False,PinType.bIsSet=False,PinType.bIsArray=False,PinType.bIsReference=False,PinType.bIsConst=False,PinType.bIsWeakPointer=False,DefaultValue="false",AutogeneratedDefaultValue="false",LinkedTo=(K2Node_CallFunction_1 C57D3581488D11C825293B9D27152EAA,),PersistentGuid=00000000000000000000000000000000,bHidden=False,bNotConnectable=False,bDefaultValueIsReadOnly=False,bDefaultValueIsIgnored=False,bAdvancedView=False,)
+
+	   CustomProperties Pin (PinId=621DC1E84469C5C518DA709A73A026D3,PinName="B",PinType.PinCategory="bool",PinType.PinSubCategory="",PinType.PinSubCategoryObject=None,PinType.PinSubCategoryMemberReference=(),PinType.PinValueType=(),PinType.bIsMap=False,PinType.bIsSet=False,PinType.bIsArray=False,PinType.bIsReference=False,PinType.bIsConst=False,PinType.bIsWeakPointer=False,DefaultValue="false",AutogeneratedDefaultValue="false",LinkedTo=(K2Node_VariableGet_3890 67F622DF47E23B683771228981CB162A,),PersistentGuid=00000000000000000000000000000000,bHidden=False,bNotConnectable=False,bDefaultValueIsReadOnly=False,bDefaultValueIsIgnored=False,bAdvancedView=False,)
+
+	   CustomProperties Pin (PinId=7B8DC675463DABC748E5C1BC81E5CA8F,PinName="ReturnValue",Direction="EGPD_Output",PinType.PinCategory="bool",PinType.PinSubCategory="",PinType.PinSubCategoryObject=None,PinType.PinSubCategoryMemberReference=(),PinType.PinValueType=(),PinType.bIsMap=False,PinType.bIsSet=False,PinType.bIsArray=False,PinType.bIsReference=False,PinType.bIsConst=False,PinType.bIsWeakPointer=False,AutogeneratedDefaultValue="false",LinkedTo=(K2Node_IfThenElse_0 78E5DA334930E09580EF89A92BF81356,),PersistentGuid=00000000000000000000000000000000,bHidden=False,bNotConnectable=False,bDefaultValueIsReadOnly=False,bDefaultValueIsIgnored=False,bAdvancedView=False,)
+
+	End Object
+
+	Begin Object Class=K2Node_VariableGet Name="K2Node_VariableGet_3890"
+
+	   VariableReference=(MemberName="bCanExplode",MemberGuid=9BCA29C8404B219E30C266861239680A,bSelfContext=True)
+
+	   NodePosX=192
+
+	   NodePosY=768
+
+	   NodeGuid=86C4B2684F0EADBAF3EE9D935A3206E4
+
+	   CustomProperties Pin (PinId=67F622DF47E23B683771228981CB162A,PinName="bCanExplode",Direction="EGPD_Output",PinType.PinCategory="bool",PinType.PinSubCategory="",PinType.PinSubCategoryObject=None,PinType.PinSubCategoryMemberReference=(),PinType.PinValueType=(),PinType.bIsMap=False,PinType.bIsSet=False,PinType.bIsArray=False,PinType.bIsReference=False,PinType.bIsConst=False,PinType.bIsWeakPointer=False,AutogeneratedDefaultValue="false",LinkedTo=(K2Node_CommutativeAssociativeBinaryOperator_1 621DC1E84469C5C518DA709A73A026D3,),PersistentGuid=00000000000000000000000000000000,bHidden=False,bNotConnectable=False,bDefaultValueIsReadOnly=False,bDefaultValueIsIgnored=False,bAdvancedView=False,)
+
+	   CustomProperties Pin (PinId=99CA6DAD437625FCFFB1A9B99D79CC7C,PinName="self",PinFriendlyName=NSLOCTEXT("K2Node", "Target", "Target"),PinType.PinCategory="object",PinType.PinSubCategory="",PinType.PinSubCategoryObject=BlueprintGeneratedClass'/GrenadeGun/Blueprints/Weapons/GrenadeGunPistol.GrenadeGunPistol_C',PinType.PinSubCategoryMemberReference=(),PinType.PinValueType=(),PinType.bIsMap=False,PinType.bIsSet=False,PinType.bIsArray=False,PinType.bIsReference=False,PinType.bIsConst=False,PinType.bIsWeakPointer=False,PersistentGuid=00000000000000000000000000000000,bHidden=True,bNotConnectable=False,bDefaultValueIsReadOnly=False,bDefaultValueIsIgnored=False,bAdvancedView=False,)
+
+	End Object
+
+	Begin Object Class=K2Node_MacroInstance Name="K2Node_MacroInstance_0"
+
+	   MacroGraphReference=(MacroGraph=EdGraph'/Engine/EditorBlueprintResources/StandardMacros.StandardMacros:DoOnce',GraphBlueprint=Blueprint'/Engine/EditorBlueprintResources/StandardMacros.StandardMacros',GraphGuid=1281F54248A2ECB5B8B2C5B24AE6FDF4)
+
+	   NodePosX=768
+
+	   NodePosY=640
+
+	   NodeGuid=748F09D5462211ADA9B0F7BEDE3C62A6
+
+	   CustomProperties Pin (PinId=2B445179486EFD71E4EB229C0D980F97,PinName="execute",PinType.PinCategory="exec",PinType.PinSubCategory="",PinType.PinSubCategoryObject=None,PinType.PinSubCategoryMemberReference=(),PinType.PinValueType=(),PinType.bIsMap=False,PinType.bIsSet=False,PinType.bIsArray=False,PinType.bIsReference=False,PinType.bIsConst=False,PinType.bIsWeakPointer=False,LinkedTo=(K2Node_IfThenElse_0 899DE7DA4D10D721FAF47DB62C728E8C,),PersistentGuid=00000000000000000000000000000000,bHidden=False,bNotConnectable=False,bDefaultValueIsReadOnly=False,bDefaultValueIsIgnored=False,bAdvancedView=False,)
+
+	   CustomProperties Pin (PinId=6BAE0D934838B2C2F364C683B8361A92,PinName="Reset",PinType.PinCategory="exec",PinType.PinSubCategory="",PinType.PinSubCategoryObject=None,PinType.PinSubCategoryMemberReference=(),PinType.PinValueType=(),PinType.bIsMap=False,PinType.bIsSet=False,PinType.bIsArray=False,PinType.bIsReference=False,PinType.bIsConst=False,PinType.bIsWeakPointer=False,PersistentGuid=00000000000000000000000000000000,bHidden=False,bNotConnectable=False,bDefaultValueIsReadOnly=False,bDefaultValueIsIgnored=False,bAdvancedView=False,)
+
+	   CustomProperties Pin (PinId=17D9B10D4074CE6727A8AB8DF1984591,PinName="Start Closed",PinType.PinCategory="bool",PinType.PinSubCategory="",PinType.PinSubCategoryObject=None,PinType.PinSubCategoryMemberReference=(),PinType.PinValueType=(),PinType.bIsMap=False,PinType.bIsSet=False,PinType.bIsArray=False,PinType.bIsReference=False,PinType.bIsConst=False,PinType.bIsWeakPointer=False,PersistentGuid=00000000000000000000000000000000,bHidden=False,bNotConnectable=False,bDefaultValueIsReadOnly=False,bDefaultValueIsIgnored=False,bAdvancedView=False,)
+
+	   CustomProperties Pin (PinId=7B80FE2849D77B7B2400F69815D9369F,PinName="Completed",Direction="EGPD_Output",PinType.PinCategory="exec",PinType.PinSubCategory="",PinType.PinSubCategoryObject=None,PinType.PinSubCategoryMemberReference=(),PinType.PinValueType=(),PinType.bIsMap=False,PinType.bIsSet=False,PinType.bIsArray=False,PinType.bIsReference=False,PinType.bIsConst=False,PinType.bIsWeakPointer=False,LinkedTo=(K2Node_CallFunction_3 22BE44D544FB103B193288B3A615D3BC,),PersistentGuid=00000000000000000000000000000000,bHidden=False,bNotConnectable=False,bDefaultValueIsReadOnly=False,bDefaultValueIsIgnored=False,bAdvancedView=False,)
+
+	End Object
+
+	Begin Object Class=K2Node_VariableGet Name="K2Node_VariableGet_7"
+
+	   VariableReference=(MemberName="AmmoRemaining",bSelfContext=True)
+
+	   NodePosX=400
+
+	   NodePosY=800
+
+	   NodeGuid=B2014F4E4D216D43CDF52ABAC68002A9
+
+	   CustomProperties Pin (PinId=915FC79F4607E52143A197AEB2281236,PinName="AmmoRemaining",Direction="EGPD_Output",PinType.PinCategory="int",PinType.PinSubCategory="",PinType.PinSubCategoryObject=None,PinType.PinSubCategoryMemberReference=(),PinType.PinValueType=(),PinType.bIsMap=False,PinType.bIsSet=False,PinType.bIsArray=False,PinType.bIsReference=False,PinType.bIsConst=False,PinType.bIsWeakPointer=False,AutogeneratedDefaultValue="0",LinkedTo=(K2Node_CallFunction_163 BEDB0BDB4C3FF69C501902BF11C98C55,),PersistentGuid=00000000000000000000000000000000,bHidden=False,bNotConnectable=False,bDefaultValueIsReadOnly=False,bDefaultValueIsIgnored=False,bAdvancedView=False,)
+
+	   CustomProperties Pin (PinId=5C274E104A67F344FA8628B922901D65,PinName="self",PinFriendlyName=NSLOCTEXT("K2Node", "Target", "Target"),PinType.PinCategory="object",PinType.PinSubCategory="",PinType.PinSubCategoryObject=Class'/Script/Odin.OdinWeapon',PinType.PinSubCategoryMemberReference=(),PinType.PinValueType=(),PinType.bIsMap=False,PinType.bIsSet=False,PinType.bIsArray=False,PinType.bIsReference=False,PinType.bIsConst=False,PinType.bIsWeakPointer=False,PersistentGuid=00000000000000000000000000000000,bHidden=True,bNotConnectable=False,bDefaultValueIsReadOnly=False,bDefaultValueIsIgnored=False,bAdvancedView=False,)
+
+	End Object
+
+	Begin Object Class=K2Node_CallFunction Name="K2Node_CallFunction_163"
+
+	   bIsPureFunc=True
+
+	   FunctionReference=(MemberParent=Class'/Script/Engine.KismetMathLibrary',MemberName="Multiply_IntFloat")
+
+	   NodePosX=592
+
+	   NodePosY=784
+
+	   NodeGuid=E5EC4C2346D74D275F83F88776AF1AE1
+
+	   CustomProperties Pin (PinId=9BA2C45C4E148ED28830EB89690601B9,PinName="self",PinFriendlyName=NSLOCTEXT("K2Node", "Target", "Target"),PinType.PinCategory="object",PinType.PinSubCategory="",PinType.PinSubCategoryObject=Class'/Script/Engine.KismetMathLibrary',PinType.PinSubCategoryMemberReference=(),PinType.PinValueType=(),PinType.bIsMap=False,PinType.bIsSet=False,PinType.bIsArray=False,PinType.bIsReference=False,PinType.bIsConst=False,PinType.bIsWeakPointer=False,DefaultObject="/Script/Engine.Default__KismetMathLibrary",PersistentGuid=00000000000000000000000000000000,bHidden=True,bNotConnectable=False,bDefaultValueIsReadOnly=False,bDefaultValueIsIgnored=False,bAdvancedView=False,)
+
+	   CustomProperties Pin (PinId=BEDB0BDB4C3FF69C501902BF11C98C55,PinName="A",PinType.PinCategory="int",PinType.PinSubCategory="",PinType.PinSubCategoryObject=None,PinType.PinSubCategoryMemberReference=(),PinType.PinValueType=(),PinType.bIsMap=False,PinType.bIsSet=False,PinType.bIsArray=False,PinType.bIsReference=False,PinType.bIsConst=False,PinType.bIsWeakPointer=False,DefaultValue="0",AutogeneratedDefaultValue="0",LinkedTo=(K2Node_VariableGet_7 915FC79F4607E52143A197AEB2281236,),PersistentGuid=00000000000000000000000000000000,bHidden=False,bNotConnectable=False,bDefaultValueIsReadOnly=False,bDefaultValueIsIgnored=False,bAdvancedView=False,)
+
+	   CustomProperties Pin (PinId=FA338AAD4E82204450249D93C7983E82,PinName="B",PinType.PinCategory="float",PinType.PinSubCategory="",PinType.PinSubCategoryObject=None,PinType.PinSubCategoryMemberReference=(),PinType.PinValueType=(),PinType.bIsMap=False,PinType.bIsSet=False,PinType.bIsArray=False,PinType.bIsReference=False,PinType.bIsConst=False,PinType.bIsWeakPointer=False,DefaultValue="1000",AutogeneratedDefaultValue="0.0",PersistentGuid=00000000000000000000000000000000,bHidden=False,bNotConnectable=False,bDefaultValueIsReadOnly=False,bDefaultValueIsIgnored=False,bAdvancedView=False,)
+
+	   CustomProperties Pin (PinId=15D69B10433010FDF68F04AE2C08D7B5,PinName="ReturnValue",Direction="EGPD_Output",PinType.PinCategory="float",PinType.PinSubCategory="",PinType.PinSubCategoryObject=None,PinType.PinSubCategoryMemberReference=(),PinType.PinValueType=(),PinType.bIsMap=False,PinType.bIsSet=False,PinType.bIsArray=False,PinType.bIsReference=False,PinType.bIsConst=False,PinType.bIsWeakPointer=False,AutogeneratedDefaultValue="0.0",LinkedTo=(K2Node_CommutativeAssociativeBinaryOperator_10 6500E078476E3740DE23DF80643FC3FD,),PersistentGuid=00000000000000000000000000000000,bHidden=False,bNotConnectable=False,bDefaultValueIsReadOnly=False,bDefaultValueIsIgnored=False,bAdvancedView=False,)
+
+	End Object
+
+	Begin Object Class=K2Node_Knot Name="K2Node_Knot_546"
+
+	   NodePosX=240
+
+	   NodePosY=864
+
+	   NodeGuid=DD7B52CE4DE7C80BFA56AB850F88BC39
+
+	   CustomProperties Pin (PinId=6EEFCD4849DD5591D5B1CB9B8FB76133,PinName="InputPin",PinType.PinCategory="struct",PinType.PinSubCategory="",PinType.PinSubCategoryObject=ScriptStruct'/Script/CoreUObject.Vector',PinType.PinSubCategoryMemberReference=(),PinType.PinValueType=(),PinType.bIsMap=False,PinType.bIsSet=False,PinType.bIsArray=False,PinType.bIsReference=False,PinType.bIsConst=False,PinType.bIsWeakPointer=False,LinkedTo=(K2Node_Event_3 3739097840F8ACAB5630D8AEA7CADC9B,),PersistentGuid=00000000000000000000000000000000,bHidden=False,bNotConnectable=False,bDefaultValueIsReadOnly=False,bDefaultValueIsIgnored=True,bAdvancedView=False,)
+
+	   CustomProperties Pin (PinId=1D7421E14F9B31B78E2CFDAACAC148EB,PinName="OutputPin",Direction="EGPD_Output",PinType.PinCategory="struct",PinType.PinSubCategory="",PinType.PinSubCategoryObject=ScriptStruct'/Script/CoreUObject.Vector',PinType.PinSubCategoryMemberReference=(),PinType.PinValueType=(),PinType.bIsMap=False,PinType.bIsSet=False,PinType.bIsArray=False,PinType.bIsReference=False,PinType.bIsConst=False,PinType.bIsWeakPointer=False,LinkedTo=(K2Node_Knot_545 3112B9B54E93377E8BF8F7A5C78DB120,),PersistentGuid=00000000000000000000000000000000,bHidden=False,bNotConnectable=False,bDefaultValueIsReadOnly=False,bDefaultValueIsIgnored=False,bAdvancedView=False,)
+
+	End Object
+
+	Begin Object Class=K2Node_CommutativeAssociativeBinaryOperator Name="K2Node_CommutativeAssociativeBinaryOperator_10"
+
+	   bIsPureFunc=True
+
+	   FunctionReference=(MemberParent=Class'/Script/Engine.KismetMathLibrary',MemberName="Add_FloatFloat")
+
+	   NodePosX=784
+
+	   NodePosY=784
+
+	   NodeGuid=ACAED5DD4662A84D55C1C9AC8B9A57F2
+
+	   CustomProperties Pin (PinId=F48331684C77102383429ABCA0CA94EE,PinName="self",PinFriendlyName=NSLOCTEXT("K2Node", "Target", "Target"),PinType.PinCategory="object",PinType.PinSubCategory="",PinType.PinSubCategoryObject=Class'/Script/Engine.KismetMathLibrary',PinType.PinSubCategoryMemberReference=(),PinType.PinValueType=(),PinType.bIsMap=False,PinType.bIsSet=False,PinType.bIsArray=False,PinType.bIsReference=False,PinType.bIsConst=False,PinType.bIsWeakPointer=False,DefaultObject="/Script/Engine.Default__KismetMathLibrary",PersistentGuid=00000000000000000000000000000000,bHidden=True,bNotConnectable=False,bDefaultValueIsReadOnly=False,bDefaultValueIsIgnored=False,bAdvancedView=False,)
+
+	   CustomProperties Pin (PinId=6500E078476E3740DE23DF80643FC3FD,PinName="A",PinType.PinCategory="float",PinType.PinSubCategory="",PinType.PinSubCategoryObject=None,PinType.PinSubCategoryMemberReference=(),PinType.PinValueType=(),PinType.bIsMap=False,PinType.bIsSet=False,PinType.bIsArray=False,PinType.bIsReference=False,PinType.bIsConst=False,PinType.bIsWeakPointer=False,DefaultValue="0.0",AutogeneratedDefaultValue="0.0",LinkedTo=(K2Node_CallFunction_163 15D69B10433010FDF68F04AE2C08D7B5,),PersistentGuid=00000000000000000000000000000000,bHidden=False,bNotConnectable=False,bDefaultValueIsReadOnly=False,bDefaultValueIsIgnored=False,bAdvancedView=False,)
+
+	   CustomProperties Pin (PinId=7339BB064C34963DE41EAC98E7F5FCBD,PinName="B",PinType.PinCategory="float",PinType.PinSubCategory="",PinType.PinSubCategoryObject=None,PinType.PinSubCategoryMemberReference=(),PinType.PinValueType=(),PinType.bIsMap=False,PinType.bIsSet=False,PinType.bIsArray=False,PinType.bIsReference=False,PinType.bIsConst=False,PinType.bIsWeakPointer=False,DefaultValue="1000",AutogeneratedDefaultValue="1.000000",PersistentGuid=00000000000000000000000000000000,bHidden=False,bNotConnectable=False,bDefaultValueIsReadOnly=False,bDefaultValueIsIgnored=False,bAdvancedView=False,)
+
+	   CustomProperties Pin (PinId=0DA0468E49FFB654F7C5D4841F83ADAC,PinName="ReturnValue",Direction="EGPD_Output",PinType.PinCategory="float",PinType.PinSubCategory="",PinType.PinSubCategoryObject=None,PinType.PinSubCategoryMemberReference=(),PinType.PinValueType=(),PinType.bIsMap=False,PinType.bIsSet=False,PinType.bIsArray=False,PinType.bIsReference=False,PinType.bIsConst=False,PinType.bIsWeakPointer=False,AutogeneratedDefaultValue="0.0",LinkedTo=(K2Node_CallFunction_3 EAD1F906420BA584197D0A838C95F09C,),PersistentGuid=00000000000000000000000000000000,bHidden=False,bNotConnectable=False,bDefaultValueIsReadOnly=False,bDefaultValueIsIgnored=False,bAdvancedView=False,)
+
+	End Object
+[/EXCERPT:snippet]
+-->
+
