Copying //UE4/Dev-Documentation to //UE4/Samples-Main (Source: //UE4/Dev-Documentation @ 3372893)

diff --git a/Engine/Documentation/Source/Modding/RoboRecall/HowTo/ExtendFunctionality/BuildGrabEvents/BuildGrabEvents.INT.udn b/Engine/Documentation/Source/Modding/RoboRecall/HowTo/ExtendFunctionality/BuildGrabEvents/BuildGrabEvents.INT.udn
new file mode 100644
--- /dev/null
+++ b/Engine/Documentation/Source/Modding/RoboRecall/HowTo/ExtendFunctionality/BuildGrabEvents/BuildGrabEvents.INT.udn
@@ -0,0 +1,310 @@
+Availability: Public
+Title:4. Building the Grab Events
+Crumbs:%ROOT%
+Description:Robo Recall Modding Documentation
+Parent:Modding/RoboRecall/HowTo/ExtendFunctionality
+type:how-to
+order:4
+tags:Modding
+tags:Robo Recall
+topic-image:Modding/RoboRecall/RoboRecall_Topic.png
+social-image:Modding/RoboRecall/RoboRecall_Social.png
+checkpoint:mod_extend
+
+<!-- Location: Modding/RoboRecall/Quickstarts/WeaponMod/4 -->
+
+[NAV]
+
+<!-- Title Image if Needed Here -->
+
+In this step you will be overriding the **Grab Events** on our GrenadeGunPistol so we can check to see if we want the gun to explode or not when it collides with things.
+
+## Steps
+
+1. **Right-click** in the [Event Graph](Engine/Blueprints/UserGuide/EventGraph) and search for **Begin Primary Grab**. Select **Event Begin Primary Grab** to add the node to the Event GRaph.
+
+	![image alt text](image_0.png)(w:409 h:409 convert:false)
+
+2. Repeat this process for:
+
+    * **Event Begin Secondary Grab**
+    * **Event End Primary Grab**
+    * **Event End Secondary Grab**
+
+3. For each of these Event nodes, you need to make a call to the same Event in the parent Blueprint to make sure that any existing functionality is still executed. To do this, **Right-click** the four event nodes you just created and select **Add call to parent function**.
+
+	![image alt text](image_1.png)(w:169 h:489 convert:false)
+
+4. Connect the Event Nodes to their Parent Nodes.
+
+	![image alt text](image_3.png)(w:522 h:460 convert:false)
+
+5. Click the **+** button next to **Variables** in the **My Blueprint** panel to create a new variable and name it **bCanExplode**.
+
+	![image alt text](image_4.png)(w:381 h:368 convert:false) ![image alt text](image_5.png)(w:381 h:388 convert:false)
+
+	[REGION:note]
+	If you haven't created a new variable in this session in Unreal Editor, it will automatically be a **Bool** (the red icon indicates that it is a **Bool**). However, if you need to change it, select the variable in the **My Blueprint** tab and then you can change it in the **Details** panel: 
+
+	![image alt text](image_6.png)(w:636 h:392 convert:false)
+	[/REGION]
+
+6. **Click + Drag** the **bCanExplode** variable into the **Event Graph**, and select the **Set** option.
+
+	[REGION:raw]
+	![image alt text](image_7.png)(w:793 h:328 convert:false)
+	[/REGION]
+
+7. Duplicate the **Set Can Explode** node using **Ctrl + W** and click the checkbox in the new node to set the value of **bCanExplode** to **True**.
+
+	![image alt text](image_8.png)(w:201 h:184 convert:false)
+
+8. Finally, connect the Parent nodes to the corresponding **Set Can Explode** nodes as shown below.
+
+	[OBJECT:Snippet]
+		[PARAMLITERAL:id]
+		result
+		[/PARAMLITERAL]
+		[PARAMLITERAL:type]
+		blueprint
+		[/PARAMLITERAL]
+		[PARAM:image]
+		![image alt text](image_9.png)(w:786 h:464 convert:false)
+		[/PARAM]
+		[PARAM:snippet]
+			[INCLUDE:#snippet]
+		[/PARAM]
+	[/OBJECT]
+
+	[REGION:caption]
+	Click the icon in the upper left corner of this image to copy the Blueprint Graph and paste it into your project.
+	[/REGION]
+
+9. In the Blueprint Editor toolbar, click **Compile** to update the Blueprint and click **Save** to preserve your changes.
+
+## Results
+
+These sets of Events are called when our gun is grabbed and released. You'll use the **bCanExplode** variable in the next section of the Blueprint to prevent the gun from exploding in your hand when you punch a robot with it.
+
+[NAV]
+ 
+<!--
+[EXCERPT:snippet]
+	Begin Object Class=K2Node_Event Name="K2Node_Event_7"
+
+		EventReference=(MemberParent=BlueprintGeneratedClass'/Game/Prototype/Game/Gun_Base.Gun_Base_C',MemberName="End Primary Grab",MemberGuid=5AC0E94444892A8794605F86BF8B03E2)
+
+		bOverrideFunction=True
+
+		NodePosX=48
+
+		NodePosY=1584
+
+		NodeGuid=4DB300D94131B538C1B2E1AFF11EB6B2
+
+		CustomProperties Pin (PinId=59D3C9964C20448F1FFD22B92B32E3F2,PinName="OutputDelegate",Direction="EGPD_Output",PinType.PinCategory="delegate",PinType.PinSubCategory="",PinType.PinSubCategoryObject=None,PinType.PinSubCategoryMemberReference=(MemberParent=BlueprintGeneratedClass'/Game/Prototype/Game/Gun_Base.Gun_Base_C',MemberName="End Primary Grab",MemberGuid=5AC0E94444892A8794605F86BF8B03E2),PinType.PinValueType=(),PinType.bIsMap=False,PinType.bIsSet=False,PinType.bIsArray=False,PinType.bIsReference=False,PinType.bIsConst=False,PinType.bIsWeakPointer=False,PersistentGuid=00000000000000000000000000000000,bHidden=False,bNotConnectable=False,bDefaultValueIsReadOnly=False,bDefaultValueIsIgnored=False,bAdvancedView=False,)
+
+		CustomProperties Pin (PinId=DDE6DFF14B7B9ED8FABA8A8A2B48E090,PinName="then",Direction="EGPD_Output",PinType.PinCategory="exec",PinType.PinSubCategory="",PinType.PinSubCategoryObject=None,PinType.PinSubCategoryMemberReference=(),PinType.PinValueType=(),PinType.bIsMap=False,PinType.bIsSet=False,PinType.bIsArray=False,PinType.bIsReference=False,PinType.bIsConst=False,PinType.bIsWeakPointer=False,LinkedTo=(K2Node_CallParentFunction_3 440FA2384D499E510943C5952AED1D52,),PersistentGuid=00000000000000000000000000000000,bHidden=False,bNotConnectable=False,bDefaultValueIsReadOnly=False,bDefaultValueIsIgnored=False,bAdvancedView=False,)
+
+	End Object
+
+	Begin Object Class=K2Node_Event Name="K2Node_Event_10"
+
+		EventReference=(MemberParent=BlueprintGeneratedClass'/Game/Prototype/Game/Gun_Base.Gun_Base_C',MemberName="Begin Primary Grab",MemberGuid=86E248C047254FEECC0603BB6C0FE484)
+
+		bOverrideFunction=True
+
+		NodePosX=48
+
+		NodePosY=1200
+
+		NodeGuid=2698FDBE4F93DAEDA8C18FAAB85204CC
+
+		CustomProperties Pin (PinId=D5884B0048B127BF178702BB8251314F,PinName="OutputDelegate",Direction="EGPD_Output",PinType.PinCategory="delegate",PinType.PinSubCategory="",PinType.PinSubCategoryObject=None,PinType.PinSubCategoryMemberReference=(MemberParent=BlueprintGeneratedClass'/Game/Prototype/Game/Gun_Base.Gun_Base_C',MemberName="Begin Primary Grab",MemberGuid=86E248C047254FEECC0603BB6C0FE484),PinType.PinValueType=(),PinType.bIsMap=False,PinType.bIsSet=False,PinType.bIsArray=False,PinType.bIsReference=False,PinType.bIsConst=False,PinType.bIsWeakPointer=False,PersistentGuid=00000000000000000000000000000000,bHidden=False,bNotConnectable=False,bDefaultValueIsReadOnly=False,bDefaultValueIsIgnored=False,bAdvancedView=False,)
+
+		CustomProperties Pin (PinId=01317E634C1C07F9D5D729B8421618F2,PinName="then",Direction="EGPD_Output",PinType.PinCategory="exec",PinType.PinSubCategory="",PinType.PinSubCategoryObject=None,PinType.PinSubCategoryMemberReference=(),PinType.PinValueType=(),PinType.bIsMap=False,PinType.bIsSet=False,PinType.bIsArray=False,PinType.bIsReference=False,PinType.bIsConst=False,PinType.bIsWeakPointer=False,LinkedTo=(K2Node_CallParentFunction_0 FDB36E2741A2BCEF856502AAC746B07F,),PersistentGuid=00000000000000000000000000000000,bHidden=False,bNotConnectable=False,bDefaultValueIsReadOnly=False,bDefaultValueIsIgnored=False,bAdvancedView=False,)
+
+	End Object
+
+	Begin Object Class=K2Node_CallParentFunction Name="K2Node_CallParentFunction_0"
+
+		FunctionReference=(MemberParent=BlueprintGeneratedClass'/Game/Prototype/Game/Gun_Pistol.Gun_Pistol_C',MemberName="Begin Primary Grab",MemberGuid=86E248C047254FEECC0603BB6C0FE484)
+
+		NodePosX=320
+
+		NodePosY=1200
+
+		NodeGuid=F79E93794CF4A9FF0F133187E49F3FF4
+
+		CustomProperties Pin (PinId=FDB36E2741A2BCEF856502AAC746B07F,PinName="execute",PinType.PinCategory="exec",PinType.PinSubCategory="",PinType.PinSubCategoryObject=None,PinType.PinSubCategoryMemberReference=(),PinType.PinValueType=(),PinType.bIsMap=False,PinType.bIsSet=False,PinType.bIsArray=False,PinType.bIsReference=False,PinType.bIsConst=False,PinType.bIsWeakPointer=False,LinkedTo=(K2Node_Event_10 01317E634C1C07F9D5D729B8421618F2,),PersistentGuid=00000000000000000000000000000000,bHidden=False,bNotConnectable=False,bDefaultValueIsReadOnly=False,bDefaultValueIsIgnored=False,bAdvancedView=False,)
+
+		CustomProperties Pin (PinId=1B74752B4533A65EB598E787671714A5,PinName="then",Direction="EGPD_Output",PinType.PinCategory="exec",PinType.PinSubCategory="",PinType.PinSubCategoryObject=None,PinType.PinSubCategoryMemberReference=(),PinType.PinValueType=(),PinType.bIsMap=False,PinType.bIsSet=False,PinType.bIsArray=False,PinType.bIsReference=False,PinType.bIsConst=False,PinType.bIsWeakPointer=False,LinkedTo=(K2Node_VariableSet_0 AE7B514B460DD96C37ECE7A72350ADB1,),PersistentGuid=00000000000000000000000000000000,bHidden=False,bNotConnectable=False,bDefaultValueIsReadOnly=False,bDefaultValueIsIgnored=False,bAdvancedView=False,)
+
+		CustomProperties Pin (PinId=A448FB8B4EEA51822C66DBA48B3087AA,PinName="self",PinFriendlyName=NSLOCTEXT("K2Node", "Target", "Target"),PinType.PinCategory="object",PinType.PinSubCategory="",PinType.PinSubCategoryObject=BlueprintGeneratedClass'/Game/Prototype/Game/Gun_Base.Gun_Base_C',PinType.PinSubCategoryMemberReference=(),PinType.PinValueType=(),PinType.bIsMap=False,PinType.bIsSet=False,PinType.bIsArray=False,PinType.bIsReference=False,PinType.bIsConst=False,PinType.bIsWeakPointer=False,PersistentGuid=00000000000000000000000000000000,bHidden=True,bNotConnectable=False,bDefaultValueIsReadOnly=False,bDefaultValueIsIgnored=False,bAdvancedView=False,)
+
+	End Object
+
+	Begin Object Class=K2Node_CallParentFunction Name="K2Node_CallParentFunction_3"
+
+		FunctionReference=(MemberParent=BlueprintGeneratedClass'/Game/Prototype/Game/Gun_Pistol.Gun_Pistol_C',MemberName="End Primary Grab",MemberGuid=5AC0E94444892A8794605F86BF8B03E2)
+
+		NodePosX=320
+
+		NodePosY=1584
+
+		NodeGuid=3F018DC04674D5E4348E1E8885C1D07A
+
+		CustomProperties Pin (PinId=440FA2384D499E510943C5952AED1D52,PinName="execute",PinType.PinCategory="exec",PinType.PinSubCategory="",PinType.PinSubCategoryObject=None,PinType.PinSubCategoryMemberReference=(),PinType.PinValueType=(),PinType.bIsMap=False,PinType.bIsSet=False,PinType.bIsArray=False,PinType.bIsReference=False,PinType.bIsConst=False,PinType.bIsWeakPointer=False,LinkedTo=(K2Node_Event_7 DDE6DFF14B7B9ED8FABA8A8A2B48E090,),PersistentGuid=00000000000000000000000000000000,bHidden=False,bNotConnectable=False,bDefaultValueIsReadOnly=False,bDefaultValueIsIgnored=False,bAdvancedView=False,)
+
+		CustomProperties Pin (PinId=B988ABBD45E77F4BBB6292BC717283DB,PinName="then",Direction="EGPD_Output",PinType.PinCategory="exec",PinType.PinSubCategory="",PinType.PinSubCategoryObject=None,PinType.PinSubCategoryMemberReference=(),PinType.PinValueType=(),PinType.bIsMap=False,PinType.bIsSet=False,PinType.bIsArray=False,PinType.bIsReference=False,PinType.bIsConst=False,PinType.bIsWeakPointer=False,LinkedTo=(K2Node_VariableSet_1553 AE7B514B460DD96C37ECE7A72350ADB1,),PersistentGuid=00000000000000000000000000000000,bHidden=False,bNotConnectable=False,bDefaultValueIsReadOnly=False,bDefaultValueIsIgnored=False,bAdvancedView=False,)
+
+		CustomProperties Pin (PinId=FC7E71C84B85D76ED3975980B0C4B1E7,PinName="self",PinFriendlyName=NSLOCTEXT("K2Node", "Target", "Target"),PinType.PinCategory="object",PinType.PinSubCategory="",PinType.PinSubCategoryObject=BlueprintGeneratedClass'/Game/Prototype/Game/Gun_Base.Gun_Base_C',PinType.PinSubCategoryMemberReference=(),PinType.PinValueType=(),PinType.bIsMap=False,PinType.bIsSet=False,PinType.bIsArray=False,PinType.bIsReference=False,PinType.bIsConst=False,PinType.bIsWeakPointer=False,PersistentGuid=00000000000000000000000000000000,bHidden=True,bNotConnectable=False,bDefaultValueIsReadOnly=False,bDefaultValueIsIgnored=False,bAdvancedView=False,)
+
+	End Object
+
+	Begin Object Class=K2Node_Event Name="K2Node_Event_14"
+
+		EventReference=(MemberParent=BlueprintGeneratedClass'/Game/Prototype/Game/Gun_Base.Gun_Base_C',MemberName="Begin Secondary Grab",MemberGuid=788DAEE54912BA1513C0EA9EBD35484D)
+
+		bOverrideFunction=True
+
+		NodePosX=48
+
+		NodePosY=1296
+
+		NodeGuid=68AB4E004B109F18E9831889357D0DB8
+
+		CustomProperties Pin (PinId=9AB423404CAEE074165D7CABBA120EE1,PinName="OutputDelegate",Direction="EGPD_Output",PinType.PinCategory="delegate",PinType.PinSubCategory="",PinType.PinSubCategoryObject=None,PinType.PinSubCategoryMemberReference=(MemberParent=BlueprintGeneratedClass'/Game/Prototype/Game/Gun_Base.Gun_Base_C',MemberName="Begin Secondary Grab",MemberGuid=788DAEE54912BA1513C0EA9EBD35484D),PinType.PinValueType=(),PinType.bIsMap=False,PinType.bIsSet=False,PinType.bIsArray=False,PinType.bIsReference=False,PinType.bIsConst=False,PinType.bIsWeakPointer=False,PersistentGuid=00000000000000000000000000000000,bHidden=False,bNotConnectable=False,bDefaultValueIsReadOnly=False,bDefaultValueIsIgnored=False,bAdvancedView=False,)
+
+		CustomProperties Pin (PinId=8A7BF15740D6B74BE105CCBDD79423E5,PinName="then",Direction="EGPD_Output",PinType.PinCategory="exec",PinType.PinSubCategory="",PinType.PinSubCategoryObject=None,PinType.PinSubCategoryMemberReference=(),PinType.PinValueType=(),PinType.bIsMap=False,PinType.bIsSet=False,PinType.bIsArray=False,PinType.bIsReference=False,PinType.bIsConst=False,PinType.bIsWeakPointer=False,LinkedTo=(K2Node_CallParentFunction_8 280509624149FB7B8299579E0218E824,),PersistentGuid=00000000000000000000000000000000,bHidden=False,bNotConnectable=False,bDefaultValueIsReadOnly=False,bDefaultValueIsIgnored=False,bAdvancedView=False,)
+
+	End Object
+
+	Begin Object Class=K2Node_Event Name="K2Node_Event_19"
+
+		EventReference=(MemberParent=BlueprintGeneratedClass'/Game/Prototype/Game/Gun_Base.Gun_Base_C',MemberName="End Secondary Grab",MemberGuid=CE30C1A0429D6BC4F021EC894F025A9A)
+
+		bOverrideFunction=True
+
+		NodePosX=48
+
+		NodePosY=1504
+
+		NodeGuid=CF2791714730BBCE92043AA025D0C200
+
+		CustomProperties Pin (PinId=04FA7A954BA262A18863708BE6933EF8,PinName="OutputDelegate",Direction="EGPD_Output",PinType.PinCategory="delegate",PinType.PinSubCategory="",PinType.PinSubCategoryObject=None,PinType.PinSubCategoryMemberReference=(MemberParent=BlueprintGeneratedClass'/Game/Prototype/Game/Gun_Base.Gun_Base_C',MemberName="End Secondary Grab",MemberGuid=CE30C1A0429D6BC4F021EC894F025A9A),PinType.PinValueType=(),PinType.bIsMap=False,PinType.bIsSet=False,PinType.bIsArray=False,PinType.bIsReference=False,PinType.bIsConst=False,PinType.bIsWeakPointer=False,PersistentGuid=00000000000000000000000000000000,bHidden=False,bNotConnectable=False,bDefaultValueIsReadOnly=False,bDefaultValueIsIgnored=False,bAdvancedView=False,)
+
+		CustomProperties Pin (PinId=87AB296D4EE26D0186F710BB3FA4C7B7,PinName="then",Direction="EGPD_Output",PinType.PinCategory="exec",PinType.PinSubCategory="",PinType.PinSubCategoryObject=None,PinType.PinSubCategoryMemberReference=(),PinType.PinValueType=(),PinType.bIsMap=False,PinType.bIsSet=False,PinType.bIsArray=False,PinType.bIsReference=False,PinType.bIsConst=False,PinType.bIsWeakPointer=False,LinkedTo=(K2Node_CallParentFunction_9 C2AB49794E3B99FF2DDC1983D4C3D443,),PersistentGuid=00000000000000000000000000000000,bHidden=False,bNotConnectable=False,bDefaultValueIsReadOnly=False,bDefaultValueIsIgnored=False,bAdvancedView=False,)
+
+	End Object
+
+	Begin Object Class=K2Node_CallParentFunction Name="K2Node_CallParentFunction_8"
+
+		FunctionReference=(MemberParent=BlueprintGeneratedClass'/Game/Prototype/Game/Gun_Pistol.Gun_Pistol_C',MemberName="Begin Secondary Grab",MemberGuid=788DAEE54912BA1513C0EA9EBD35484D)
+
+		NodePosX=320
+
+		NodePosY=1296
+
+		NodeGuid=5E8F1A584BE808DB38A4338935BEB20E
+
+		CustomProperties Pin (PinId=280509624149FB7B8299579E0218E824,PinName="execute",PinType.PinCategory="exec",PinType.PinSubCategory="",PinType.PinSubCategoryObject=None,PinType.PinSubCategoryMemberReference=(),PinType.PinValueType=(),PinType.bIsMap=False,PinType.bIsSet=False,PinType.bIsArray=False,PinType.bIsReference=False,PinType.bIsConst=False,PinType.bIsWeakPointer=False,LinkedTo=(K2Node_Event_14 8A7BF15740D6B74BE105CCBDD79423E5,),PersistentGuid=00000000000000000000000000000000,bHidden=False,bNotConnectable=False,bDefaultValueIsReadOnly=False,bDefaultValueIsIgnored=False,bAdvancedView=False,)
+
+		CustomProperties Pin (PinId=84FF3C864D5859B0905DAD8DC8595A9E,PinName="then",Direction="EGPD_Output",PinType.PinCategory="exec",PinType.PinSubCategory="",PinType.PinSubCategoryObject=None,PinType.PinSubCategoryMemberReference=(),PinType.PinValueType=(),PinType.bIsMap=False,PinType.bIsSet=False,PinType.bIsArray=False,PinType.bIsReference=False,PinType.bIsConst=False,PinType.bIsWeakPointer=False,LinkedTo=(K2Node_VariableSet_0 AE7B514B460DD96C37ECE7A72350ADB1,),PersistentGuid=00000000000000000000000000000000,bHidden=False,bNotConnectable=False,bDefaultValueIsReadOnly=False,bDefaultValueIsIgnored=False,bAdvancedView=False,)
+
+		CustomProperties Pin (PinId=BE812250483E2AEE00E09B9EBE6F7E60,PinName="self",PinFriendlyName=NSLOCTEXT("K2Node", "Target", "Target"),PinType.PinCategory="object",PinType.PinSubCategory="",PinType.PinSubCategoryObject=BlueprintGeneratedClass'/Game/Prototype/Game/Gun_Base.Gun_Base_C',PinType.PinSubCategoryMemberReference=(),PinType.PinValueType=(),PinType.bIsMap=False,PinType.bIsSet=False,PinType.bIsArray=False,PinType.bIsReference=False,PinType.bIsConst=False,PinType.bIsWeakPointer=False,PersistentGuid=00000000000000000000000000000000,bHidden=True,bNotConnectable=False,bDefaultValueIsReadOnly=False,bDefaultValueIsIgnored=False,bAdvancedView=False,)
+
+	End Object
+
+	Begin Object Class=K2Node_CallParentFunction Name="K2Node_CallParentFunction_9"
+
+		FunctionReference=(MemberParent=BlueprintGeneratedClass'/Game/Prototype/Game/Gun_Pistol.Gun_Pistol_C',MemberName="End Secondary Grab",MemberGuid=CE30C1A0429D6BC4F021EC894F025A9A)
+
+		NodePosX=320
+
+		NodePosY=1504
+
+		NodeGuid=0E48003242FC85D2D68655B218A52797
+
+		CustomProperties Pin (PinId=C2AB49794E3B99FF2DDC1983D4C3D443,PinName="execute",PinType.PinCategory="exec",PinType.PinSubCategory="",PinType.PinSubCategoryObject=None,PinType.PinSubCategoryMemberReference=(),PinType.PinValueType=(),PinType.bIsMap=False,PinType.bIsSet=False,PinType.bIsArray=False,PinType.bIsReference=False,PinType.bIsConst=False,PinType.bIsWeakPointer=False,LinkedTo=(K2Node_Event_19 87AB296D4EE26D0186F710BB3FA4C7B7,),PersistentGuid=00000000000000000000000000000000,bHidden=False,bNotConnectable=False,bDefaultValueIsReadOnly=False,bDefaultValueIsIgnored=False,bAdvancedView=False,)
+
+		CustomProperties Pin (PinId=04A663DF4151E8B743CBE3A7C7BCF848,PinName="then",Direction="EGPD_Output",PinType.PinCategory="exec",PinType.PinSubCategory="",PinType.PinSubCategoryObject=None,PinType.PinSubCategoryMemberReference=(),PinType.PinValueType=(),PinType.bIsMap=False,PinType.bIsSet=False,PinType.bIsArray=False,PinType.bIsReference=False,PinType.bIsConst=False,PinType.bIsWeakPointer=False,LinkedTo=(K2Node_VariableSet_1553 AE7B514B460DD96C37ECE7A72350ADB1,),PersistentGuid=00000000000000000000000000000000,bHidden=False,bNotConnectable=False,bDefaultValueIsReadOnly=False,bDefaultValueIsIgnored=False,bAdvancedView=False,)
+
+		CustomProperties Pin (PinId=1E9924B4423BFF3C997BB4974473931D,PinName="self",PinFriendlyName=NSLOCTEXT("K2Node", "Target", "Target"),PinType.PinCategory="object",PinType.PinSubCategory="",PinType.PinSubCategoryObject=BlueprintGeneratedClass'/Game/Prototype/Game/Gun_Base.Gun_Base_C',PinType.PinSubCategoryMemberReference=(),PinType.PinValueType=(),PinType.bIsMap=False,PinType.bIsSet=False,PinType.bIsArray=False,PinType.bIsReference=False,PinType.bIsConst=False,PinType.bIsWeakPointer=False,PersistentGuid=00000000000000000000000000000000,bHidden=True,bNotConnectable=False,bDefaultValueIsReadOnly=False,bDefaultValueIsIgnored=False,bAdvancedView=False,)
+
+	End Object
+
+	Begin Object Class=EdGraphNode_Comment Name="EdGraphNode_Comment_276"
+
+		NodePosY=1152
+
+		NodeWidth=880
+
+		NodeHeight=256
+
+		NodeComment="Begin Grab Events"
+
+		NodeGuid=73C6DFFF47B2E78D7A929382C5BAA6C2
+
+	End Object
+
+	Begin Object Class=EdGraphNode_Comment Name="EdGraphNode_Comment_277"
+
+		NodePosY=1456
+
+		NodeWidth=880
+
+		NodeHeight=240
+
+		NodeComment="End Grab Events"
+
+		NodeGuid=1E4AC8CB452379DDC9B4F6964B10C06B
+
+	End Object
+
+	Begin Object Class=K2Node_VariableSet Name="K2Node_VariableSet_0"
+
+		VariableReference=(MemberName="bCanExplode",MemberGuid=9BCA29C8404B219E30C266861239680A,bSelfContext=True)
+
+		NodePosX=624
+
+		NodePosY=1264
+
+		NodeGuid=8178103244831FA0B54DDAA8525A4191
+
+		CustomProperties Pin (PinId=AE7B514B460DD96C37ECE7A72350ADB1,PinName="execute",PinType.PinCategory="exec",PinType.PinSubCategory="",PinType.PinSubCategoryObject=None,PinType.PinSubCategoryMemberReference=(),PinType.PinValueType=(),PinType.bIsMap=False,PinType.bIsSet=False,PinType.bIsArray=False,PinType.bIsReference=False,PinType.bIsConst=False,PinType.bIsWeakPointer=False,LinkedTo=(K2Node_CallParentFunction_0 1B74752B4533A65EB598E787671714A5,K2Node_CallParentFunction_8 84FF3C864D5859B0905DAD8DC8595A9E,),PersistentGuid=00000000000000000000000000000000,bHidden=False,bNotConnectable=False,bDefaultValueIsReadOnly=False,bDefaultValueIsIgnored=False,bAdvancedView=False,)
+
+		CustomProperties Pin (PinId=E0F4A52E424426C37B5EF98C8FFAF121,PinName="then",Direction="EGPD_Output",PinType.PinCategory="exec",PinType.PinSubCategory="",PinType.PinSubCategoryObject=None,PinType.PinSubCategoryMemberReference=(),PinType.PinValueType=(),PinType.bIsMap=False,PinType.bIsSet=False,PinType.bIsArray=False,PinType.bIsReference=False,PinType.bIsConst=False,PinType.bIsWeakPointer=False,PersistentGuid=00000000000000000000000000000000,bHidden=False,bNotConnectable=False,bDefaultValueIsReadOnly=False,bDefaultValueIsIgnored=False,bAdvancedView=False,)
+
+		CustomProperties Pin (PinId=DAB7D07A426D67E49C8D63A7911BFC1B,PinName="bCanExplode",PinType.PinCategory="bool",PinType.PinSubCategory="",PinType.PinSubCategoryObject=None,PinType.PinSubCategoryMemberReference=(),PinType.PinValueType=(),PinType.bIsMap=False,PinType.bIsSet=False,PinType.bIsArray=False,PinType.bIsReference=False,PinType.bIsConst=False,PinType.bIsWeakPointer=False,DefaultValue="false",AutogeneratedDefaultValue="false",PersistentGuid=00000000000000000000000000000000,bHidden=False,bNotConnectable=False,bDefaultValueIsReadOnly=False,bDefaultValueIsIgnored=False,bAdvancedView=False,)
+
+		CustomProperties Pin (PinId=26DA759E42C15E427748F283E9F317CC,PinName="Output_Get",PinToolTip="Retrieves the value of the variable, can use instead of a separate Get node",Direction="EGPD_Output",PinType.PinCategory="bool",PinType.PinSubCategory="",PinType.PinSubCategoryObject=None,PinType.PinSubCategoryMemberReference=(),PinType.PinValueType=(),PinType.bIsMap=False,PinType.bIsSet=False,PinType.bIsArray=False,PinType.bIsReference=False,PinType.bIsConst=False,PinType.bIsWeakPointer=False,DefaultValue="false",AutogeneratedDefaultValue="false",PersistentGuid=00000000000000000000000000000000,bHidden=False,bNotConnectable=False,bDefaultValueIsReadOnly=False,bDefaultValueIsIgnored=False,bAdvancedView=False,)
+
+		CustomProperties Pin (PinId=93A90B8F461C1A8D79FF37B7F054C5F3,PinName="self",PinFriendlyName=NSLOCTEXT("K2Node", "Target", "Target"),PinType.PinCategory="object",PinType.PinSubCategory="",PinType.PinSubCategoryObject=BlueprintGeneratedClass'/GrenadeGun/Blueprints/Weapons/GrenadeGunPistol.GrenadeGunPistol_C',PinType.PinSubCategoryMemberReference=(),PinType.PinValueType=(),PinType.bIsMap=False,PinType.bIsSet=False,PinType.bIsArray=False,PinType.bIsReference=False,PinType.bIsConst=False,PinType.bIsWeakPointer=False,PersistentGuid=00000000000000000000000000000000,bHidden=True,bNotConnectable=False,bDefaultValueIsReadOnly=False,bDefaultValueIsIgnored=False,bAdvancedView=False,)
+
+	End Object
+
+	Begin Object Class=K2Node_VariableSet Name="K2Node_VariableSet_1553"
+
+		VariableReference=(MemberName="bCanExplode",MemberGuid=9BCA29C8404B219E30C266861239680A,bSelfContext=True)
+
+		NodePosX=624
+
+		NodePosY=1552
+
+		NodeGuid=CC14C5034D568D878B12CE8B90B0E8D6
+
+		CustomProperties Pin (PinId=AE7B514B460DD96C37ECE7A72350ADB1,PinName="execute",PinType.PinCategory="exec",PinType.PinSubCategory="",PinType.PinSubCategoryObject=None,PinType.PinSubCategoryMemberReference=(),PinType.PinValueType=(),PinType.bIsMap=False,PinType.bIsSet=False,PinType.bIsArray=False,PinType.bIsReference=False,PinType.bIsConst=False,PinType.bIsWeakPointer=False,LinkedTo=(K2Node_CallParentFunction_9 04A663DF4151E8B743CBE3A7C7BCF848,K2Node_CallParentFunction_3 B988ABBD45E77F4BBB6292BC717283DB,),PersistentGuid=00000000000000000000000000000000,bHidden=False,bNotConnectable=False,bDefaultValueIsReadOnly=False,bDefaultValueIsIgnored=False,bAdvancedView=False,)
+
+		CustomProperties Pin (PinId=E0F4A52E424426C37B5EF98C8FFAF121,PinName="then",Direction="EGPD_Output",PinType.PinCategory="exec",PinType.PinSubCategory="",PinType.PinSubCategoryObject=None,PinType.PinSubCategoryMemberReference=(),PinType.PinValueType=(),PinType.bIsMap=False,PinType.bIsSet=False,PinType.bIsArray=False,PinType.bIsReference=False,PinType.bIsConst=False,PinType.bIsWeakPointer=False,PersistentGuid=00000000000000000000000000000000,bHidden=False,bNotConnectable=False,bDefaultValueIsReadOnly=False,bDefaultValueIsIgnored=False,bAdvancedView=False,)
+
+		CustomProperties Pin (PinId=DAB7D07A426D67E49C8D63A7911BFC1B,PinName="bCanExplode",PinType.PinCategory="bool",PinType.PinSubCategory="",PinType.PinSubCategoryObject=None,PinType.PinSubCategoryMemberReference=(),PinType.PinValueType=(),PinType.bIsMap=False,PinType.bIsSet=False,PinType.bIsArray=False,PinType.bIsReference=False,PinType.bIsConst=False,PinType.bIsWeakPointer=False,DefaultValue="true",AutogeneratedDefaultValue="false",PersistentGuid=00000000000000000000000000000000,bHidden=False,bNotConnectable=False,bDefaultValueIsReadOnly=False,bDefaultValueIsIgnored=False,bAdvancedView=False,)
+
+		CustomProperties Pin (PinId=26DA759E42C15E427748F283E9F317CC,PinName="Output_Get",PinToolTip="Retrieves the value of the variable, can use instead of a separate Get node",Direction="EGPD_Output",PinType.PinCategory="bool",PinType.PinSubCategory="",PinType.PinSubCategoryObject=None,PinType.PinSubCategoryMemberReference=(),PinType.PinValueType=(),PinType.bIsMap=False,PinType.bIsSet=False,PinType.bIsArray=False,PinType.bIsReference=False,PinType.bIsConst=False,PinType.bIsWeakPointer=False,DefaultValue="false",AutogeneratedDefaultValue="false",PersistentGuid=00000000000000000000000000000000,bHidden=False,bNotConnectable=False,bDefaultValueIsReadOnly=False,bDefaultValueIsIgnored=False,bAdvancedView=False,)
+
+		CustomProperties Pin (PinId=93A90B8F461C1A8D79FF37B7F054C5F3,PinName="self",PinFriendlyName=NSLOCTEXT("K2Node", "Target", "Target"),PinType.PinCategory="object",PinType.PinSubCategory="",PinType.PinSubCategoryObject=BlueprintGeneratedClass'/GrenadeGun/Blueprints/Weapons/GrenadeGunPistol.GrenadeGunPistol_C',PinType.PinSubCategoryMemberReference=(),PinType.PinValueType=(),PinType.bIsMap=False,PinType.bIsSet=False,PinType.bIsArray=False,PinType.bIsReference=False,PinType.bIsConst=False,PinType.bIsWeakPointer=False,PersistentGuid=00000000000000000000000000000000,bHidden=True,bNotConnectable=False,bDefaultValueIsReadOnly=False,bDefaultValueIsIgnored=False,bAdvancedView=False,)
+
+	End Object
+[/EXCERPT:snippet]
+-->
+
diff --git a/Engine/Documentation/Source/Modding/RoboRecall/HowTo/ExtendFunctionality/CreateMaterialInstance/CreateMaterialInstance.INT.udn b/Engine/Documentation/Source/Modding/RoboRecall/HowTo/ExtendFunctionality/CreateMaterialInstance/CreateMaterialInstance.INT.udn
new file mode 100644
--- /dev/null
+++ b/Engine/Documentation/Source/Modding/RoboRecall/HowTo/ExtendFunctionality/CreateMaterialInstance/CreateMaterialInstance.INT.udn
@@ -0,0 +1,73 @@
+Availability: Public
+Title:2. Creating a New Material Instance
+Crumbs:%ROOT%
+Description:Robo Recall Modding Documentation
+Parent:Modding/RoboRecall/HowTo/ExtendFunctionality
+type:how-to
+order:2
+tags:Modding
+tags:Robo Recall
+topic-image:Modding/RoboRecall/RoboRecall_Topic.png
+social-image:Modding/RoboRecall/RoboRecall_Social.png
+checkpoint:mod_extend
+
+<!-- Location: Modding/RoboRecall/Quickstarts/WeaponMod/2 -->
+
+[NAV]
+
+<!-- Title Image if Needed Here -->
+
+You have a custom pistol, but it still behaves like the default pistol. When thrown, it has a "dissolve" effect that is blue. Since our pistol is going to explode, a red effect would make more sense. To set that up, you need a new [Material Instance](Engine/Rendering/Materials/MaterialInstances) of the "dissolving" material used by the **Pistol**.
+
+<!-- probably should explain what a material instance is for newcomers -->
+
+## Steps
+
+1. Search for **MI_Pistol_Dissolve** in the **Content Browser** using the search box.
+
+	![image alt text](image_0.png)(w:805 h:352 convert:false)
+
+1. **Right-Click** the asset and choose **Create Material Instance** to create a new Material Instance based off _MI_Pistol_Dissolve_.
+
+	![image alt text](image_1.png)(w:273 h:604 convert:false)
+
+1. Drag the new Material Instance to the **GrenadeGun Content** folder. Choose **Move Here** when you release the mouse button to move it to that location.
+
+	![image alt text](image_2.png)(w:810 h:385 convert:false)
+
+1. Press **F2** to rename the Material Instance and enter **GrenadeGun_Dissolve_MI** as the new name.
+
+	![image alt text](image_3.png)(w:805 h:352 convert:false)
+
+1. **Double-Click** the _GrenadeGun_Dissolve_MI_ to edit it in the [Material Instance Editor](Engine/Rendering/Materials/MaterialInstances/Editor).
+
+	![image alt text](image_4.png)(w:934 h:600 convert:false)
+
+1. Click the checkbox next to the **Dissolve** parameter to enable overriding it. Set the value of the **Dissolve** parameter to a value in the range of `-0.2` to `0.33` - the value you use will depend on the preview shape you are currently using in the Preview viewport - to cause the dissolve effect to be visible.
+
+1. Click the checkbox next to the **DissolveBurnColor** parameter to enable overriding it. This parameter controls the color used for the dissolve effect.
+
+	![image alt text](image_5.png)(w:328 h:370 convert:false)
+
+1. Click the colored bar to summon the Color Picker and set it to a red-ish color. Here we've used the values of: (2.0, 0.23705, 0.0032819, 15000.0)
+
+	![image alt text](image_6.png)(w:440 h:536 convert:false)
+
+	[REGION:tip]
+	You can also enter the R, G, and B values directly in the Details panel by expanding the **DissolveBurnColor** parameter to expose the individual values.
+
+	![image alt text](image_7.png)(w:328 h:370 convert:false)
+	[/REGION]
+
+1. In the Material Instance Editor toolbar, Click the checkbox next to the **Dissolve** parameter again to disable it since we were just using that for preview purposes.
+
+1. Click the **Save** button to preserve your changes
+
+## Results
+
+You now have a version of the dissolve Material with a red effect instead of the default blue effect.
+
+![image alt text](image_9.png)(w:934 h:600 convert:false)
+
+[NAV]
+
diff --git a/Engine/Documentation/Source/Modding/RoboRecall/HowTo/ExtendFunctionality/CreatePistolMod/CreatePistolMod.INT.udn b/Engine/Documentation/Source/Modding/RoboRecall/HowTo/ExtendFunctionality/CreatePistolMod/CreatePistolMod.INT.udn
new file mode 100644
--- /dev/null
+++ b/Engine/Documentation/Source/Modding/RoboRecall/HowTo/ExtendFunctionality/CreatePistolMod/CreatePistolMod.INT.udn
@@ -0,0 +1,61 @@
+Availability: Public
+Title:1. Creating a New Pistol Mod
+Crumbs:%ROOT%
+Description:Robo Recall Modding Documentation
+Parent:Modding/RoboRecall/HowTo/ExtendFunctionality
+type:how-to
+order:1
+tags:Modding
+tags:Robo Recall
+topic-image:Modding/RoboRecall/RoboRecall_Topic.png
+social-image:Modding/RoboRecall/RoboRecall_Social.png
+checkpoint:mod_extend
+
+<!-- Location: Modding/RoboRecall/Quickstarts/WeaponMod/1 -->
+
+[NAV]
+
+<!-- Title Image if Needed Here -->
+
+Overview Text
+
+## Steps
+
+1. Click the **Create Mod** button on the Toolbar.
+
+	![](mod_create.png)
+
+1. Select the **Custom Pistol** from hte available mod types.
+	
+	![](mod_select.png)
+
+	[REGION:tip]
+	You can multiselect in this dialog to create multiple assets at the same time. Feel free to add any other weapons, bots, etc. by **Ctrl + Clicking** the assets you want to create in your mod. 
+	[/REGION]
+
+1. Enter a name for the mod in the **Name** field. In this case, we've chosen **GrenadeGun**.
+
+	![](mod_name.png)
+
+1. You can enter an **Author** and **Description**, and this data will show in the Mod menu in the game. If you'd like to change this later, you can do so by accessing your mod through the Editor's **Edit -> Plugins** menu.
+
+	![](mod_info.png)
+
+1. Click the **Create Mod** button.
+
+1. Your new mod package will be automatically focused in the **Content Browser**, but it will be at the root level of the mod.
+
+	![image alt text](image_3.png)(w:934 h:288 convert:false)
+
+1. Double click the **Blueprints** folder, then the **Weapons** folder to find the Blueprint asset for your new custom pistol.
+
+	![image alt text](image_4.png)(w:934 h:288 convert:false)
+
+## Results
+
+You now have a mod with a custom pistol that will replace the default pistol in the game. Now it needs the functionality to cause it to explode when thrown.
+
+![image alt text](image_0.png)(w:805 h:352 convert:false)
+
+[NAV]
+
diff --git a/Engine/Documentation/Source/Modding/RoboRecall/HowTo/ExtendFunctionality/TestYourMod/TestYourMod.INT.udn b/Engine/Documentation/Source/Modding/RoboRecall/HowTo/ExtendFunctionality/TestYourMod/TestYourMod.INT.udn
new file mode 100644
--- /dev/null
+++ b/Engine/Documentation/Source/Modding/RoboRecall/HowTo/ExtendFunctionality/TestYourMod/TestYourMod.INT.udn
@@ -0,0 +1,61 @@
+Availability: Public
+Title:7. Testing and Next Steps
+Crumbs:%ROOT%
+Description:Robo Recall Modding Documentation
+Parent:Modding/RoboRecall/HowTo/ExtendFunctionality
+type:how-to
+order:7
+tags:Modding
+tags:Robo Recall
+topic-image:Modding/RoboRecall/RoboRecall_Topic.png
+social-image:Modding/RoboRecall/RoboRecall_Social.png
+checkpoint:mod_extend
+
+<!-- Location: Modding/RoboRecall/Quickstarts/WeaponMod/7 -->
+
+[NAV]
+
+<!-- Title Image if Needed Here -->
+
+With your custom pistol completed, you can click **Play** while the the hub level is loaded, and enable your GrenadeGunPistol through the mod menu. You can now use the Gun Range to toss the gun at the targets and watch them detonate!
+
+![image alt text](image_1.png)(w:934 h:604 convert:false)
+
+## Next Steps
+
+Now that you successfully created a new pistol, it’s time to see if you can take it further! If you’re having a creative block, here are some ideas to get the creativity flowing while also getting you to dive further into the gun Blueprints.
+
+* Balance the damage of the explosion
+
+* Make the explosion radius dependent on how many bullets are left in the gun
+
+* Have the gun bounce once or twice before exploding
+
+* By disconnecting the **Instigated by Controller** pin in the GrenadeGunPistol Blueprint, you can have the explosion damage the Player
+
+* Why not more explosions? Make it a MIRV Grenade!
+
+    * **Bonus:** Going with a more literal definition of MIRV? Have it call down a ballistic missile strike from space!
+
+* Make the pistol a Machine Pistol with extended magazines and fully automatic firing!
+
+    * **Bonus:** Adjust the damage that the Machine Pistol does to be more balanced for its new fire rate
+
+    * **Mega Bonus:** Make the bullets it fires homing rounds!
+
+* If you plan to continue using the base mesh for the pistol, figure out how to have the Blueprint add Attachments (laser, compensator, sights, etc…)
+
+* Give the gun bonus damage on its next shot based on how many times you juggle it
+
+[REGION:tip]
+
+In the top right of the GrenadeGunPistol Blueprint you will see a little magnifying glass:
+
+![image alt text](image_2.png)(w:486 h:78 convert:false)
+
+This button will take highlight the current Blueprint’s Parent in the Content Browser. By using this button you can navigate down to Gun_Pistol and Gun_Base, which contain a mountain of functionality for the Pistol. Browse through them to see what variables and functions you have available to you.
+
+[/REGION]
+
+[NAV]
+
diff --git a/Engine/Documentation/Source/Modding/RoboRecall/HowTo/ExtendFunctionality/WeaponMod.INT.udn b/Engine/Documentation/Source/Modding/RoboRecall/HowTo/ExtendFunctionality/WeaponMod.INT.udn
new file mode 100644
--- /dev/null
+++ b/Engine/Documentation/Source/Modding/RoboRecall/HowTo/ExtendFunctionality/WeaponMod.INT.udn
@@ -0,0 +1,28 @@
+Availability: Public
+Title:Overriding and Extending Functionality
+Crumbs:%ROOT%
+Description:Guide that walks you through overriding the existing functionality of an item from the game
+Parent:Modding/RoboRecall
+type:how-to
+order:10
+tags:Modding
+tags:Robo Recall
+topic-image:RoboRecallModding_Guns_Topic.png
+social-image:RoboRecallModding_Guns_Social.png
+skilllevel:Intermediate
+checkpoint:mod_extend
+
+<!-- Location: Modding/RoboRecall/Quickstarts/WeaponMod-->
+
+![](RoboRecallModding_Guns.png)
+
+In this guide, you will make a new weapon mod using the Pistol as a base which you will modify to explode like a grenade when thrown into the world by overriding some of the existing functionality.
+
+## Steps
+
+[DIR(end:"1" output:"checkpoint" org:"hierarchy" parent:"Modding/RoboRecall/HowTo/ExtendFunctionality")]
+
+[REGION:call_to_action]
+[Click to Start](Modding/RoboRecall/HowTo/ExtendFunctionality/CreatePistolMod "%Modding/RoboRecall/HowTo/ExtendFunctionality/CreatePistolMod:title%")
+[/REGION]
+
diff --git a/Engine/Documentation/Source/Modding/RoboRecall/HowTo/InstallAMod/InstallAMod.INT.udn b/Engine/Documentation/Source/Modding/RoboRecall/HowTo/InstallAMod/InstallAMod.INT.udn
new file mode 100644
--- /dev/null
+++ b/Engine/Documentation/Source/Modding/RoboRecall/HowTo/InstallAMod/InstallAMod.INT.udn
@@ -0,0 +1,36 @@
+Availability: Public
+Title:Installing a Robo Recall Mod
+Crumbs:%ROOT%
+Description:Instructions for installing Robo Recall mods for use in the game
+Parent:Modding/RoboRecall
+type:how-to
+order:2
+tags:Modding
+tags:Robo Recall
+topic-image:install_topic.png
+
+Mods you create yourself, or those you download from others, must be installed in order to be available in the game.
+
+[REGION:note]
+The installation process works by associating the mod filetype with the game. You must play the game at least once for the association to be set up.
+[/REGION]
+
+1. Locate the `.robo` file fopr the mod you want to install.
+
+	![](robo_file.png)
+
+1. Double-click the `.robo` file and click **Yes** to begin the installation.
+
+	![](robo_install.png)
+
+1. When the installtion is complete, you have the option to jump right into the game, view the install folder, or exit.
+
+	![](robo_complete.png)
+
+	The install folder contains all of the files for the mod. You should see a `.uplugin` file along with a Content folder in the root directory.
+
+	![](robo_folder.png)
+
+1. The mod you installed is now available in the **MODS** menu in the game.
+
+	![](Modding/RoboRecall/HowTo/EnableMods/menu_mod_enabled.png)
\ No newline at end of file
diff --git a/Engine/Documentation/Source/Modding/RoboRecall/HowTo/InstallModKit/InstallModKit.INT.udn b/Engine/Documentation/Source/Modding/RoboRecall/HowTo/InstallModKit/InstallModKit.INT.udn
new file mode 100644
--- /dev/null
+++ b/Engine/Documentation/Source/Modding/RoboRecall/HowTo/InstallModKit/InstallModKit.INT.udn
@@ -0,0 +1,37 @@
+Availability: Public
+Title:Installing the Robo Recall Mod Kit
+Crumbs:%ROOT%
+Description:Guide to downloading and installing the Robo Recall mod Kit used to make mods for the game
+Parent:Modding/RoboRecall
+type:how-to
+order:1
+tags:Modding
+tags:Robo Recall
+topic-image:install_topic.png
+social-image:install.png
+
+1. [Sign up for an Epic ID](GettingStarted/Installation#bookmark1)
+
+1. [Download and install the Epic Games Launcher](GettingStarted/Installation#bookmark2)
+
+1. Click **Install** to download the Robo Recall Mod Kit from the Epic Games Launcher.
+
+	![](install.png)(w:934 h:433 convert:false)
+
+	[REGION:tip]
+	This can take some time as the Robo Recall Editor is quite large (20+ GB). Take this time to review some of the other Robo Recall Documentation or perhaps some of the Unreal Engine 4 getting started guides:
+
+	[REGION:buttonlist]
+	* [](Engine/QuickStart)
+	* [](Engine/Blueprints)
+	* [](Programming/QuickStart)
+	* [](Engine/Content)
+	[/REGION]
+
+	Once the download is complete, head to the next step.
+	[/REGION]
+
+1. Launch the Robo Recall Editor.
+
+	![](launch.png)
+
diff --git a/Engine/Documentation/Source/Modding/RoboRecall/HowTo/SetupGameplay/BeginLevelScript/BeginLevelScript.INT.udn b/Engine/Documentation/Source/Modding/RoboRecall/HowTo/SetupGameplay/BeginLevelScript/BeginLevelScript.INT.udn
new file mode 100644
--- /dev/null
+++ b/Engine/Documentation/Source/Modding/RoboRecall/HowTo/SetupGameplay/BeginLevelScript/BeginLevelScript.INT.udn
@@ -0,0 +1,171 @@
+Availability:Public
+Title:4. Level Startup Scripting
+Crumbs:
+Description:
+Parent:Modding/RoboRecall/HowTo/SetupGameplay
+Order:4
+tags:Modding
+tags:Robo Recall
+type:how-to
+checkpoint:mod_gameplay
+
+[NAV]
+
+When the level begins, you want the player to start at the location of the PlayerStart Actor you placed previously. In order to make this happen, you need to move the player to that location and orient them correctly.
+
+1. In the Level Editor toolbar click the **Blueprints** button and then choose **Open Level Blueprint** to open the script for the current level in the Blueprint Editor.
+
+	![](level_blueprint.png)
+
+1. Right-click in the Event Graph and choose **Event BeginPlay** under **Add Event**.
+
+	![](begin_play.png)
+
+1. Right-click in the Event Graph and choose **Get Game Mode** to get a reference to the current [Game Mode](Gameplay/Framework/GameMode) instance. Drag off the **Return Value** pin and choose **Cast to BP_OdinGameMode**. Connect the exec output of the **BeginPlay** event to the exec input of the **Cast to BP_OdinGameMode** node.
+
+	![](game_mode.png)
+
+1. Drag off the **As BP Odin Game Mode** pin and choose **End Loading Sequence**. This performs some necessary cleanup after loading the level.
+
+	![](end_loading.png)
+
+1. Drag off the **As BP Odin Game Mode** output and choose **Get Main Pawn** (under the **game** category). This gives you a reference to the player which you can use to set their location.
+
+	![](player.png)
+
+1. Drag off the **Main Pawn** pin and choose **Teleport** (under the **Utilities > Transformation** category). Connect the exec output of the **End Loading Sequence** node to the exec input of the **Teleport** node.
+
+	![](teleport.png)
+
+1. In the Level Editor Viewport, select the **Player Start** Actor you place previously. Back in the Blueprint Editor, right-click in the Event Graph and choose **Create a Reference to [ActorName]**.
+
+	![](player_start.png)
+
+1. Drag off the output pin on the Player Start reference and choose **Get ActorLocation**. Connect the **Return Value** output to the **Dest Location** input of the **Teleport** node.
+
+	![](location.png)
+
+1. Drag off the output pin on the Player Start reference and choose **Get ActorRotation**. Connect the **Return Value** output to the **Dest Rotation** input of the **Teleport** node.
+
+	[OBJECT:Snippet]
+		[PARAMLITERAL:id]
+		result
+		[/PARAMLITERAL]
+		[PARAMLITERAL:type]
+		blueprint
+		[/PARAMLITERAL]
+		[PARAM:image]
+		![](rotation.png)(convert:false)
+		[/PARAM]
+		[PARAM:snippet]
+			[INCLUDE:#snippet]
+		[/PARAM]
+	[/OBJECT]
+
+	[REGION:caption]
+	Click the icon in the upper left corner of this image to copy the Blueprint Graph and paste it into your project.
+	[/REGION]
+
+## Result
+
+When you click **Play** in the Level Editor toolbar , you will now start off in the location of the Player Start Actor, facing the direction of the oncoming robots.
+
+![](test.png)
+
+[NAV]
+
+
+<!--
+[EXCERPT:snippet]
+	Begin Object Class=K2Node_Event Name="K2Node_Event_52"
+	   EventReference=(MemberParent=Class'/Script/Engine.Actor',MemberName="ReceiveBeginPlay")
+	   bOverrideFunction=True
+	   NodePosX=-928
+	   NodePosY=-4480
+	   NodeGuid=6C849E7241A6108A36D5C1BEAAEAE94F
+	   CustomProperties Pin (PinId=B35535C443B530261C41B8A84FF2CC42,PinName="OutputDelegate",Direction="EGPD_Output",PinType.PinCategory="delegate",PinType.PinSubCategory="",PinType.PinSubCategoryObject=None,PinType.PinSubCategoryMemberReference=(MemberParent=Class'/Script/Engine.Actor',MemberName="ReceiveBeginPlay"),PinType.PinValueType=(),PinType.bIsMap=False,PinType.bIsSet=False,PinType.bIsArray=False,PinType.bIsReference=False,PinType.bIsConst=False,PinType.bIsWeakPointer=False,PersistentGuid=00000000000000000000000000000000,bHidden=False,bNotConnectable=False,bDefaultValueIsReadOnly=False,bDefaultValueIsIgnored=False,bAdvancedView=False,)
+	   CustomProperties Pin (PinId=5A6DBA7144DD625770231A9C725C0B49,PinName="then",Direction="EGPD_Output",PinType.PinCategory="exec",PinType.PinSubCategory="",PinType.PinSubCategoryObject=None,PinType.PinSubCategoryMemberReference=(),PinType.PinValueType=(),PinType.bIsMap=False,PinType.bIsSet=False,PinType.bIsArray=False,PinType.bIsReference=False,PinType.bIsConst=False,PinType.bIsWeakPointer=False,LinkedTo=(K2Node_DynamicCast_5 3228E8FA4B588A45DCEA3A839D739E9A,),PersistentGuid=00000000000000000000000000000000,bHidden=False,bNotConnectable=False,bDefaultValueIsReadOnly=False,bDefaultValueIsIgnored=False,bAdvancedView=False,)
+	End Object
+	Begin Object Class=K2Node_CallFunction Name="K2Node_CallFunction_280"
+	   bIsPureFunc=True
+	   FunctionReference=(MemberParent=Class'/Script/Engine.GameplayStatics',MemberName="GetGameMode")
+	   NodePosX=-720
+	   NodePosY=-4416
+	   NodeGuid=2CB9ABB94CF9F2C3BAEBB0B38ED0F003
+	   CustomProperties Pin (PinId=1AAC46094DE5E3D0CB1F8488BEDBFE3B,PinName="self",PinFriendlyName=NSLOCTEXT("K2Node", "Target", "Target"),PinToolTip="Target\nGameplay Statics Reference",PinType.PinCategory="object",PinType.PinSubCategory="",PinType.PinSubCategoryObject=Class'/Script/Engine.GameplayStatics',PinType.PinSubCategoryMemberReference=(),PinType.PinValueType=(),PinType.bIsMap=False,PinType.bIsSet=False,PinType.bIsArray=False,PinType.bIsReference=False,PinType.bIsConst=False,PinType.bIsWeakPointer=False,DefaultObject="/Script/Engine.Default__GameplayStatics",PersistentGuid=00000000000000000000000000000000,bHidden=True,bNotConnectable=False,bDefaultValueIsReadOnly=False,bDefaultValueIsIgnored=False,bAdvancedView=False,)
+	   CustomProperties Pin (PinId=539E882646828F110DA8539329878814,PinName="WorldContextObject",PinToolTip="World Context Object\nObject Reference",PinType.PinCategory="object",PinType.PinSubCategory="",PinType.PinSubCategoryObject=Class'/Script/CoreUObject.Object',PinType.PinSubCategoryMemberReference=(),PinType.PinValueType=(),PinType.bIsMap=False,PinType.bIsSet=False,PinType.bIsArray=False,PinType.bIsReference=False,PinType.bIsConst=True,PinType.bIsWeakPointer=False,PersistentGuid=00000000000000000000000000000000,bHidden=True,bNotConnectable=False,bDefaultValueIsReadOnly=False,bDefaultValueIsIgnored=False,bAdvancedView=False,)
+	   CustomProperties Pin (PinId=F81E36D8454550AAAEF10DB57E9FC71F,PinName="ReturnValue",PinToolTip="Return Value\nGame Mode Base Reference",Direction="EGPD_Output",PinType.PinCategory="object",PinType.PinSubCategory="",PinType.PinSubCategoryObject=Class'/Script/Engine.GameModeBase',PinType.PinSubCategoryMemberReference=(),PinType.PinValueType=(),PinType.bIsMap=False,PinType.bIsSet=False,PinType.bIsArray=False,PinType.bIsReference=False,PinType.bIsConst=False,PinType.bIsWeakPointer=False,LinkedTo=(K2Node_DynamicCast_5 8B0B2634480FE02DE82004ADB105AB79,),PersistentGuid=00000000000000000000000000000000,bHidden=False,bNotConnectable=False,bDefaultValueIsReadOnly=False,bDefaultValueIsIgnored=False,bAdvancedView=False,)
+	End Object
+	Begin Object Class=K2Node_DynamicCast Name="K2Node_DynamicCast_5"
+	   TargetType=BlueprintGeneratedClass'/Game/Prototype/Game/BP_OdinGameMode.BP_OdinGameMode_C'
+	   NodePosX=-512
+	   NodePosY=-4480
+	   NodeGuid=2A2928BD4CBEFB243D86D796866C7C1B
+	   CustomProperties Pin (PinId=3228E8FA4B588A45DCEA3A839D739E9A,PinName="execute",PinType.PinCategory="exec",PinType.PinSubCategory="",PinType.PinSubCategoryObject=None,PinType.PinSubCategoryMemberReference=(),PinType.PinValueType=(),PinType.bIsMap=False,PinType.bIsSet=False,PinType.bIsArray=False,PinType.bIsReference=False,PinType.bIsConst=False,PinType.bIsWeakPointer=False,LinkedTo=(K2Node_Event_52 5A6DBA7144DD625770231A9C725C0B49,),PersistentGuid=00000000000000000000000000000000,bHidden=False,bNotConnectable=False,bDefaultValueIsReadOnly=False,bDefaultValueIsIgnored=False,bAdvancedView=False,)
+	   CustomProperties Pin (PinId=C46EEAC7458FE7A6949E78BA06EE5B54,PinName="then",Direction="EGPD_Output",PinType.PinCategory="exec",PinType.PinSubCategory="",PinType.PinSubCategoryObject=None,PinType.PinSubCategoryMemberReference=(),PinType.PinValueType=(),PinType.bIsMap=False,PinType.bIsSet=False,PinType.bIsArray=False,PinType.bIsReference=False,PinType.bIsConst=False,PinType.bIsWeakPointer=False,LinkedTo=(K2Node_CallFunction_389 B9D0319843BB38F57642CEB7B91707F9,),PersistentGuid=00000000000000000000000000000000,bHidden=False,bNotConnectable=False,bDefaultValueIsReadOnly=False,bDefaultValueIsIgnored=False,bAdvancedView=False,)
+	   CustomProperties Pin (PinId=9EAA2A464990955344C076A409669F6D,PinName="CastFailed",Direction="EGPD_Output",PinType.PinCategory="exec",PinType.PinSubCategory="",PinType.PinSubCategoryObject=None,PinType.PinSubCategoryMemberReference=(),PinType.PinValueType=(),PinType.bIsMap=False,PinType.bIsSet=False,PinType.bIsArray=False,PinType.bIsReference=False,PinType.bIsConst=False,PinType.bIsWeakPointer=False,PersistentGuid=00000000000000000000000000000000,bHidden=False,bNotConnectable=False,bDefaultValueIsReadOnly=False,bDefaultValueIsIgnored=False,bAdvancedView=False,)
+	   CustomProperties Pin (PinId=8B0B2634480FE02DE82004ADB105AB79,PinName="Object",PinType.PinCategory="object",PinType.PinSubCategory="",PinType.PinSubCategoryObject=Class'/Script/CoreUObject.Object',PinType.PinSubCategoryMemberReference=(),PinType.PinValueType=(),PinType.bIsMap=False,PinType.bIsSet=False,PinType.bIsArray=False,PinType.bIsReference=False,PinType.bIsConst=False,PinType.bIsWeakPointer=False,LinkedTo=(K2Node_CallFunction_280 F81E36D8454550AAAEF10DB57E9FC71F,),PersistentGuid=00000000000000000000000000000000,bHidden=False,bNotConnectable=False,bDefaultValueIsReadOnly=False,bDefaultValueIsIgnored=False,bAdvancedView=False,)
+	   CustomProperties Pin (PinId=6B739B644EEBFA053240C1B460B2E558,PinName="AsBP Odin Game Mode",Direction="EGPD_Output",PinType.PinCategory="object",PinType.PinSubCategory="",PinType.PinSubCategoryObject=BlueprintGeneratedClass'/Game/Prototype/Game/BP_OdinGameMode.BP_OdinGameMode_C',PinType.PinSubCategoryMemberReference=(),PinType.PinValueType=(),PinType.bIsMap=False,PinType.bIsSet=False,PinType.bIsArray=False,PinType.bIsReference=False,PinType.bIsConst=False,PinType.bIsWeakPointer=False,LinkedTo=(K2Node_CallFunction_389 3117651A4CB6ACA2EAEB6BA4272EC84B,),PersistentGuid=00000000000000000000000000000000,bHidden=False,bNotConnectable=False,bDefaultValueIsReadOnly=False,bDefaultValueIsIgnored=False,bAdvancedView=False,)
+	   CustomProperties Pin (PinId=BEB2F603410D2A627B7D089246E322C7,PinName="bSuccess",Direction="EGPD_Output",PinType.PinCategory="bool",PinType.PinSubCategory="",PinType.PinSubCategoryObject=None,PinType.PinSubCategoryMemberReference=(),PinType.PinValueType=(),PinType.bIsMap=False,PinType.bIsSet=False,PinType.bIsArray=False,PinType.bIsReference=False,PinType.bIsConst=False,PinType.bIsWeakPointer=False,PersistentGuid=00000000000000000000000000000000,bHidden=True,bNotConnectable=False,bDefaultValueIsReadOnly=False,bDefaultValueIsIgnored=False,bAdvancedView=False,)
+	End Object
+	Begin Object Class=K2Node_CallFunction Name="K2Node_CallFunction_389"
+	   FunctionReference=(MemberParent=BlueprintGeneratedClass'/Game/Prototype/Game/BP_OdinGameMode.BP_OdinGameMode_C',MemberName="End Loading Sequence",MemberGuid=A1D38DD0498B32ACC9A936A032D370AC)
+	   NodePosX=-144
+	   NodePosY=-4496
+	   NodeGuid=4498E305455CCB64B47710BE1ECE74B6
+	   CustomProperties Pin (PinId=B9D0319843BB38F57642CEB7B91707F9,PinName="execute",PinToolTip="\nExec",PinType.PinCategory="exec",PinType.PinSubCategory="",PinType.PinSubCategoryObject=None,PinType.PinSubCategoryMemberReference=(),PinType.PinValueType=(),PinType.bIsMap=False,PinType.bIsSet=False,PinType.bIsArray=False,PinType.bIsReference=False,PinType.bIsConst=False,PinType.bIsWeakPointer=False,LinkedTo=(K2Node_DynamicCast_5 C46EEAC7458FE7A6949E78BA06EE5B54,),PersistentGuid=00000000000000000000000000000000,bHidden=False,bNotConnectable=False,bDefaultValueIsReadOnly=False,bDefaultValueIsIgnored=False,bAdvancedView=False,)
+	   CustomProperties Pin (PinId=080DB6114DA99E281AD26280646A113F,PinName="then",PinToolTip="\nExec",Direction="EGPD_Output",PinType.PinCategory="exec",PinType.PinSubCategory="",PinType.PinSubCategoryObject=None,PinType.PinSubCategoryMemberReference=(),PinType.PinValueType=(),PinType.bIsMap=False,PinType.bIsSet=False,PinType.bIsArray=False,PinType.bIsReference=False,PinType.bIsConst=False,PinType.bIsWeakPointer=False,LinkedTo=(K2Node_CallFunction_147 0D4F8EEB4254987507E56590D0CCA616,),PersistentGuid=00000000000000000000000000000000,bHidden=False,bNotConnectable=False,bDefaultValueIsReadOnly=False,bDefaultValueIsIgnored=False,bAdvancedView=False,)
+	   CustomProperties Pin (PinId=3117651A4CB6ACA2EAEB6BA4272EC84B,PinName="self",PinFriendlyName=NSLOCTEXT("K2Node", "Target", "Target"),PinToolTip="Target\nBP Odin Game Mode Reference",PinType.PinCategory="object",PinType.PinSubCategory="",PinType.PinSubCategoryObject=BlueprintGeneratedClass'/Game/Prototype/Game/BP_OdinGameMode.BP_OdinGameMode_C',PinType.PinSubCategoryMemberReference=(),PinType.PinValueType=(),PinType.bIsMap=False,PinType.bIsSet=False,PinType.bIsArray=False,PinType.bIsReference=False,PinType.bIsConst=False,PinType.bIsWeakPointer=False,LinkedTo=(K2Node_DynamicCast_5 6B739B644EEBFA053240C1B460B2E558,),PersistentGuid=00000000000000000000000000000000,bHidden=False,bNotConnectable=False,bDefaultValueIsReadOnly=False,bDefaultValueIsIgnored=False,bAdvancedView=False,)
+	End Object
+	Begin Object Class=K2Node_CallFunction Name="K2Node_CallFunction_146"
+	   bIsPureFunc=True
+	   FunctionReference=(MemberParent=Class'/Script/Engine.GameplayStatics',MemberName="GetPlayerPawn")
+	   NodePosX=-32
+	   NodePosY=-4352
+	   NodeGuid=5AC649304096413924B030B0014FD935
+	   CustomProperties Pin (PinId=329EDA7E40EFEF35615E988D15DE170A,PinName="self",PinFriendlyName=NSLOCTEXT("K2Node", "Target", "Target"),PinToolTip="Target\nGameplay Statics Reference",PinType.PinCategory="object",PinType.PinSubCategory="",PinType.PinSubCategoryObject=Class'/Script/Engine.GameplayStatics',PinType.PinSubCategoryMemberReference=(),PinType.PinValueType=(),PinType.bIsMap=False,PinType.bIsSet=False,PinType.bIsArray=False,PinType.bIsReference=False,PinType.bIsConst=False,PinType.bIsWeakPointer=False,DefaultObject="/Script/Engine.Default__GameplayStatics",PersistentGuid=00000000000000000000000000000000,bHidden=True,bNotConnectable=False,bDefaultValueIsReadOnly=False,bDefaultValueIsIgnored=False,bAdvancedView=False,)
+	   CustomProperties Pin (PinId=89848130466D916E810B3F99D04CBF1F,PinName="WorldContextObject",PinToolTip="World Context Object\nObject Reference",PinType.PinCategory="object",PinType.PinSubCategory="",PinType.PinSubCategoryObject=Class'/Script/CoreUObject.Object',PinType.PinSubCategoryMemberReference=(),PinType.PinValueType=(),PinType.bIsMap=False,PinType.bIsSet=False,PinType.bIsArray=False,PinType.bIsReference=False,PinType.bIsConst=True,PinType.bIsWeakPointer=False,PersistentGuid=00000000000000000000000000000000,bHidden=True,bNotConnectable=False,bDefaultValueIsReadOnly=False,bDefaultValueIsIgnored=False,bAdvancedView=False,)
+	   CustomProperties Pin (PinId=02F269A847C5688444C9AEAEA78D5124,PinName="PlayerIndex",PinToolTip="Player Index\nInteger",PinType.PinCategory="int",PinType.PinSubCategory="",PinType.PinSubCategoryObject=None,PinType.PinSubCategoryMemberReference=(),PinType.PinValueType=(),PinType.bIsMap=False,PinType.bIsSet=False,PinType.bIsArray=False,PinType.bIsReference=False,PinType.bIsConst=False,PinType.bIsWeakPointer=False,DefaultValue="0",AutogeneratedDefaultValue="0",PersistentGuid=00000000000000000000000000000000,bHidden=False,bNotConnectable=False,bDefaultValueIsReadOnly=False,bDefaultValueIsIgnored=False,bAdvancedView=False,)
+	   CustomProperties Pin (PinId=A55AE1CA4B113E55AEA44FA5BD868FEB,PinName="ReturnValue",PinToolTip="Return Value\nPawn Reference",Direction="EGPD_Output",PinType.PinCategory="object",PinType.PinSubCategory="",PinType.PinSubCategoryObject=Class'/Script/Engine.Pawn',PinType.PinSubCategoryMemberReference=(),PinType.PinValueType=(),PinType.bIsMap=False,PinType.bIsSet=False,PinType.bIsArray=False,PinType.bIsReference=False,PinType.bIsConst=False,PinType.bIsWeakPointer=False,LinkedTo=(K2Node_CallFunction_147 05DD2A284290ED5C1DD79C95F86B089B,),PersistentGuid=00000000000000000000000000000000,bHidden=False,bNotConnectable=False,bDefaultValueIsReadOnly=False,bDefaultValueIsIgnored=False,bAdvancedView=False,)
+	End Object
+	Begin Object Class=K2Node_CallFunction Name="K2Node_CallFunction_147"
+	   FunctionReference=(MemberName="K2_TeleportTo",bSelfContext=True)
+	   NodePosX=400
+	   NodePosY=-4496
+	   NodeGuid=97C3515E434225B49FD8669033C40E7D
+	   CustomProperties Pin (PinId=0D4F8EEB4254987507E56590D0CCA616,PinName="execute",PinToolTip="\nExec",PinType.PinCategory="exec",PinType.PinSubCategory="",PinType.PinSubCategoryObject=None,PinType.PinSubCategoryMemberReference=(),PinType.PinValueType=(),PinType.bIsMap=False,PinType.bIsSet=False,PinType.bIsArray=False,PinType.bIsReference=False,PinType.bIsConst=False,PinType.bIsWeakPointer=False,LinkedTo=(K2Node_CallFunction_389 080DB6114DA99E281AD26280646A113F,),PersistentGuid=00000000000000000000000000000000,bHidden=False,bNotConnectable=False,bDefaultValueIsReadOnly=False,bDefaultValueIsIgnored=False,bAdvancedView=False,)
+	   CustomProperties Pin (PinId=DA74E61542E071192041BF8670B2B044,PinName="then",PinToolTip="\nExec",Direction="EGPD_Output",PinType.PinCategory="exec",PinType.PinSubCategory="",PinType.PinSubCategoryObject=None,PinType.PinSubCategoryMemberReference=(),PinType.PinValueType=(),PinType.bIsMap=False,PinType.bIsSet=False,PinType.bIsArray=False,PinType.bIsReference=False,PinType.bIsConst=False,PinType.bIsWeakPointer=False,PersistentGuid=00000000000000000000000000000000,bHidden=False,bNotConnectable=False,bDefaultValueIsReadOnly=False,bDefaultValueIsIgnored=False,bAdvancedView=False,)
+	   CustomProperties Pin (PinId=05DD2A284290ED5C1DD79C95F86B089B,PinName="self",PinFriendlyName=NSLOCTEXT("K2Node", "Target", "Target"),PinToolTip="Target\nActor Reference",PinType.PinCategory="object",PinType.PinSubCategory="",PinType.PinSubCategoryObject=Class'/Script/Engine.Actor',PinType.PinSubCategoryMemberReference=(),PinType.PinValueType=(),PinType.bIsMap=False,PinType.bIsSet=False,PinType.bIsArray=False,PinType.bIsReference=False,PinType.bIsConst=False,PinType.bIsWeakPointer=False,LinkedTo=(K2Node_CallFunction_146 A55AE1CA4B113E55AEA44FA5BD868FEB,),PersistentGuid=00000000000000000000000000000000,bHidden=False,bNotConnectable=False,bDefaultValueIsReadOnly=False,bDefaultValueIsIgnored=False,bAdvancedView=False,)
+	   CustomProperties Pin (PinId=7ADD16464BF350CD2D3D5DA206062506,PinName="DestLocation",PinToolTip="Dest Location\nVector \n\nThe target destination point",PinType.PinCategory="struct",PinType.PinSubCategory="",PinType.PinSubCategoryObject=ScriptStruct'/Script/CoreUObject.Vector',PinType.PinSubCategoryMemberReference=(),PinType.PinValueType=(),PinType.bIsMap=False,PinType.bIsSet=False,PinType.bIsArray=False,PinType.bIsReference=False,PinType.bIsConst=False,PinType.bIsWeakPointer=False,DefaultValue="0, 0, 0",AutogeneratedDefaultValue="0, 0, 0",LinkedTo=(K2Node_CallFunction_161 8EDCF5B4407E0D67A46F56B3A7EBDC6E,),PersistentGuid=00000000000000000000000000000000,bHidden=False,bNotConnectable=False,bDefaultValueIsReadOnly=False,bDefaultValueIsIgnored=False,bAdvancedView=False,)
+	   CustomProperties Pin (PinId=CADAD3724E0A28E00552BCA72427B545,PinName="DestRotation",PinToolTip="Dest Rotation\nRotator \n\nThe target rotation at the destination",PinType.PinCategory="struct",PinType.PinSubCategory="",PinType.PinSubCategoryObject=ScriptStruct'/Script/CoreUObject.Rotator',PinType.PinSubCategoryMemberReference=(),PinType.PinValueType=(),PinType.bIsMap=False,PinType.bIsSet=False,PinType.bIsArray=False,PinType.bIsReference=False,PinType.bIsConst=False,PinType.bIsWeakPointer=False,DefaultValue="0, 0, 0",AutogeneratedDefaultValue="0, 0, 0",LinkedTo=(K2Node_CallFunction_197 177F737D497F0B71A0D9548711E45134,),PersistentGuid=00000000000000000000000000000000,bHidden=False,bNotConnectable=False,bDefaultValueIsReadOnly=False,bDefaultValueIsIgnored=False,bAdvancedView=False,)
+	   CustomProperties Pin (PinId=150C2F4B41C32C40A5EE5F928C2473FE,PinName="ReturnValue",PinToolTip="Return Value\nBoolean\n\ntrue if the actor has been successfully moved, or false if it couldn\'t fit.",Direction="EGPD_Output",PinType.PinCategory="bool",PinType.PinSubCategory="",PinType.PinSubCategoryObject=None,PinType.PinSubCategoryMemberReference=(),PinType.PinValueType=(),PinType.bIsMap=False,PinType.bIsSet=False,PinType.bIsArray=False,PinType.bIsReference=False,PinType.bIsConst=False,PinType.bIsWeakPointer=False,DefaultValue="false",AutogeneratedDefaultValue="false",PersistentGuid=00000000000000000000000000000000,bHidden=False,bNotConnectable=False,bDefaultValueIsReadOnly=False,bDefaultValueIsIgnored=False,bAdvancedView=False,)
+	End Object
+	Begin Object Class=K2Node_Literal Name="K2Node_Literal_0"
+	   ObjectRef=PlayerStart'/BasicCity/Maps/BasicCityCityCenterDay.BasicCityCityCenterDay:PersistentLevel.PlayerStart2_3'
+	   NodePosX=-272
+	   NodePosY=-4192
+	   NodeGuid=58EDCAD94D643329430D2D879E37B3E5
+	   CustomProperties Pin (PinId=AFA007CD4AFA013315CE299F79589275,PinName="PlayerStart2",PinFriendlyName=NSLOCTEXT("", "C0C388B24DBCBE06EB1E1B8C7FFFE555", "PlayerStart2"),Direction="EGPD_Output",PinType.PinCategory="object",PinType.PinSubCategory="",PinType.PinSubCategoryObject=Class'/Script/Engine.PlayerStart',PinType.PinSubCategoryMemberReference=(),PinType.PinValueType=(),PinType.bIsMap=False,PinType.bIsSet=False,PinType.bIsArray=False,PinType.bIsReference=False,PinType.bIsConst=False,PinType.bIsWeakPointer=False,LinkedTo=(K2Node_CallFunction_161 76593FC04576E3C4E850799E8DD55EC5,K2Node_CallFunction_197 40BECA4844231CB62E8E729102B58F73,),PersistentGuid=00000000000000000000000000000000,bHidden=False,bNotConnectable=False,bDefaultValueIsReadOnly=False,bDefaultValueIsIgnored=False,bAdvancedView=False,)
+	End Object
+	Begin Object Class=K2Node_CallFunction Name="K2Node_CallFunction_161"
+	   bIsPureFunc=True
+	   bIsConstFunc=True
+	   FunctionReference=(MemberName="K2_GetActorLocation",bSelfContext=True)
+	   NodePosX=16
+	   NodePosY=-4272
+	   NodeGuid=72AB37F441B22C2E9CB0D3BFA6D9A48F
+	   CustomProperties Pin (PinId=76593FC04576E3C4E850799E8DD55EC5,PinName="self",PinFriendlyName=NSLOCTEXT("K2Node", "Target", "Target"),PinToolTip="Target\nActor Reference",PinType.PinCategory="object",PinType.PinSubCategory="",PinType.PinSubCategoryObject=Class'/Script/Engine.Actor',PinType.PinSubCategoryMemberReference=(),PinType.PinValueType=(),PinType.bIsMap=False,PinType.bIsSet=False,PinType.bIsArray=False,PinType.bIsReference=False,PinType.bIsConst=False,PinType.bIsWeakPointer=False,LinkedTo=(K2Node_Literal_0 AFA007CD4AFA013315CE299F79589275,),PersistentGuid=00000000000000000000000000000000,bHidden=False,bNotConnectable=False,bDefaultValueIsReadOnly=False,bDefaultValueIsIgnored=False,bAdvancedView=False,)
+	   CustomProperties Pin (PinId=8EDCF5B4407E0D67A46F56B3A7EBDC6E,PinName="ReturnValue",PinToolTip="Return Value\nVector ",Direction="EGPD_Output",PinType.PinCategory="struct",PinType.PinSubCategory="",PinType.PinSubCategoryObject=ScriptStruct'/Script/CoreUObject.Vector',PinType.PinSubCategoryMemberReference=(),PinType.PinValueType=(),PinType.bIsMap=False,PinType.bIsSet=False,PinType.bIsArray=False,PinType.bIsReference=False,PinType.bIsConst=False,PinType.bIsWeakPointer=False,AutogeneratedDefaultValue="0, 0, 0",LinkedTo=(K2Node_CallFunction_147 7ADD16464BF350CD2D3D5DA206062506,),PersistentGuid=00000000000000000000000000000000,bHidden=False,bNotConnectable=False,bDefaultValueIsReadOnly=False,bDefaultValueIsIgnored=False,bAdvancedView=False,)
+	End Object
+	Begin Object Class=K2Node_CallFunction Name="K2Node_CallFunction_197"
+	   bIsPureFunc=True
+	   bIsConstFunc=True
+	   FunctionReference=(MemberName="K2_GetActorRotation",bSelfContext=True)
+	   NodePosX=16
+	   NodePosY=-4176
+	   NodeGuid=AF5DAC314C10340F9EA01999BB42D085
+	   CustomProperties Pin (PinId=40BECA4844231CB62E8E729102B58F73,PinName="self",PinFriendlyName=NSLOCTEXT("K2Node", "Target", "Target"),PinToolTip="Target\nActor Reference",PinType.PinCategory="object",PinType.PinSubCategory="",PinType.PinSubCategoryObject=Class'/Script/Engine.Actor',PinType.PinSubCategoryMemberReference=(),PinType.PinValueType=(),PinType.bIsMap=False,PinType.bIsSet=False,PinType.bIsArray=False,PinType.bIsReference=False,PinType.bIsConst=False,PinType.bIsWeakPointer=False,LinkedTo=(K2Node_Literal_0 AFA007CD4AFA013315CE299F79589275,),PersistentGuid=00000000000000000000000000000000,bHidden=False,bNotConnectable=False,bDefaultValueIsReadOnly=False,bDefaultValueIsIgnored=False,bAdvancedView=False,)
+	   CustomProperties Pin (PinId=177F737D497F0B71A0D9548711E45134,PinName="ReturnValue",PinToolTip="Return Value\nRotator ",Direction="EGPD_Output",PinType.PinCategory="struct",PinType.PinSubCategory="",PinType.PinSubCategoryObject=ScriptStruct'/Script/CoreUObject.Rotator',PinType.PinSubCategoryMemberReference=(),PinType.PinValueType=(),PinType.bIsMap=False,PinType.bIsSet=False,PinType.bIsArray=False,PinType.bIsReference=False,PinType.bIsConst=False,PinType.bIsWeakPointer=False,AutogeneratedDefaultValue="0, 0, 0",LinkedTo=(K2Node_CallFunction_147 CADAD3724E0A28E00552BCA72427B545,),PersistentGuid=00000000000000000000000000000000,bHidden=False,bNotConnectable=False,bDefaultValueIsReadOnly=False,bDefaultValueIsIgnored=False,bAdvancedView=False,)
+	End Object
+[/EXCERPT:snippet]
+-->
\ No newline at end of file
diff --git a/Engine/Documentation/Source/Modding/RoboRecall/HowTo/SetupGameplay/CreateMapMod/CreateMapMod.INT.udn b/Engine/Documentation/Source/Modding/RoboRecall/HowTo/SetupGameplay/CreateMapMod/CreateMapMod.INT.udn
new file mode 100644
--- /dev/null
+++ b/Engine/Documentation/Source/Modding/RoboRecall/HowTo/SetupGameplay/CreateMapMod/CreateMapMod.INT.udn
@@ -0,0 +1,44 @@
+Availability:Public
+Title:1. Creating a Map Mod
+Crumbs:
+Description:
+Parent:Modding/RoboRecall/HowTo/SetupGameplay
+Order:1
+tags:Modding
+tags:Robo Recall
+type:how-to
+checkpoint:mod_gameplay
+
+[NAV]
+
+[REGION:note]
+This instructions below walk you through creating a new map mod from scratch. If you already have a map created, you can skip this step and move on to [placing Actors](Modding/RoboRecall/HowTo/SetupGameplay/PlaceActors).
+[/REGION]
+
+1. In the Level Editor toolbar, click the **Create Mod** button to open the **New Game Mod** window.
+
+1. Select one of the map templates. In this example, we chose **Custom City Center Day**.
+
+	![](mod_create.png)(convert:false)
+
+1. Enter a name for the new mod in the **Name** field.
+
+	![](mod_name.png)(convert:false)
+	
+1. Click the **Create Mod** button to generate your mod files.
+
+1. In the **Content Browser**, double-click the **Maps** folder and then double-click your map asset to load your new map.
+
+	![](map_unbuilt.png)(convert:false)
+
+1. In the Level Editor toolbar, click the **Build** button to generate lighting, navigation, and other data for your map.
+
+	![](map_build.png)(convert:false)
+
+## Result
+
+You now have a map that provides a perfect environment to add gameplay elements to.
+
+![](map_built.png)(convert:false)
+
+[NAV]
\ No newline at end of file
diff --git a/Engine/Documentation/Source/Modding/RoboRecall/HowTo/SetupGameplay/EndLevelScript/EndLevelScript.INT.udn b/Engine/Documentation/Source/Modding/RoboRecall/HowTo/SetupGameplay/EndLevelScript/EndLevelScript.INT.udn
new file mode 100644
--- /dev/null
+++ b/Engine/Documentation/Source/Modding/RoboRecall/HowTo/SetupGameplay/EndLevelScript/EndLevelScript.INT.udn
@@ -0,0 +1,166 @@
+Availability:Public
+Title:5. Level Completed Scripting
+Crumbs:
+Description:
+Parent:Modding/RoboRecall/HowTo/SetupGameplay
+Order:5
+tags:Modding
+tags:Robo Recall
+type:how-to
+checkpoint:mod_gameplay
+
+[NAV]
+
+When all the robots in the wave have been recalled and the level is complete, you want the player to return to the Hub and continue the normal flow of the game.
+
+1. In the Level Editor toolbar click the **Blueprints** button and then choose **Open Level Blueprint** to open the script for the current level in the Blueprint Editor.
+
+	![](level_blueprint.png)(convert:false)
+
+1. In the Level Editor Viewport, select the **Odin AI Wave Manager** Actor you place previously. Back in the Blueprint Editor, right-click in the Event Graph and choose **Add On Wave Completed**.
+
+	![](wave_complete.png)(convert:false)
+
+1. Drag off the exec output and choose **Delay**. Set the **Duration** to `5.0`. This will allow  a brief period after the level is complete before the player is taken back to the Hub.
+
+	![](delay.png)(convert:false)
+
+1. Right-click in the Event Graph and choose **Get Game Mode** to get a reference to the current [Game Mode](Gameplay/Framework/GameMode) instance. Drag off the **Return Value** pin and choose **Cast to BP_OdinGameMode**. Connect the exec output of the **Delay** event to the exec input of the **Cast to BP_OdinGameMode** node.
+
+	![](game_mode.png)(convert:false)
+
+1. Drag off the **As BP Odin Game Mode** pin and choose **Bind Event to onBeginLoadingSequenceDone**. Drag off the **event** input and choose **Add Custom Event**. Name the event **ReturnToHub**.
+
+	![](loading_event.png)(convert:false)
+
+1. Drag off the exec output and choose **Begin Loading Sequence** to start the loading screen. Connectthe **Target** input to the **As BP Odin Game Mode** output from the cast node.
+
+	![](begin_loading.png)(convert:false)
+
+1. Drag off the **As BP Odin Game Mode** output on the cast node and choose **Get Main Pawn** to get a reference to the player. Drag off **Main Pawn** and choose **Get Controller** to get a reference to the player's Controller. Drag off **Return Value** and choose **Cast to Player Controller** to make sure it is the Controller of a human-controlled player, which is necessary to execute a console command.
+
+	![](controller.png)
+
+1. Drag off the exec output of the ReturnToHub event and choose **execute Console Command**. Enter _open hub_evn_ in the **Command** property. This will send a command to open the Hub level. Connect the **Specific Player** input to the **As Player Controller** output on the cast node.
+
+	[OBJECT:Snippet]
+		[PARAMLITERAL:id]
+		result
+		[/PARAMLITERAL]
+		[PARAMLITERAL:type]
+		blueprint
+		[/PARAMLITERAL]
+		[PARAM:image]
+		![](console_command.png)(convert:false)
+		[/PARAM]
+		[PARAM:snippet]
+			[INCLUDE:#snippet]
+		[/PARAM]
+	[/OBJECT]
+
+	[REGION:caption]
+	Click the icon in the upper left corner of this image to copy the Blueprint Graph and paste it into your project.
+	[/REGION]
+
+## Result
+
+When the wave is complete, the loading sequence begins and the player is taken to the Hub.
+
+![](fade.png)(convert:false)
+
+[NAV]
+
+<!--
+[EXCERPT:snippet]
+	Begin Object Class=K2Node_Event Name="K2Node_Event_52"
+	   EventReference=(MemberParent=Class'/Script/Engine.Actor',MemberName="ReceiveBeginPlay")
+	   bOverrideFunction=True
+	   NodePosX=-928
+	   NodePosY=-4480
+	   NodeGuid=6C849E7241A6108A36D5C1BEAAEAE94F
+	   CustomProperties Pin (PinId=B35535C443B530261C41B8A84FF2CC42,PinName="OutputDelegate",Direction="EGPD_Output",PinType.PinCategory="delegate",PinType.PinSubCategory="",PinType.PinSubCategoryObject=None,PinType.PinSubCategoryMemberReference=(MemberParent=Class'/Script/Engine.Actor',MemberName="ReceiveBeginPlay"),PinType.PinValueType=(),PinType.bIsMap=False,PinType.bIsSet=False,PinType.bIsArray=False,PinType.bIsReference=False,PinType.bIsConst=False,PinType.bIsWeakPointer=False,PersistentGuid=00000000000000000000000000000000,bHidden=False,bNotConnectable=False,bDefaultValueIsReadOnly=False,bDefaultValueIsIgnored=False,bAdvancedView=False,)
+	   CustomProperties Pin (PinId=5A6DBA7144DD625770231A9C725C0B49,PinName="then",Direction="EGPD_Output",PinType.PinCategory="exec",PinType.PinSubCategory="",PinType.PinSubCategoryObject=None,PinType.PinSubCategoryMemberReference=(),PinType.PinValueType=(),PinType.bIsMap=False,PinType.bIsSet=False,PinType.bIsArray=False,PinType.bIsReference=False,PinType.bIsConst=False,PinType.bIsWeakPointer=False,LinkedTo=(K2Node_DynamicCast_5 3228E8FA4B588A45DCEA3A839D739E9A,),PersistentGuid=00000000000000000000000000000000,bHidden=False,bNotConnectable=False,bDefaultValueIsReadOnly=False,bDefaultValueIsIgnored=False,bAdvancedView=False,)
+	End Object
+	Begin Object Class=K2Node_CallFunction Name="K2Node_CallFunction_280"
+	   bIsPureFunc=True
+	   FunctionReference=(MemberParent=Class'/Script/Engine.GameplayStatics',MemberName="GetGameMode")
+	   NodePosX=-720
+	   NodePosY=-4416
+	   NodeGuid=2CB9ABB94CF9F2C3BAEBB0B38ED0F003
+	   CustomProperties Pin (PinId=1AAC46094DE5E3D0CB1F8488BEDBFE3B,PinName="self",PinFriendlyName=NSLOCTEXT("K2Node", "Target", "Target"),PinToolTip="Target\nGameplay Statics Reference",PinType.PinCategory="object",PinType.PinSubCategory="",PinType.PinSubCategoryObject=Class'/Script/Engine.GameplayStatics',PinType.PinSubCategoryMemberReference=(),PinType.PinValueType=(),PinType.bIsMap=False,PinType.bIsSet=False,PinType.bIsArray=False,PinType.bIsReference=False,PinType.bIsConst=False,PinType.bIsWeakPointer=False,DefaultObject="/Script/Engine.Default__GameplayStatics",PersistentGuid=00000000000000000000000000000000,bHidden=True,bNotConnectable=False,bDefaultValueIsReadOnly=False,bDefaultValueIsIgnored=False,bAdvancedView=False,)
+	   CustomProperties Pin (PinId=539E882646828F110DA8539329878814,PinName="WorldContextObject",PinToolTip="World Context Object\nObject Reference",PinType.PinCategory="object",PinType.PinSubCategory="",PinType.PinSubCategoryObject=Class'/Script/CoreUObject.Object',PinType.PinSubCategoryMemberReference=(),PinType.PinValueType=(),PinType.bIsMap=False,PinType.bIsSet=False,PinType.bIsArray=False,PinType.bIsReference=False,PinType.bIsConst=True,PinType.bIsWeakPointer=False,PersistentGuid=00000000000000000000000000000000,bHidden=True,bNotConnectable=False,bDefaultValueIsReadOnly=False,bDefaultValueIsIgnored=False,bAdvancedView=False,)
+	   CustomProperties Pin (PinId=F81E36D8454550AAAEF10DB57E9FC71F,PinName="ReturnValue",PinToolTip="Return Value\nGame Mode Base Reference",Direction="EGPD_Output",PinType.PinCategory="object",PinType.PinSubCategory="",PinType.PinSubCategoryObject=Class'/Script/Engine.GameModeBase',PinType.PinSubCategoryMemberReference=(),PinType.PinValueType=(),PinType.bIsMap=False,PinType.bIsSet=False,PinType.bIsArray=False,PinType.bIsReference=False,PinType.bIsConst=False,PinType.bIsWeakPointer=False,LinkedTo=(K2Node_DynamicCast_5 8B0B2634480FE02DE82004ADB105AB79,),PersistentGuid=00000000000000000000000000000000,bHidden=False,bNotConnectable=False,bDefaultValueIsReadOnly=False,bDefaultValueIsIgnored=False,bAdvancedView=False,)
+	End Object
+	Begin Object Class=K2Node_DynamicCast Name="K2Node_DynamicCast_5"
+	   TargetType=BlueprintGeneratedClass'/Game/Prototype/Game/BP_OdinGameMode.BP_OdinGameMode_C'
+	   NodePosX=-512
+	   NodePosY=-4480
+	   NodeGuid=2A2928BD4CBEFB243D86D796866C7C1B
+	   CustomProperties Pin (PinId=3228E8FA4B588A45DCEA3A839D739E9A,PinName="execute",PinType.PinCategory="exec",PinType.PinSubCategory="",PinType.PinSubCategoryObject=None,PinType.PinSubCategoryMemberReference=(),PinType.PinValueType=(),PinType.bIsMap=False,PinType.bIsSet=False,PinType.bIsArray=False,PinType.bIsReference=False,PinType.bIsConst=False,PinType.bIsWeakPointer=False,LinkedTo=(K2Node_Event_52 5A6DBA7144DD625770231A9C725C0B49,),PersistentGuid=00000000000000000000000000000000,bHidden=False,bNotConnectable=False,bDefaultValueIsReadOnly=False,bDefaultValueIsIgnored=False,bAdvancedView=False,)
+	   CustomProperties Pin (PinId=C46EEAC7458FE7A6949E78BA06EE5B54,PinName="then",Direction="EGPD_Output",PinType.PinCategory="exec",PinType.PinSubCategory="",PinType.PinSubCategoryObject=None,PinType.PinSubCategoryMemberReference=(),PinType.PinValueType=(),PinType.bIsMap=False,PinType.bIsSet=False,PinType.bIsArray=False,PinType.bIsReference=False,PinType.bIsConst=False,PinType.bIsWeakPointer=False,LinkedTo=(K2Node_CallFunction_389 B9D0319843BB38F57642CEB7B91707F9,),PersistentGuid=00000000000000000000000000000000,bHidden=False,bNotConnectable=False,bDefaultValueIsReadOnly=False,bDefaultValueIsIgnored=False,bAdvancedView=False,)
+	   CustomProperties Pin (PinId=9EAA2A464990955344C076A409669F6D,PinName="CastFailed",Direction="EGPD_Output",PinType.PinCategory="exec",PinType.PinSubCategory="",PinType.PinSubCategoryObject=None,PinType.PinSubCategoryMemberReference=(),PinType.PinValueType=(),PinType.bIsMap=False,PinType.bIsSet=False,PinType.bIsArray=False,PinType.bIsReference=False,PinType.bIsConst=False,PinType.bIsWeakPointer=False,PersistentGuid=00000000000000000000000000000000,bHidden=False,bNotConnectable=False,bDefaultValueIsReadOnly=False,bDefaultValueIsIgnored=False,bAdvancedView=False,)
+	   CustomProperties Pin (PinId=8B0B2634480FE02DE82004ADB105AB79,PinName="Object",PinType.PinCategory="object",PinType.PinSubCategory="",PinType.PinSubCategoryObject=Class'/Script/CoreUObject.Object',PinType.PinSubCategoryMemberReference=(),PinType.PinValueType=(),PinType.bIsMap=False,PinType.bIsSet=False,PinType.bIsArray=False,PinType.bIsReference=False,PinType.bIsConst=False,PinType.bIsWeakPointer=False,LinkedTo=(K2Node_CallFunction_280 F81E36D8454550AAAEF10DB57E9FC71F,),PersistentGuid=00000000000000000000000000000000,bHidden=False,bNotConnectable=False,bDefaultValueIsReadOnly=False,bDefaultValueIsIgnored=False,bAdvancedView=False,)
+	   CustomProperties Pin (PinId=6B739B644EEBFA053240C1B460B2E558,PinName="AsBP Odin Game Mode",Direction="EGPD_Output",PinType.PinCategory="object",PinType.PinSubCategory="",PinType.PinSubCategoryObject=BlueprintGeneratedClass'/Game/Prototype/Game/BP_OdinGameMode.BP_OdinGameMode_C',PinType.PinSubCategoryMemberReference=(),PinType.PinValueType=(),PinType.bIsMap=False,PinType.bIsSet=False,PinType.bIsArray=False,PinType.bIsReference=False,PinType.bIsConst=False,PinType.bIsWeakPointer=False,LinkedTo=(K2Node_CallFunction_389 3117651A4CB6ACA2EAEB6BA4272EC84B,),PersistentGuid=00000000000000000000000000000000,bHidden=False,bNotConnectable=False,bDefaultValueIsReadOnly=False,bDefaultValueIsIgnored=False,bAdvancedView=False,)
+	   CustomProperties Pin (PinId=BEB2F603410D2A627B7D089246E322C7,PinName="bSuccess",Direction="EGPD_Output",PinType.PinCategory="bool",PinType.PinSubCategory="",PinType.PinSubCategoryObject=None,PinType.PinSubCategoryMemberReference=(),PinType.PinValueType=(),PinType.bIsMap=False,PinType.bIsSet=False,PinType.bIsArray=False,PinType.bIsReference=False,PinType.bIsConst=False,PinType.bIsWeakPointer=False,PersistentGuid=00000000000000000000000000000000,bHidden=True,bNotConnectable=False,bDefaultValueIsReadOnly=False,bDefaultValueIsIgnored=False,bAdvancedView=False,)
+	End Object
+	Begin Object Class=K2Node_CallFunction Name="K2Node_CallFunction_389"
+	   FunctionReference=(MemberParent=BlueprintGeneratedClass'/Game/Prototype/Game/BP_OdinGameMode.BP_OdinGameMode_C',MemberName="End Loading Sequence",MemberGuid=A1D38DD0498B32ACC9A936A032D370AC)
+	   NodePosX=-144
+	   NodePosY=-4496
+	   NodeGuid=4498E305455CCB64B47710BE1ECE74B6
+	   CustomProperties Pin (PinId=B9D0319843BB38F57642CEB7B91707F9,PinName="execute",PinToolTip="\nExec",PinType.PinCategory="exec",PinType.PinSubCategory="",PinType.PinSubCategoryObject=None,PinType.PinSubCategoryMemberReference=(),PinType.PinValueType=(),PinType.bIsMap=False,PinType.bIsSet=False,PinType.bIsArray=False,PinType.bIsReference=False,PinType.bIsConst=False,PinType.bIsWeakPointer=False,LinkedTo=(K2Node_DynamicCast_5 C46EEAC7458FE7A6949E78BA06EE5B54,),PersistentGuid=00000000000000000000000000000000,bHidden=False,bNotConnectable=False,bDefaultValueIsReadOnly=False,bDefaultValueIsIgnored=False,bAdvancedView=False,)
+	   CustomProperties Pin (PinId=080DB6114DA99E281AD26280646A113F,PinName="then",PinToolTip="\nExec",Direction="EGPD_Output",PinType.PinCategory="exec",PinType.PinSubCategory="",PinType.PinSubCategoryObject=None,PinType.PinSubCategoryMemberReference=(),PinType.PinValueType=(),PinType.bIsMap=False,PinType.bIsSet=False,PinType.bIsArray=False,PinType.bIsReference=False,PinType.bIsConst=False,PinType.bIsWeakPointer=False,LinkedTo=(K2Node_CallFunction_147 0D4F8EEB4254987507E56590D0CCA616,),PersistentGuid=00000000000000000000000000000000,bHidden=False,bNotConnectable=False,bDefaultValueIsReadOnly=False,bDefaultValueIsIgnored=False,bAdvancedView=False,)
+	   CustomProperties Pin (PinId=3117651A4CB6ACA2EAEB6BA4272EC84B,PinName="self",PinFriendlyName=NSLOCTEXT("K2Node", "Target", "Target"),PinToolTip="Target\nBP Odin Game Mode Reference",PinType.PinCategory="object",PinType.PinSubCategory="",PinType.PinSubCategoryObject=BlueprintGeneratedClass'/Game/Prototype/Game/BP_OdinGameMode.BP_OdinGameMode_C',PinType.PinSubCategoryMemberReference=(),PinType.PinValueType=(),PinType.bIsMap=False,PinType.bIsSet=False,PinType.bIsArray=False,PinType.bIsReference=False,PinType.bIsConst=False,PinType.bIsWeakPointer=False,LinkedTo=(K2Node_DynamicCast_5 6B739B644EEBFA053240C1B460B2E558,),PersistentGuid=00000000000000000000000000000000,bHidden=False,bNotConnectable=False,bDefaultValueIsReadOnly=False,bDefaultValueIsIgnored=False,bAdvancedView=False,)
+	End Object
+	Begin Object Class=K2Node_CallFunction Name="K2Node_CallFunction_146"
+	   bIsPureFunc=True
+	   FunctionReference=(MemberParent=Class'/Script/Engine.GameplayStatics',MemberName="GetPlayerPawn")
+	   NodePosX=-32
+	   NodePosY=-4352
+	   NodeGuid=5AC649304096413924B030B0014FD935
+	   CustomProperties Pin (PinId=329EDA7E40EFEF35615E988D15DE170A,PinName="self",PinFriendlyName=NSLOCTEXT("K2Node", "Target", "Target"),PinToolTip="Target\nGameplay Statics Reference",PinType.PinCategory="object",PinType.PinSubCategory="",PinType.PinSubCategoryObject=Class'/Script/Engine.GameplayStatics',PinType.PinSubCategoryMemberReference=(),PinType.PinValueType=(),PinType.bIsMap=False,PinType.bIsSet=False,PinType.bIsArray=False,PinType.bIsReference=False,PinType.bIsConst=False,PinType.bIsWeakPointer=False,DefaultObject="/Script/Engine.Default__GameplayStatics",PersistentGuid=00000000000000000000000000000000,bHidden=True,bNotConnectable=False,bDefaultValueIsReadOnly=False,bDefaultValueIsIgnored=False,bAdvancedView=False,)
+	   CustomProperties Pin (PinId=89848130466D916E810B3F99D04CBF1F,PinName="WorldContextObject",PinToolTip="World Context Object\nObject Reference",PinType.PinCategory="object",PinType.PinSubCategory="",PinType.PinSubCategoryObject=Class'/Script/CoreUObject.Object',PinType.PinSubCategoryMemberReference=(),PinType.PinValueType=(),PinType.bIsMap=False,PinType.bIsSet=False,PinType.bIsArray=False,PinType.bIsReference=False,PinType.bIsConst=True,PinType.bIsWeakPointer=False,PersistentGuid=00000000000000000000000000000000,bHidden=True,bNotConnectable=False,bDefaultValueIsReadOnly=False,bDefaultValueIsIgnored=False,bAdvancedView=False,)
+	   CustomProperties Pin (PinId=02F269A847C5688444C9AEAEA78D5124,PinName="PlayerIndex",PinToolTip="Player Index\nInteger",PinType.PinCategory="int",PinType.PinSubCategory="",PinType.PinSubCategoryObject=None,PinType.PinSubCategoryMemberReference=(),PinType.PinValueType=(),PinType.bIsMap=False,PinType.bIsSet=False,PinType.bIsArray=False,PinType.bIsReference=False,PinType.bIsConst=False,PinType.bIsWeakPointer=False,DefaultValue="0",AutogeneratedDefaultValue="0",PersistentGuid=00000000000000000000000000000000,bHidden=False,bNotConnectable=False,bDefaultValueIsReadOnly=False,bDefaultValueIsIgnored=False,bAdvancedView=False,)
+	   CustomProperties Pin (PinId=A55AE1CA4B113E55AEA44FA5BD868FEB,PinName="ReturnValue",PinToolTip="Return Value\nPawn Reference",Direction="EGPD_Output",PinType.PinCategory="object",PinType.PinSubCategory="",PinType.PinSubCategoryObject=Class'/Script/Engine.Pawn',PinType.PinSubCategoryMemberReference=(),PinType.PinValueType=(),PinType.bIsMap=False,PinType.bIsSet=False,PinType.bIsArray=False,PinType.bIsReference=False,PinType.bIsConst=False,PinType.bIsWeakPointer=False,LinkedTo=(K2Node_CallFunction_147 05DD2A284290ED5C1DD79C95F86B089B,),PersistentGuid=00000000000000000000000000000000,bHidden=False,bNotConnectable=False,bDefaultValueIsReadOnly=False,bDefaultValueIsIgnored=False,bAdvancedView=False,)
+	End Object
+	Begin Object Class=K2Node_CallFunction Name="K2Node_CallFunction_147"
+	   FunctionReference=(MemberName="K2_TeleportTo",bSelfContext=True)
+	   NodePosX=400
+	   NodePosY=-4496
+	   NodeGuid=97C3515E434225B49FD8669033C40E7D
+	   CustomProperties Pin (PinId=0D4F8EEB4254987507E56590D0CCA616,PinName="execute",PinToolTip="\nExec",PinType.PinCategory="exec",PinType.PinSubCategory="",PinType.PinSubCategoryObject=None,PinType.PinSubCategoryMemberReference=(),PinType.PinValueType=(),PinType.bIsMap=False,PinType.bIsSet=False,PinType.bIsArray=False,PinType.bIsReference=False,PinType.bIsConst=False,PinType.bIsWeakPointer=False,LinkedTo=(K2Node_CallFunction_389 080DB6114DA99E281AD26280646A113F,),PersistentGuid=00000000000000000000000000000000,bHidden=False,bNotConnectable=False,bDefaultValueIsReadOnly=False,bDefaultValueIsIgnored=False,bAdvancedView=False,)
+	   CustomProperties Pin (PinId=DA74E61542E071192041BF8670B2B044,PinName="then",PinToolTip="\nExec",Direction="EGPD_Output",PinType.PinCategory="exec",PinType.PinSubCategory="",PinType.PinSubCategoryObject=None,PinType.PinSubCategoryMemberReference=(),PinType.PinValueType=(),PinType.bIsMap=False,PinType.bIsSet=False,PinType.bIsArray=False,PinType.bIsReference=False,PinType.bIsConst=False,PinType.bIsWeakPointer=False,PersistentGuid=00000000000000000000000000000000,bHidden=False,bNotConnectable=False,bDefaultValueIsReadOnly=False,bDefaultValueIsIgnored=False,bAdvancedView=False,)
+	   CustomProperties Pin (PinId=05DD2A284290ED5C1DD79C95F86B089B,PinName="self",PinFriendlyName=NSLOCTEXT("K2Node", "Target", "Target"),PinToolTip="Target\nActor Reference",PinType.PinCategory="object",PinType.PinSubCategory="",PinType.PinSubCategoryObject=Class'/Script/Engine.Actor',PinType.PinSubCategoryMemberReference=(),PinType.PinValueType=(),PinType.bIsMap=False,PinType.bIsSet=False,PinType.bIsArray=False,PinType.bIsReference=False,PinType.bIsConst=False,PinType.bIsWeakPointer=False,LinkedTo=(K2Node_CallFunction_146 A55AE1CA4B113E55AEA44FA5BD868FEB,),PersistentGuid=00000000000000000000000000000000,bHidden=False,bNotConnectable=False,bDefaultValueIsReadOnly=False,bDefaultValueIsIgnored=False,bAdvancedView=False,)
+	   CustomProperties Pin (PinId=7ADD16464BF350CD2D3D5DA206062506,PinName="DestLocation",PinToolTip="Dest Location\nVector \n\nThe target destination point",PinType.PinCategory="struct",PinType.PinSubCategory="",PinType.PinSubCategoryObject=ScriptStruct'/Script/CoreUObject.Vector',PinType.PinSubCategoryMemberReference=(),PinType.PinValueType=(),PinType.bIsMap=False,PinType.bIsSet=False,PinType.bIsArray=False,PinType.bIsReference=False,PinType.bIsConst=False,PinType.bIsWeakPointer=False,DefaultValue="0, 0, 0",AutogeneratedDefaultValue="0, 0, 0",LinkedTo=(K2Node_CallFunction_161 8EDCF5B4407E0D67A46F56B3A7EBDC6E,),PersistentGuid=00000000000000000000000000000000,bHidden=False,bNotConnectable=False,bDefaultValueIsReadOnly=False,bDefaultValueIsIgnored=False,bAdvancedView=False,)
+	   CustomProperties Pin (PinId=CADAD3724E0A28E00552BCA72427B545,PinName="DestRotation",PinToolTip="Dest Rotation\nRotator \n\nThe target rotation at the destination",PinType.PinCategory="struct",PinType.PinSubCategory="",PinType.PinSubCategoryObject=ScriptStruct'/Script/CoreUObject.Rotator',PinType.PinSubCategoryMemberReference=(),PinType.PinValueType=(),PinType.bIsMap=False,PinType.bIsSet=False,PinType.bIsArray=False,PinType.bIsReference=False,PinType.bIsConst=False,PinType.bIsWeakPointer=False,DefaultValue="0, 0, 0",AutogeneratedDefaultValue="0, 0, 0",LinkedTo=(K2Node_CallFunction_197 177F737D497F0B71A0D9548711E45134,),PersistentGuid=00000000000000000000000000000000,bHidden=False,bNotConnectable=False,bDefaultValueIsReadOnly=False,bDefaultValueIsIgnored=False,bAdvancedView=False,)
+	   CustomProperties Pin (PinId=150C2F4B41C32C40A5EE5F928C2473FE,PinName="ReturnValue",PinToolTip="Return Value\nBoolean\n\ntrue if the actor has been successfully moved, or false if it couldn\'t fit.",Direction="EGPD_Output",PinType.PinCategory="bool",PinType.PinSubCategory="",PinType.PinSubCategoryObject=None,PinType.PinSubCategoryMemberReference=(),PinType.PinValueType=(),PinType.bIsMap=False,PinType.bIsSet=False,PinType.bIsArray=False,PinType.bIsReference=False,PinType.bIsConst=False,PinType.bIsWeakPointer=False,DefaultValue="false",AutogeneratedDefaultValue="false",PersistentGuid=00000000000000000000000000000000,bHidden=False,bNotConnectable=False,bDefaultValueIsReadOnly=False,bDefaultValueIsIgnored=False,bAdvancedView=False,)
+	End Object
+	Begin Object Class=K2Node_Literal Name="K2Node_Literal_0"
+	   ObjectRef=PlayerStart'/BasicCity/Maps/BasicCityCityCenterDay.BasicCityCityCenterDay:PersistentLevel.PlayerStart2_3'
+	   NodePosX=-272
+	   NodePosY=-4192
+	   NodeGuid=58EDCAD94D643329430D2D879E37B3E5
+	   CustomProperties Pin (PinId=AFA007CD4AFA013315CE299F79589275,PinName="PlayerStart2",PinFriendlyName=NSLOCTEXT("", "C0C388B24DBCBE06EB1E1B8C7FFFE555", "PlayerStart2"),Direction="EGPD_Output",PinType.PinCategory="object",PinType.PinSubCategory="",PinType.PinSubCategoryObject=Class'/Script/Engine.PlayerStart',PinType.PinSubCategoryMemberReference=(),PinType.PinValueType=(),PinType.bIsMap=False,PinType.bIsSet=False,PinType.bIsArray=False,PinType.bIsReference=False,PinType.bIsConst=False,PinType.bIsWeakPointer=False,LinkedTo=(K2Node_CallFunction_161 76593FC04576E3C4E850799E8DD55EC5,K2Node_CallFunction_197 40BECA4844231CB62E8E729102B58F73,),PersistentGuid=00000000000000000000000000000000,bHidden=False,bNotConnectable=False,bDefaultValueIsReadOnly=False,bDefaultValueIsIgnored=False,bAdvancedView=False,)
+	End Object
+	Begin Object Class=K2Node_CallFunction Name="K2Node_CallFunction_161"
+	   bIsPureFunc=True
+	   bIsConstFunc=True
+	   FunctionReference=(MemberName="K2_GetActorLocation",bSelfContext=True)
+	   NodePosX=16
+	   NodePosY=-4272
+	   NodeGuid=72AB37F441B22C2E9CB0D3BFA6D9A48F
+	   CustomProperties Pin (PinId=76593FC04576E3C4E850799E8DD55EC5,PinName="self",PinFriendlyName=NSLOCTEXT("K2Node", "Target", "Target"),PinToolTip="Target\nActor Reference",PinType.PinCategory="object",PinType.PinSubCategory="",PinType.PinSubCategoryObject=Class'/Script/Engine.Actor',PinType.PinSubCategoryMemberReference=(),PinType.PinValueType=(),PinType.bIsMap=False,PinType.bIsSet=False,PinType.bIsArray=False,PinType.bIsReference=False,PinType.bIsConst=False,PinType.bIsWeakPointer=False,LinkedTo=(K2Node_Literal_0 AFA007CD4AFA013315CE299F79589275,),PersistentGuid=00000000000000000000000000000000,bHidden=False,bNotConnectable=False,bDefaultValueIsReadOnly=False,bDefaultValueIsIgnored=False,bAdvancedView=False,)
+	   CustomProperties Pin (PinId=8EDCF5B4407E0D67A46F56B3A7EBDC6E,PinName="ReturnValue",PinToolTip="Return Value\nVector ",Direction="EGPD_Output",PinType.PinCategory="struct",PinType.PinSubCategory="",PinType.PinSubCategoryObject=ScriptStruct'/Script/CoreUObject.Vector',PinType.PinSubCategoryMemberReference=(),PinType.PinValueType=(),PinType.bIsMap=False,PinType.bIsSet=False,PinType.bIsArray=False,PinType.bIsReference=False,PinType.bIsConst=False,PinType.bIsWeakPointer=False,AutogeneratedDefaultValue="0, 0, 0",LinkedTo=(K2Node_CallFunction_147 7ADD16464BF350CD2D3D5DA206062506,),PersistentGuid=00000000000000000000000000000000,bHidden=False,bNotConnectable=False,bDefaultValueIsReadOnly=False,bDefaultValueIsIgnored=False,bAdvancedView=False,)
+	End Object
+	Begin Object Class=K2Node_CallFunction Name="K2Node_CallFunction_197"
+	   bIsPureFunc=True
+	   bIsConstFunc=True
+	   FunctionReference=(MemberName="K2_GetActorRotation",bSelfContext=True)
+	   NodePosX=16
+	   NodePosY=-4176
+	   NodeGuid=AF5DAC314C10340F9EA01999BB42D085
+	   CustomProperties Pin (PinId=40BECA4844231CB62E8E729102B58F73,PinName="self",PinFriendlyName=NSLOCTEXT("K2Node", "Target", "Target"),PinToolTip="Target\nActor Reference",PinType.PinCategory="object",PinType.PinSubCategory="",PinType.PinSubCategoryObject=Class'/Script/Engine.Actor',PinType.PinSubCategoryMemberReference=(),PinType.PinValueType=(),PinType.bIsMap=False,PinType.bIsSet=False,PinType.bIsArray=False,PinType.bIsReference=False,PinType.bIsConst=False,PinType.bIsWeakPointer=False,LinkedTo=(K2Node_Literal_0 AFA007CD4AFA013315CE299F79589275,),PersistentGuid=00000000000000000000000000000000,bHidden=False,bNotConnectable=False,bDefaultValueIsReadOnly=False,bDefaultValueIsIgnored=False,bAdvancedView=False,)
+	   CustomProperties Pin (PinId=177F737D497F0B71A0D9548711E45134,PinName="ReturnValue",PinToolTip="Return Value\nRotator ",Direction="EGPD_Output",PinType.PinCategory="struct",PinType.PinSubCategory="",PinType.PinSubCategoryObject=ScriptStruct'/Script/CoreUObject.Rotator',PinType.PinSubCategoryMemberReference=(),PinType.PinValueType=(),PinType.bIsMap=False,PinType.bIsSet=False,PinType.bIsArray=False,PinType.bIsReference=False,PinType.bIsConst=False,PinType.bIsWeakPointer=False,AutogeneratedDefaultValue="0, 0, 0",LinkedTo=(K2Node_CallFunction_147 CADAD3724E0A28E00552BCA72427B545,),PersistentGuid=00000000000000000000000000000000,bHidden=False,bNotConnectable=False,bDefaultValueIsReadOnly=False,bDefaultValueIsIgnored=False,bAdvancedView=False,)
+	End Object
+[/EXCERPT:snippet]
+-->
\ No newline at end of file
diff --git a/Engine/Documentation/Source/Modding/RoboRecall/HowTo/SetupGameplay/PlaceActors/PlaceGameplayActors.INT.udn b/Engine/Documentation/Source/Modding/RoboRecall/HowTo/SetupGameplay/PlaceActors/PlaceGameplayActors.INT.udn
new file mode 100644
--- /dev/null
+++ b/Engine/Documentation/Source/Modding/RoboRecall/HowTo/SetupGameplay/PlaceActors/PlaceGameplayActors.INT.udn
@@ -0,0 +1,37 @@
+Availability:Public
+Title:2. Placing Gameplay Actors
+Crumbs:
+Description:
+Parent:Modding/RoboRecall/HowTo/SetupGameplay
+Order:2
+tags:Modding
+tags:Robo Recall
+type:how-to
+checkpoint:mod_gameplay
+
+[NAV]
+
+1. In the [Modes panel](Engine/UI/LevelEditor/Modes), select the **Place** mode.
+
+	![](place_mode.png)(convert:false)
+
+1. Select the **Basic** tab and drag a **Player Start** Actor into the viewport. This will act as the starting location for the player, so place it where you want the player to begin the level from.
+
+	![](place_playerstart.png)(convert:false)
+
+1. In the **Content Browser**, select the **Content** folder and search for **EnemyStart**. Drag a few instances of this Blueprint into the viewport. These act as locations to spawn the robots from. Place them wehre you want the robot waves to eminate from.
+
+	![](place_enemystart.png)(convert:false)
+
+1. In the **Content Browser**, select the **C++ Classes** folder and search for **OdinAIWaveManager**. Drag an instance of this Actor into the viewport. The location is not important as its job is just to act as a controller for spawning the robots, but you should place it somewhere easy to find and select.
+
+	![](place_wavemanager.png)(convert:false)
+
+
+## Result
+
+You now have all the Actors necessary to create a wave of attacking robots. There are some properties to set and connections to make before they will do anything, so let's do that now!
+
+![](placed_actors.png)(convert:false)
+
+[NAV]
\ No newline at end of file
diff --git a/Engine/Documentation/Source/Modding/RoboRecall/HowTo/SetupGameplay/SettingUpGameplay.INT.udn b/Engine/Documentation/Source/Modding/RoboRecall/HowTo/SetupGameplay/SettingUpGameplay.INT.udn
new file mode 100644
--- /dev/null
+++ b/Engine/Documentation/Source/Modding/RoboRecall/HowTo/SetupGameplay/SettingUpGameplay.INT.udn
@@ -0,0 +1,24 @@
+Availability: Public
+Title:Adding Gameplay To Your Map
+Crumbs:%ROOT%
+Description:Guide that demonstrates adding gameplay elements to a custom map to work with the Robo Recall game
+Parent:Modding/RoboRecall
+type:how-to
+order:11
+tags:Modding
+tags:Robo Recall
+topic-image:MapGameplay_Topic.png
+social-image:MapGameplay.png
+skilllevel:Intermediate
+checkpoint:mod_gameplay
+
+![](MapGameplay.png)
+
+## Steps
+
+[DIR(end:"1" output:"checkpoint" org:"hierarchy" parent:"Modding/RoboRecall/HowTo/SetupGameplay")]
+
+[REGION:call_to_action]
+[Click to Start](Modding/RoboRecall/HowTo/SetupGameplay/CreateMapMod "%Modding/RoboRecall/HowTo/SetupGameplay/CreateMapMod:title%")
+[/REGION]
+
