Copying //UE4/Dev-Documentation to //UE4/Samples-Main (Source: //UE4/Dev-Documentation @ 3372893)

diff --git a/Engine/Documentation/Source/Modding/RoboRecall/HowTo/SetupGameplay/SetupWave/SetupRobotWave.INT.udn b/Engine/Documentation/Source/Modding/RoboRecall/HowTo/SetupGameplay/SetupWave/SetupRobotWave.INT.udn
new file mode 100644
--- /dev/null
+++ b/Engine/Documentation/Source/Modding/RoboRecall/HowTo/SetupGameplay/SetupWave/SetupRobotWave.INT.udn
@@ -0,0 +1,73 @@
+Availability:Public
+Title:3. Setting Up Waves of Robots
+Crumbs:
+Description:
+Parent:Modding/RoboRecall/HowTo/SetupGameplay
+Order:3
+tags:Modding
+tags:Robo Recall
+type:how-to
+checkpoint:mod_gameplay
+
+[NAV]
+
+The Wave Manager controls the generation of waves of attackling robots. You can tell it what types of robots to spawn, where to spawn them from, how many to spawn, etc. Once you set up the properties for the Wave Manager, robots will begin to spawn and attack just like in the game.
+
+1. Select the **OdinAIWaveManager** Actor you previously placed to view its proeprties in the **Details** panel.
+
+1. Enable the **Auto Activate** checkbox to force the waves to automatically begin spawning when the level begins.
+
+	[REGION:note]
+	You can also leave this disabled and enable it during play in response to some event in the game using the [Level Blueprint](Engine/Blueprints/UserGuide/Types/LevelBlueprint) if that better fits the gameplay you are trying to create.
+	[/REGION]
+
+	![](auto_activate.png)
+
+1. The **Possible Spawn Locations** property determines where the robots will spawn from for the wave. Click the **+** button to add elements to the array. You need one for each **Enemy Start** Actor you placed previously.
+
+	![](spawn_locations.png)
+
+1. Click the **Pick Actor from Scene** button for each element in the array and then select one of the **Enemy Start** Actors in the viewport to assign it to the element.
+
+	![](spawn_locations_pick.png)
+
+1. The Wave Manager can control any number of waves of robots. You specify the waves of robots by adding elements to the array. CLick the **+** button to add a new wave.
+
+	![](wave_add.png)
+
+1. You must set the type of wave to see the available properties. Fortunately, there is only one type available currently: Wave. Choose **Wave** from the class picker.
+
+	![](wave_select.png)
+
+1. Enter a name for the wave using the **Wave Name** property. When dealing with multiple waves, giving each wave a name enables you to set up wave-specific behavior in the Level Blueprint.
+
+	![](wave_name.png)
+
+1. Each wave can have any number of groups, which are collections of robots of the same type. Click the **+** button for the **WaveGroups** property to add a new wave group.
+
+	![](wave_group_add.png)
+
+1. As with the wave, you need to set the type of the group, and again there is only one type available: _WaveGroup_. Choose **WaveGroup** from the class picker.
+
+	![](wave_group_select.png)
+
+1. Give the wave group a name using the **Group Name** property. Again, this makes it possible to set up scripted events related to a particular group.
+
+	![](wave_group_name.png)
+
+1. The **Character Class* property determines what type of robots the group contains. There are a variety of types available, but for this example we will use a basic pistol-carrying robot. Choose **BipedChar_Pistol** from the class picker.
+
+	![](wave_group_class.png)
+
+1. Finally, you need to specify how many robots the group will contain total, how many to spawn at once, and when to spawn them. **Pool Size** is the total number of robots the group will spawn. **Max Count** is the number that are spawned at once. **Min Count** determines when to spawn new robots. The Wave Manager will spawn new robots when there are only **Min Count** left. Set **Pool Size**, **Max Count**, and **Min Count** to 8, 4, and 1 respectively. Using these values, the Wave Manager will initially spawn 4 (Max Count) robots. When the player kills 3 of those so there is only 1 (Min COunt) left, the Wave Manager will spawn more until there are 4 (Mac Count) again. Once the Wave Manager has spawned 8 (Pool Size) robots total for this group, no more robots will be spawned for the group.
+
+	![](wave_group_size.png)
+
+
+## Result
+
+Robots with pistols will now spawn and attack the player.
+
+![](robots.png)
+
+[NAV]
\ No newline at end of file
diff --git a/Engine/Documentation/Source/Modding/RoboRecall/Overview/RoboRecallModdingOverview.INT.udn b/Engine/Documentation/Source/Modding/RoboRecall/Overview/RoboRecallModdingOverview.INT.udn
new file mode 100644
--- /dev/null
+++ b/Engine/Documentation/Source/Modding/RoboRecall/Overview/RoboRecallModdingOverview.INT.udn
@@ -0,0 +1,55 @@
+Availability: Epic
+Title:Robo Recall Modding Overview
+Crumbs:%ROOT%
+Description:Overview of the features available to Robo Recall mod creators
+Parent:Modding/RoboRecall
+type:overview
+order:1
+tags:Modding
+tags:Robo Recall
+tags:Getting Started
+topic-image:RoboRecallModding_Overview_Topic.png
+social-image:RoboRecall_Overview_Social.png
+
+![image alt text](Overview.png)(w:934 h:490 convert:false)
+
+[Overview Text Here]
+
+<!-- Note the primary paths to creating mods, why the primary mod paths exist, how stacking works, and finally, it's UE4, you can also just do whatever -->
+
+## Weapons
+
+Tools of the trade for any Robo Ready Recall Specialist. Have an idea for a new lightning gun that chains from bot to bot? How about a classic Katana sword? Flamethrower?
+
+[Weapon Links]
+
+## Maps
+
+The places where all the action takes place! Want to save a small ocean side villa from rampaging bots? Or take the fight to a frozen medieval castle? That's where your custom maps will come into play.
+
+[Map Links]
+
+## Robots
+
+[Content]
+
+[Bot Links]
+
+## Bosses
+
+The boss in Robo Recall has a number of special features that separate it from the other bots, as such he's gets his own sections discussing how he works.
+
+[Boss Links]
+
+## Player
+
+Want to change the hands of your Robo Ready Recall Specialist? Adjust some base stats? Give him/her a full body? Grant them super powers? Creating a custom VR Pawn enables you to do all this and more!
+
+[Player Pawn Links]
+
+## Advanced
+
+The thing to remember when modding Robo Recall is that you have access to the entire Unreal Engine 4 editor, Robo Recall's source code, and most of the art assets used in the game (outside of a few licensed assets). To this end, you should check out the rest of the Unreal Engine Documentation,the Unreal Engine forums, and the Unreal Engine Answer Hub, where you'll find tons of information on making games! And if you get stuck, post a question on the Unreal Engine Forums or in the Answer Hub.
+
+[Doc Link, Forum Link, Answer Hub Link]
+
diff --git a/Engine/Documentation/Source/Modding/RoboRecall/QuickStart/ApplyModifications/ApplyModifications.INT.udn b/Engine/Documentation/Source/Modding/RoboRecall/QuickStart/ApplyModifications/ApplyModifications.INT.udn
new file mode 100644
--- /dev/null
+++ b/Engine/Documentation/Source/Modding/RoboRecall/QuickStart/ApplyModifications/ApplyModifications.INT.udn
@@ -0,0 +1,54 @@
+Availability: Public
+Title:4. Apply the Weapon Material
+Crumbs:%ROOT%
+Description:
+Parent:Modding/RoboRecall/QuickStart
+type:quick start
+order:4
+tags:Modding
+tags:Robo Recall
+topic-image:Modding/RoboRecall/RoboRecall_Topic.png
+social-image:Modding/RoboRecall/RoboRecall_Social.png
+checkpoint:mod_qs
+
+[NAV]
+
+It's time to apply your ceramic Material to your weapons so they look ceramic instead of metallic.
+
+## Steps
+
+[EXCERPT:NewPistol]
+
+1. Double-click one of your weapon Blueprints in the Content Browser to edit it in the [Blueprint Editor](Engine/Blueprints/Editor).
+
+	![](weapon_edit.png)
+
+1. Select the **WeapMesh** Component in the Components panel.
+
+	![](mesh_select.png)
+
+1. Back in the Content Browser, select the _Mat_Ceramic_ Material.
+
+1. In the Details panel, find the **Materials** category and click the **Use Selected Asset from Content BRowser** button to apply the _Mat_Ceramic_ Material to the weapon mesh.
+
+	![](mat_assign.png)
+
+1. Click the **Compile** button in the Blueprint Editor toolbar to update the Blueprint with the changes.
+
+	![](weapon_compile.png)
+
+1. Click the **Save** button in the Blueprint Editor toolbar to save the weapon asset.
+
+	![](weapon_save.png)
+
+1. Repeat these steps for the other weapon Blueprints in your mod to make them all ceramic.
+
+[/EXCERPT]
+
+## Result
+
+The Preview panel shows the ceramic weapon. In the next step, you will test out the mod to see your ceramic weapons in action!
+
+![](weapon_ceramic.png)
+
+[NAV]
\ No newline at end of file
diff --git a/Engine/Documentation/Source/Modding/RoboRecall/QuickStart/CreateAMod/CreateAMod.INT.udn b/Engine/Documentation/Source/Modding/RoboRecall/QuickStart/CreateAMod/CreateAMod.INT.udn
new file mode 100644
--- /dev/null
+++ b/Engine/Documentation/Source/Modding/RoboRecall/QuickStart/CreateAMod/CreateAMod.INT.udn
@@ -0,0 +1,60 @@
+Availability: Public
+Title:2. Create a New Mod
+Crumbs:%ROOT%
+Description:
+Parent:Modding/RoboRecall/QuickStart
+type:quick start
+order:2
+tags:Modding
+tags:Robo Recall
+topic-image:Modding/RoboRecall/RoboRecall_Topic.png
+social-image:Modding/RoboRecall/RoboRecall_Social.png
+checkpoint:mod_qs
+
+[NAV]
+
+In this step, you'll be creating a new mod. While some of this is open ended (Do you want to make a new Robot? A new Gun? Level?), this is the first thing that needs to be done before you start modding Robo Recall.
+
+## Steps
+
+[EXCERPT:NewPistol]
+
+1. Click the **Create Mod** button on the Toolbar.
+
+1. Select the type of mod you wish to create.
+
+	![](select_type.png)
+
+	[REGION:tip]
+	You can multiselect in this dialog to create multiple assets at the same time. Feel free to add any other weapons, bots, etc… by **Ctrl Clicking** the assets you want to create in your mod. 
+	[/REGION]
+
+1. Enter a new mod name, in this case, we've chosen **GrenadeGun**.
+
+	![](mod_name.png)
+
+1. You can also fill out the **Author** and **Description** fields, and this data will show in the Mod menu in the game. If you'd like to change this later, you can do so by accessing your mod through the Editor's **Edit -> Plugins** menu.
+
+	![](mod_info.png)
+
+1. Click the **Create Mod** button to generate your mod files.
+
+1. Wait while the **Robo Recall Editor** creates your new mod structure, a popup tells you when it has completed.
+
+1. Your new mod is automatically focused in the **Content Browser**.
+
+	![](mod_content.png)
+
+1. Double click the **Blueprints** folder, then the **Weapons** folder to find your weapon Blueprints.
+	
+	![](mod_blueprints.png)
+
+[/EXCERPT]
+
+## Result
+
+With your new mod structure setup you can now start creating your mod. Now, just because you started with one mod type doesn't limit you to just modifying the couple of assets you'll find in your mod's Content Folder. You can modify any number of assets from Robo Recall and have them exist in one mod.
+
+Just remember, that any **new** assets you import **need** to go into your mod's Content Folder. If you put them elsewhere, they will fail to package!
+
+[NAV]
\ No newline at end of file
diff --git a/Engine/Documentation/Source/Modding/RoboRecall/QuickStart/InstallModKit/InstallModKit.INT.udn b/Engine/Documentation/Source/Modding/RoboRecall/QuickStart/InstallModKit/InstallModKit.INT.udn
new file mode 100644
--- /dev/null
+++ b/Engine/Documentation/Source/Modding/RoboRecall/QuickStart/InstallModKit/InstallModKit.INT.udn
@@ -0,0 +1,21 @@
+Availability: Public
+Title:1. Install the Robo Recall Mod Kit
+Crumbs:%ROOT%
+Description:
+Parent:Modding/RoboRecall/QuickStart
+type:quick start
+order:1
+tags:Modding
+tags:Robo Recall
+topic-image:Modding/RoboRecall/RoboRecall_Topic.png
+social-image:Modding/RoboRecall/RoboRecall_Social.png
+checkpoint:mod_qs
+
+[NAV]
+
+[INCLUDE:Modding/RoboRecall/HowTo/InstallModKit]
+## Result
+
+With the Robo Recall Editor open the editor will be waiting for you to choose a type of mod to create!
+
+[NAV]
\ No newline at end of file
diff --git a/Engine/Documentation/Source/Modding/RoboRecall/QuickStart/MakeModifications/MakingModifications.INT.udn b/Engine/Documentation/Source/Modding/RoboRecall/QuickStart/MakeModifications/MakingModifications.INT.udn
new file mode 100644
--- /dev/null
+++ b/Engine/Documentation/Source/Modding/RoboRecall/QuickStart/MakeModifications/MakingModifications.INT.udn
@@ -0,0 +1,62 @@
+Availability: Public
+Title:3. Add a New Weapon Material
+Crumbs:%ROOT%
+Description:
+Parent:Modding/RoboRecall/QuickStart
+type:quick start
+order:3
+tags:Modding
+tags:Robo Recall
+topic-image:Modding/RoboRecall/RoboRecall_Topic.png
+social-image:Modding/RoboRecall/RoboRecall_Social.png
+checkpoint:mod_qs
+
+[NAV]
+
+In this step, you'll be creating a new Material asset to apply to your weapons to make them look like they are fabricated from ceramic.
+
+## Steps
+
+[EXCERPT:NewPistol]
+
+1. Right-click in the [Content Browser](Engine/Content/Browser) and choose **New Folder** to create a new folder in your mod's content directory. Name the folder **Materials**.
+
+	![](new_folder.png)
+
+1. Double-click the Materials folder to open it and click the Add New button to add a new Material asset. Name the Material **Mat_Ceramic**.
+
+	![](new_material.png)
+
+1. Double-click the _Mat_Ceramic_ Material to edit it in the [Material Editor](Engine/Rendering/Materials/Editor).
+
+	![](mat_edit.png)
+
+1. Drag a Constant3Vector expression into the graph from the Palette and connect it to the **Base Color** input on the Material node. 
+
+	![](mat_base.png)
+
+1. Double-click the black color preview on the expression to open the [Color Picker](Engine/UI/ColorPicker). Set the R, G, and B values to 0.02 to give the surface just a tiny bit of color.
+
+	![](mat_base_color.png)
+
+1. Drag a Constant expression into the graph from the Palette and connect it to the **Metallic** input on the Material node. Leave the value of the expression at 0. This will make the material have no metallic characteristics in its appearance.
+
+	![](mat_metallic.png)
+
+1. Drag a Constant expression into the graph from the Palette and connect it to the Base Color input on the Material node. Leave the value of the expression at 0. This will cause the material to appear extremely shiny.
+
+	![](mat_rough.png)
+
+1. Click the Apply button to save the changes to the _Mat_Ceramic_ Material.
+
+	![](mat_apply.png)
+
+[/EXCERPT]
+
+## Result
+
+The Preview panel shows a dark, shiny surface that looks like ceramic. In the next step, you will apply this Material to your modded weapons to give them a ceramic appearance.
+
+![](mat_preview.png)
+
+[NAV]
\ No newline at end of file
diff --git a/Engine/Documentation/Source/Modding/RoboRecall/QuickStart/RoboRecallModQuickStart.INT.udn b/Engine/Documentation/Source/Modding/RoboRecall/QuickStart/RoboRecallModQuickStart.INT.udn
new file mode 100644
--- /dev/null
+++ b/Engine/Documentation/Source/Modding/RoboRecall/QuickStart/RoboRecallModQuickStart.INT.udn
@@ -0,0 +1,25 @@
+Availability:Public
+Title: Creating Your First Robo Recall Mod
+Crumbs:
+Description:Getting up and running with your first Robo Recall mod
+Parent:Modding/RoboRecall
+Order:2
+type:quick start
+tags:Modding
+tags:Robo Recall
+tags:Getting Started
+topic-image:RoboRecallModding_QS_Topic.png
+social-image:RoboRecallModding_QS_Social.png
+checkpoint:mod_qs
+
+![](RoboRecallModding_QS_Ceramic.png)
+
+
+
+## Steps
+
+[DIR(output:"checkpoint" end:"1" org:"hierarchy" parent:"Modding/RoboRecall/QuickStart")]
+
+[REGION:call_to_action]
+[Click to Start](Modding/RoboRecall/QuickStart/InstallModKit "%Modding/RoboRecall/QuickStart/InstallModKit:title%")
+[/REGION]
diff --git a/Engine/Documentation/Source/Modding/RoboRecall/QuickStart/TestModifications/TestYourMod.INT.udn b/Engine/Documentation/Source/Modding/RoboRecall/QuickStart/TestModifications/TestYourMod.INT.udn
new file mode 100644
--- /dev/null
+++ b/Engine/Documentation/Source/Modding/RoboRecall/QuickStart/TestModifications/TestYourMod.INT.udn
@@ -0,0 +1,46 @@
+Availability: Public
+Title:5. Test Your Mod
+Crumbs:%ROOT%
+Description:
+Parent:Modding/RoboRecall/QuickStart
+type:quick start
+order:5
+tags:Modding
+tags:Robo Recall
+topic-image:Modding/RoboRecall/RoboRecall_Topic.png
+social-image:Modding/RoboRecall/RoboRecall_Social.png
+checkpoint:mod_qs
+
+[NAV]
+
+Let's make sure our mod is working as intended by enabling it int he game and shooting some targets on the gun range.
+
+## Steps
+
+1. In the Level Editor toolbar, click the **Play** button to launch into the Hub.
+
+	![](mod_play.png)
+
+1. Teleport to the holostation to open up the menu and select **MODS**.
+
+	![](menu_mods.png)
+
+1. In the Mods menu, find the **Ceramic Guns** mod and select it to enable it.
+
+	![](menu_mod_select.png)(h:280 convert:false) ![](menu_mod_enabled.png)(h:280 convert:false)
+
+1. Select **GUN RANGE** from the main menu to do some target practice.
+
+	![](menu_gun_range.png)
+
+1. Your ceramic guns are available to be equipped.
+
+	![](gun_select.png)
+
+## Result
+
+Shoot some targets with your new ceramic weapons!
+
+![](gun_ceramic.png)
+
+[NAV]
\ No newline at end of file
diff --git a/Engine/Documentation/Source/Modding/RoboRecall/Reference/Weapons/Gun/GunReference.INT.udn b/Engine/Documentation/Source/Modding/RoboRecall/Reference/Weapons/Gun/GunReference.INT.udn
new file mode 100644
--- /dev/null
+++ b/Engine/Documentation/Source/Modding/RoboRecall/Reference/Weapons/Gun/GunReference.INT.udn
@@ -0,0 +1,145 @@
+Availability: Public
+Title:Base Gun Reference
+Crumbs:%ROOT%
+Description:
+Parent:Modding/RoboRecall
+type:reference
+order:1
+tags:Modding
+tags:Robo Recall
+tags:Weapons
+topic-image:RoboRecallModding_Guns_Ref_Topic.png
+social-image:RoboRecallModding_Guns_Ref.png
+
+![](RoboRecallModding_Guns_Ref.png)
+
+[TOC]
+
+## Events
+
+### Event Notify Start Firing
+
+[REGION:graph]
+    [OBJECT:BlueprintNode]
+        [PARAMLITERAL:type]
+            event
+        [/PARAMLITERAL]
+        [PARAMLITERAL:title]
+            Event Notify Start Firing
+        [/PARAMLITERAL]
+        [PARAM:outputs]
+            [OBJECT:BlueprintPin]
+                [PARAMLITERAL:type]
+                    exec
+                [/PARAMLITERAL]
+                [PARAMLITERAL:id]
+                    Return_Value
+                [/PARAMLITERAL]
+                [PARAMLITERAL:title]
+                [/PARAMLITERAL]
+                [PARAMLITERAL:value]
+                [/PARAMLITERAL]
+            [/OBJECT]
+        [/PARAM]
+    [/OBJECT]
+[/REGION]
+
+The Notify Start Firing event is called each time the player presses the 'fire' button. Any nodes connected to the _exec_ output pin will be executed at this time. Use this event to extend or override the firing functionality of the weapon, such as changing the firing location or spawning firing effects.
+
+[OBJECT:Section]
+    [PARAMLITERAL:id]
+    outputs
+    [/PARAMLITERAL]
+    [PARAM:heading]
+    Outputs
+    [/PARAM]
+    [PARAM:content]
+        [OBJECT:MemberIconListHeadBlank]
+        [/OBJECT]
+        [OBJECT:ActionPinListItem]
+            [PARAM:icons]
+                [REGION:input_exec exec]
+                    
+                [/REGION]
+            [/PARAM]
+            [PARAMLITERAL:name]
+                
+            [/PARAMLITERAL]
+            [PARAMLITERAL:type]
+                exec
+            [/PARAMLITERAL]
+            [PARAM:tooltip]
+                Nodes connected to this are exected when the event is called.
+            [/PARAM]
+            [PARAMLITERAL:id]
+                Out_Exec
+            [/PARAMLITERAL]
+        [/OBJECT]
+        [OBJECT:MemberIconListTail]
+        [/OBJECT]
+    [/PARAM]
+[/OBJECT]
+
+### Event Notify Stop Firing
+
+[REGION:graph]
+    [OBJECT:BlueprintNode]
+        [PARAMLITERAL:type]
+            event
+        [/PARAMLITERAL]
+        [PARAMLITERAL:title]
+            Event Notify Stop Firing
+        [/PARAMLITERAL]
+        [PARAM:outputs]
+            [OBJECT:BlueprintPin]
+                [PARAMLITERAL:type]
+                    exec
+                [/PARAMLITERAL]
+                [PARAMLITERAL:id]
+                    Return_Value
+                [/PARAMLITERAL]
+                [PARAMLITERAL:title]
+                [/PARAMLITERAL]
+                [PARAMLITERAL:value]
+                [/PARAMLITERAL]
+            [/OBJECT]
+        [/PARAM]
+    [/OBJECT]
+[/REGION]
+
+The Notify Stop Firing event is called each time the player releases the 'fire' button. Any nodes connected to the _exec_ output pin will be executed at this time. Use this event to extend or override the 'stop firing' functionality of the weapon, such as cleaning up firing effects spawned when firing.
+
+[OBJECT:Section]
+    [PARAMLITERAL:id]
+    outputs
+    [/PARAMLITERAL]
+    [PARAM:heading]
+    Outputs
+    [/PARAM]
+    [PARAM:content]
+        [OBJECT:MemberIconListHeadBlank]
+        [/OBJECT]
+        [OBJECT:ActionPinListItem]
+            [PARAM:icons]
+                [REGION:input_exec exec]
+                    
+                [/REGION]
+            [/PARAM]
+            [PARAMLITERAL:name]
+                
+            [/PARAMLITERAL]
+            [PARAMLITERAL:type]
+                exec
+            [/PARAMLITERAL]
+            [PARAM:tooltip]
+                Nodes connected to this are exected when the event is called.
+            [/PARAM]
+            [PARAMLITERAL:id]
+                Out_Exec
+            [/PARAMLITERAL]
+        [/OBJECT]
+        [OBJECT:MemberIconListTail]
+        [/OBJECT]
+    [/PARAM]
+[/OBJECT]
+
diff --git a/Engine/Documentation/Source/Modding/RoboRecall/Reference/Weapons/WeaponReference.INT.udn b/Engine/Documentation/Source/Modding/RoboRecall/Reference/Weapons/WeaponReference.INT.udn
new file mode 100644
--- /dev/null
+++ b/Engine/Documentation/Source/Modding/RoboRecall/Reference/Weapons/WeaponReference.INT.udn
@@ -0,0 +1,38 @@
+Availability: Docs
+Title:Base Weapon Reference
+Crumbs:%ROOT%
+Description:
+Parent:Modding/RoboRecall
+type:reference
+order:1
+tags:Modding
+tags:Robo Recall
+tags:Weapons
+
+## Properties
+
+[OBJECT:VariableListHead]
+[/OBJECT]
+[OBJECT:VariableListItem]
+    [PARAM:icons]
+    [/PARAM]
+    [PARAM:type]
+        TArray&lt;TSubclassOf&lt;AOdinWeapon&gt;&gt;
+    [/PARAM]
+    [PARAM:name]
+        Class Overrides
+    [/PARAM]
+    [PARAM:link]
+	#
+    [/PARAM]
+    [PARAM:description]
+        List of classes that this will override when activated as a mod.
+    [/PARAM]
+[/OBJECT]
+
+## Events
+
+
+
+## Functions
+
diff --git a/Engine/Documentation/Source/Modding/RoboRecall/RoboRecall.INT.udn b/Engine/Documentation/Source/Modding/RoboRecall/RoboRecall.INT.udn
new file mode 100644
--- /dev/null
+++ b/Engine/Documentation/Source/Modding/RoboRecall/RoboRecall.INT.udn
@@ -0,0 +1,48 @@
+Availability: Public
+Title:Robo Recall Modding
+Crumbs:%ROOT%
+Description:Guides to modding Epic Games' Robo Recall for Oculus Touch.
+Parent:Modding
+type:landing
+order:1
+tags:Modding
+tags:Robo Recall
+topic-image:RoboRecall_Topic.png
+social-image:RoboRecall_Social.png
+
+![](RoboRecall_Modding.png)
+
+## Starting Out
+
+[DIR(end:"1" output:"topic" org:"hierarchy" parent:"Modding/RoboRecall" tags:"Getting Started")]
+
+## Guides
+
+[DIR(end:"1" output:"topic" org:"hierarchy" parent:"Modding/RoboRecall" tags:"!Getting Started" type:"!reference")]
+
+## Reference
+
+[DIR(end:"1" output:"topic" org:"hierarchy" parent:"Modding/RoboRecall" tags:"!Getting Started" type:"reference")]
+
+<!--
+
+Should list:
+
+* How-To Landing
+
+* Robo Recall Modding Overview
+
+* Robo Recall Art Considerations
+
+* Robo Recall Modding Limitations
+
+* Robo Recall Advanced Modding
+
+* Robo Recall Getting Started
+
+* Learning Resources Mods
+
+* Robo Recall Reference
+
+-->
+
diff --git a/Engine/Documentation/Source/Modding/RoboRecall/Troubleshooting/TroubleshootingMods.INT.udn b/Engine/Documentation/Source/Modding/RoboRecall/Troubleshooting/TroubleshootingMods.INT.udn
new file mode 100644
--- /dev/null
+++ b/Engine/Documentation/Source/Modding/RoboRecall/Troubleshooting/TroubleshootingMods.INT.udn
@@ -0,0 +1,23 @@
+Availability:Public
+Title: Troubleshooting Robo Recall Mods
+Crumbs:
+Description:Disagnosing and solving issues encountered when modding Robo Recall
+Parent:Modding/RoboRecall
+Order:100
+type:overview
+tags:Modding
+tags:Robo Recall
+topic-image:RoboRecallBotsplosion_Topic.png
+social-image:RoboRecallBotsplosion_Social.png
+
+![](RoboRecallBotsplosion.png)
+
+[TOC]
+
+## How To Uninstall a Mod That Crashes
+
+1. Navigate to your Robo Recall install directory and open the **Plugins** folder.
+
+1. Locate the folder for the mod that crashes and delete it.
+
+1. Run the game and go to **MODS** in the Holostation menu. The mod should no longer appear in the list of installed mods.
diff --git a/Engine/Documentation/Source/Platforms/Android/Android.INT.udn b/Engine/Documentation/Source/Platforms/Android/Android.INT.udn
--- a/Engine/Documentation/Source/Platforms/Android/Android.INT.udn
+++ b/Engine/Documentation/Source/Platforms/Android/Android.INT.udn
@@ -10,46 +10,16 @@ tags:Mobile
 parent:Platforms/Mobile
 type:landing
 
-[OBJECT:TopicButtonList]
-	[PARAM:icon]
-		![](%ROOT%/start_icon.png)(convert:false)
-	[/PARAM]
-	[PARAM:title]
-		Starting Out
-	[/PARAM]
-	[PARAM:description]
-	[/PARAM]
-	[PARAM:links]
-		[DIR(output:"topiccompact" parent:"Platforms/Android" type:"quick start")]
-	[/PARAM]
-[/OBJECT]
+## Starting Out
 
-[OBJECT:TopicButtonList]
-	[PARAM:icon]
-		![](%ROOT%/platform_icon.png)(convert:false)
-	[/PARAM]
-	[PARAM:title]
-		Android How-Tos
-	[/PARAM]
-	[PARAM:description]
-	[/PARAM]
-	[PARAM:links]
-		[DIR(end:"1" output:"listbutton" type:"how-to" parent:"Platforms/Android")]
-	[/PARAM]
-[/OBJECT]
+[DIR(output:"topiccompact" parent:"Platforms/Android" type:"quick start")]
 
-[OBJECT:TopicButtonList]
-	[PARAM:icon]
-		![](%ROOT%/platform_icon.png)(convert:false)
-	[/PARAM]
-	[PARAM:title]
-		General Android Development
-	[/PARAM]
-	[PARAM:description]
-	[/PARAM]
-	[PARAM:links]
-		[DIR(end:"1" output:"listbutton" type:"overview" parent:"Platforms/Android")]
-		[DIR(end:"1" output:"listbutton" type:"reference" parent:"Platforms/Android")]
-	[/PARAM]
-[/OBJECT]
+## Android How-Tos
+	
+[DIR(end:"1" output:"listbutton" type:"how-to" parent:"Platforms/Android")]
+
+## General Android Development
+
+[DIR(end:"1" output:"listbutton" type:"overview" parent:"Platforms/Android")]
+[DIR(end:"1" output:"listbutton" type:"reference" parent:"Platforms/Android")]
 
diff --git a/Engine/Documentation/Source/Platforms/Android/GettingStarted/1/AndroidQS_1.INT.udn b/Engine/Documentation/Source/Platforms/Android/GettingStarted/1/AndroidQS_1.INT.udn
--- a/Engine/Documentation/Source/Platforms/Android/GettingStarted/1/AndroidQS_1.INT.udn
+++ b/Engine/Documentation/Source/Platforms/Android/GettingStarted/1/AndroidQS_1.INT.udn
@@ -19,7 +19,7 @@ ensure that you have all the required programs and code needed to deploy to Andr
 In the following section, we will cover how you can get CodeWorks for Android 1R4 installed on your development PC.
 
 
-%Globals:OSSelect%
+%Globals:osselect_linux%
 
 [OBJECT:ToggleButtonContent]
 	[PARAMLITERAL:category]
@@ -50,6 +50,20 @@ In the following section, we will cover how you can get CodeWorks for Android 1R
 	[/PARAM]
 [/OBJECT]
 
+[OBJECT:ToggleButtonContent]
+	[PARAMLITERAL:category]
+	OS
+	[/PARAMLITERAL]
+	[PARAMLITERAL:id]
+	linux
+	[/PARAMLITERAL]
+	[PARAMLITERAL:active]
+	[/PARAMLITERAL]
+	[PARAM:content]
+	[INCLUDE:#LinuxSetup]
+	[/PARAM]
+[/OBJECT]
+
 
 <!--
 [EXCERPT:WindowsSetup]
@@ -114,7 +128,7 @@ It is highly recommended that you install CodeWorks for Android 1R4 while the ed
 
 1.	A new window will be displayed that will show you the download and install progress of the required CodeWorks for Android 1R4 tools. 
 	[region:note]
-	You will need an internet connection to complete this portion. Depending on the speed of your internet connection, this could take a few minute to a few hours to complete.
+	You will need an Internet connection to complete this portion. Depending on the speed of your Internet connection, this could take a few minute to a few hours to complete.
 	[/region]
 	
 	![](AndroidWorks_06.png)
@@ -206,7 +220,7 @@ Use the links provide below to download and install the correct versions of Java
 
 1. The CodeWorks for Android 1R4 Component Manager will now begin to download and install the required tools. 
 	[region:note]
-	You will need an internet connection to complete this portion. Depending on the speed of your internet connection, this could take a few minute to a few hours to complete.
+	You will need an Internet connection to complete this portion. Depending on the speed of your Internet connection, this could take a few minute to a few hours to complete.
 	[/region]
 	![](AndroidWorks_Mac_07.png)
 
@@ -223,6 +237,12 @@ Use the links provide below to download and install the correct versions of Java
 	[/region]
 
 [/EXCERPT:MacOSSetup]
+
+[EXCERPT:LinuxSetup]
+We're currently in the process of writing this documentation. Please refer to our [New and Updated Resources](https://docs.unrealengine.com/latest/INT/Updates/index.html) page to see 
+which documentation is being published at this time.
+<!--[INCLUDE:Platforms/Linux/IntermediateLinuxDeveloper/AndroidOnLinux/1#AndroidLinux_1]-->
+[/EXCERPT:LinuxSetup]
 -->
 
 ## End Result
diff --git a/Engine/Documentation/Source/Platforms/Android/GettingStarted/2/AndroidQS_2.INT.udn b/Engine/Documentation/Source/Platforms/Android/GettingStarted/2/AndroidQS_2.INT.udn
--- a/Engine/Documentation/Source/Platforms/Android/GettingStarted/2/AndroidQS_2.INT.udn
+++ b/Engine/Documentation/Source/Platforms/Android/GettingStarted/2/AndroidQS_2.INT.udn
@@ -11,13 +11,13 @@ parent:Platforms/Android/GettingStarted
 order:2
 Checkpoint:AndroidQS
 
-Now that Android Works has been installed it is now time to setup your Android device so that it can be used in development. In the following section we will go over all you need to know to setup your Android device so that it can be used to run your UE4 project.
+Now that Android Works has been installed it is now time to set up your Android device so that it can be used in development. In the following section we will go over all you need to know to set up your Android device so that it can be used to run your UE4 project.
 
 [nav]
 
 ## Steps
 
-%Globals:OSSelect%
+%Globals:osselect_linux%
 
 [OBJECT:ToggleButtonContent]
 	[PARAMLITERAL:category]
@@ -48,6 +48,19 @@ Now that Android Works has been installed it is now time to setup your Android d
 	[/PARAM]
 [/OBJECT]
 
+[OBJECT:ToggleButtonContent]
+	[PARAMLITERAL:category]
+	OS
+	[/PARAMLITERAL]
+	[PARAMLITERAL:id]
+	linux
+	[/PARAMLITERAL]
+	[PARAMLITERAL:active]
+	[/PARAMLITERAL]
+	[PARAM:content]
+	[INCLUDE:#LinuxSetup]
+	[/PARAM]
+[/OBJECT]
 
 <!--
 [EXCERPT:WindowsSetup]
@@ -194,6 +207,11 @@ If you do not see your device after you input the adb devices command, here are
 
 [/EXCERPT:MacOSSetup]
 
+[EXCERPT:LinuxSetup]
+We're currently in the process of writing this documentation. Please refer to our [New and Updated Resources](https://docs.unrealengine.com/latest/INT/Updates/index.html) page to see 
+which documentation is being published at this time.
+<!--[INCLUDE:Platforms/Linux/IntermediateLinuxDeveloper/AndroidOnLinux/2#AndroidLinux_2]-->
+[/EXCERPT:LinuxSetup]
 -->
 
 ## End Result
diff --git a/Engine/Documentation/Source/Platforms/Android/GettingStarted/4/AndroidQS_4.INT.udn b/Engine/Documentation/Source/Platforms/Android/GettingStarted/4/AndroidQS_4.INT.udn
--- a/Engine/Documentation/Source/Platforms/Android/GettingStarted/4/AndroidQS_4.INT.udn
+++ b/Engine/Documentation/Source/Platforms/Android/GettingStarted/4/AndroidQS_4.INT.udn
@@ -17,7 +17,7 @@ In the following section we will be taking a look at how you can deploy a single
 [nav]
 
 [region:note]
-For this section to work you will need to make sure that your Android device is plugged into your development PC via USB. If you have not done this already please do this before going any further.
+For this section to work you will need to make sure that your Android device is plugged into your computer via USB. If you have not done this already please do this before going any further.
 [/region]
 
 To deploy your level to your Android device you need to determine which type of Texture format your Android Device supports.
diff --git a/Engine/Documentation/Source/Platforms/Android/GettingStarted/5/AndroidQS_5.INT.udn b/Engine/Documentation/Source/Platforms/Android/GettingStarted/5/AndroidQS_5.INT.udn
--- a/Engine/Documentation/Source/Platforms/Android/GettingStarted/5/AndroidQS_5.INT.udn
+++ b/Engine/Documentation/Source/Platforms/Android/GettingStarted/5/AndroidQS_5.INT.udn
@@ -18,8 +18,7 @@ In the following section we will take a look at how you can package your finishe
 
 ## Steps
 
-
-%Globals:OSSelect%
+%Globals:osselect_linux%
 
 [OBJECT:ToggleButtonContent]
 	[PARAMLITERAL:category]
@@ -50,7 +49,19 @@ In the following section we will take a look at how you can package your finishe
 	[/PARAM]
 [/OBJECT]
 
-
+[OBJECT:ToggleButtonContent]
+	[PARAMLITERAL:category]
+	OS
+	[/PARAMLITERAL]
+	[PARAMLITERAL:id]
+	linux
+	[/PARAMLITERAL]
+	[PARAMLITERAL:active]
+	[/PARAMLITERAL]
+	[PARAM:content]
+	[INCLUDE:#LinuxSetup]
+	[/PARAM]
+[/OBJECT]
 
 <!--
 [EXCERPT:WindowsSetup]
@@ -141,6 +152,12 @@ We are currently packaging in **Development** because we are testing this game,
 	![](Game_On_Phone.png)(w:320)
 
 [/EXCERPT:MacOSSetup]
+
+[EXCERPT:LinuxSetup]
+We're currently in the process of writing this documentation. Please refer to our [New and Updated Resources](https://docs.unrealengine.com/latest/INT/Updates/index.html) page to see 
+which documentation is being published at this time.
+<!--[INCLUDE:Platforms/Linux/IntermediateLinuxDeveloper/AndroidOnLinux/5#AndroidLinux_5]-->
+[/EXCERPT:LinuxSetup]
 -->
 
 ## End Result
diff --git a/Engine/Documentation/Source/Platforms/Android/GettingStarted/6/AndroidQS_6.INT.udn b/Engine/Documentation/Source/Platforms/Android/GettingStarted/6/AndroidQS_6.INT.udn
--- a/Engine/Documentation/Source/Platforms/Android/GettingStarted/6/AndroidQS_6.INT.udn
+++ b/Engine/Documentation/Source/Platforms/Android/GettingStarted/6/AndroidQS_6.INT.udn
@@ -21,14 +21,14 @@ try and accomplish the following using what you have learned from this Quick Sta
 
 * Create a new project using one of the default templates and then try and get that project running on your Android device.
 * Try packaging your project using different Android Texture formats and see which ones will work for your Android device.
-* Check the size of your packaged and then using the [Reducing APK Package size](Platforms/Android/ReducingAPKSize) document and try reducing the size of your packaged game.
-* Check out the [Gear VR Debugging](Platforms/GearVR/Debugging) document and then try and deploy your project to your Android device over a wireless network.
+* Check the size of your package, and then, using the [Reducing APK Package size](Platforms/Android/ReducingAPKSize) document, try reducing the size of your packaged game.
+* Check out the [Gear VR Debugging](Platforms/GearVR/Debugging) document and then try to deploy your project to an Android device over a wireless network.
 
 
 The following selection of links can help you with everything from which Android devices will work with UE4 to things you should consider when creating content that is aimed at mobile devices. 
-If you have not gone through these links before it is highly recommended that you do as they contain a plethora of very valuable information.
+If you have not gone through these links before it is highly recommended that you do as they contain a plethora of valuable information.
 
-For more information about device compatibility and performance checkout the following links:
+For more information about device compatibility and performance, checkout the following links:
 
 * [](Platforms\Android\DeviceCompatibility)
 * [](Platforms\DeviceProfiles)
diff --git a/Engine/Documentation/Source/Platforms/Android/GettingStarted/AndroidQuickStart.INT.udn b/Engine/Documentation/Source/Platforms/Android/GettingStarted/AndroidQuickStart.INT.udn
--- a/Engine/Documentation/Source/Platforms/Android/GettingStarted/AndroidQuickStart.INT.udn
+++ b/Engine/Documentation/Source/Platforms/Android/GettingStarted/AndroidQuickStart.INT.udn
@@ -48,26 +48,31 @@ Checkpoint:AndroidQS
 
 ![](AndroidTopic.jpg)
 
-The Android Quick Start will walk you through setting up your computer and your Android device for deployment, as well as putting a template game on your device using both the launching and packaging features.
+Unreal Engine 4 (UE4) uses a special version of the Android Software Development Kit (SDK), called **CodeWorks for Android 1R4**. 
+This version of CodeWorks for Android is the easiest way to set up the SDK and Native Development Kit (NDK) needed to develop 
+Android projects with UE4. Although we're currently using version 1R4, we intend to use subsequent releases of CodeWorks. 
+If you have other versions of the Android SDK installed (or, older versions of CodeWorks for Android), we recommend that you 
+uninstall them and install CodeWorks for Android, using the CodeWorks for Android installer that's being distributed with UE4. 
+
+[region:note]
+We've bundled the CodeWorks for Android installer with UE4 to make sure that your Android SDK will be compatible with Unreal Engine. 
+[/region]
 
 ## Goals
 
-The focus of the Android Quick start is to walk you through the basics of setting up your PC or Mac for Android Development with UE4.
+The focus of this Android Quick Start guide is to walk you through the basics of setting up your computer for Android Development with UE4.
+This guide also teaches you how to put a template game on your Android device with the Unreal Editor's launching and packaging features.
 
 ## Objectives
 
-* Install and setup the Android SDK.
+* Install and set up the Android SDK.
 * Setting up your Android device for development.
-* Creating a project that is targeted towards Android mobile development.
+* Creating a project that is targeted for Android mobile development.
 * How to deploy a level to your Android device
 * How to package your game for release.
 
 [DIR(output:"checkpoint" parent:"Platforms/Android/GettingStarted")]
 
-[region:note]
-Unreal Engine 4 (UE4) now uses a special version of the Android SDK called **CodeWorks for Android 1R4**. This version of CodeWorks for Android was designed specifically for use with UE4 Android based projects. If you have other versions of the Android SDK installed, it is **highly recommend** that you uninstall them and then re-install CodeWorks for Android using the CodeWorks for Android 1R4 installer that is bundled with UE4. This way you can ensure that the Android SDK you are using will be compatible with UE4.
-[/region]
-
 [REGION:call_to_action]
  [Click to Start](Platforms/Android/GettingStarted/1 "%Platforms/Android/GettingStarted/1:title%")
 [/REGION]
diff --git a/Engine/Documentation/Source/Platforms/Android/Reference/AndroidGettingStarted.INT.udn b/Engine/Documentation/Source/Platforms/Android/Reference/AndroidGettingStarted.INT.udn
--- a/Engine/Documentation/Source/Platforms/Android/Reference/AndroidGettingStarted.INT.udn
+++ b/Engine/Documentation/Source/Platforms/Android/Reference/AndroidGettingStarted.INT.udn
@@ -343,12 +343,12 @@ Not all Android devices are made the same. In particular, there are 4 different
 [EXCERPT:AndroidFormats]
 |Format | Description|
 |-------|------------|
-|ETC1 |Supported by all Android based devices but cannot compress alpha textures (they are stored uncompressed).  Recommend using an RGB and a separate alpha texture if need alpha to get better compression.|
-|ETC2 |Supported by all OpenGL 3.x class devices and supports alpha compression.|
-|ATC |Supported by Qualcomm Adreno GPUs and supports alpha compression.|
-|DXT |Supported by Nvidia Tegra GPUs and supports alpha compression.|
-|PVRTC |supported by PowerVR GPUs and supports alpha compression.|
-|ASTC  |Latest Texture compression format allowing more quality control by specifying block size and supports alpha compression. Available on some devices at this point and will be required for Vulkan Level 1.
+|**ETC1** |Supported by all Android based devices but cannot compress alpha textures (they are stored uncompressed).  Recommend using an RGB and a separate alpha texture if need alpha to get better compression.|
+|**ETC2** |Supported by all OpenGL 3.x class devices and supports alpha compression.|
+|**ATC** |Supported by Qualcomm Adreno GPUs and supports alpha compression.|
+|**DXT** |Supported by Nvidia Tegra GPUs and supports alpha compression.|
+|**PVRTC** |supported by PowerVR GPUs and supports alpha compression.|
+|**ASTC**  |Latest Texture compression format allowing more quality control by specifying block size and supports alpha compression. Available on some devices at this point and will be required for Vulkan Level 1.
 |
 [/EXCERPT:AndroidFormats]
 
diff --git a/Engine/Documentation/Source/Platforms/GoogleVR/Daydream/Daydream.INT.udn b/Engine/Documentation/Source/Platforms/GoogleVR/Daydream/Daydream.INT.udn
--- a/Engine/Documentation/Source/Platforms/GoogleVR/Daydream/Daydream.INT.udn
+++ b/Engine/Documentation/Source/Platforms/GoogleVR/Daydream/Daydream.INT.udn
@@ -11,86 +11,118 @@ tags:VR
 tags:Google VR
 tags:Daydream
 
+[TOC(start:2)]
+
+[region:note]
+Google has a provided a very detailed document that goes over everything you need to know to get the most out of Daydream. You can find the documents by going to the [Offical Google Daydream Unreal](https://developers.google.com/vr/unreal/) help site.
+[/region]
+
 
-## Overview
 While Google VR can work for almost any Android based device, there are certain features that will only work with Daydream devices. In the following section, we will take a look at these features and what you have to do to use them in your projects.
 
-### Android NDK 11c
-In Order to access the full functionality, Daydream offers you need to download, install Android NDK version 11c. You can download the appropriate version of NDK 11c using the following link.
+## Downloading Google VR UE4 Branch
+To get all of the pre-built functionality Google Daydream VR offers you will need to download and compile a special version of UE4 that you will find linked below. This version of UE4 contains the needed code and assets that are required to display the Daydream Motion Controller in the manner Google wants it. 
+
+* [Google VR Github Pull]([REDACTED_PROFILE_URL]/UnrealEngine/tree/4.13-googlevr)
 
-* [https://developer.android.com/ndk/downloads/index.html](https://developer.android.com/ndk/downloads/index.html)
+Once you have downloaded this version of UE4, follow the directions for the [Downloading Unreal as a Zip from Github](GettingStarted/DownloadingUnrealEngine/DownloadZip) page to complete the rest of the setup process.
 
-Once downloaded and installed you will need to point UE4 to the new NDK files by going to the **Android SDK** > **SDKConfig** section and changing the NDK path to the new NDK 11c install location. 
+## Android NDK 11c
+In order to access the full functionality of Daydream, you will need to download and install the Android NDK version 11c. 
+
+* [NDK 11c Downlaod](https://developer.android.com/ndk/downloads/index.html)
+
+Once downloaded and installed you will need to point UE4 to the new NDK files by going to the **Android SDK** > **SDKConfig** section and changing the NDK path to the new NDK 11c install location and set the SDK API level to **matchndk** and the NDK API Level to **android-24**.
 
 ![](GDaydream_New_NDK_00.png)
 
-### Scanline Racing
-Scanline Racing is a feature specific to Daydream devices that when enabled helps to reduce latency in your VR project. To enable Scanline Racing in your Google VR project for Daydream devices you will need to do the following.
+## Scanline Racing
+**Scanline Racing** is a feature specific to Daydream devices that when enabled helps to reduce latency in your VR project. To enable Scanline Racing in your Google VR project for Daydream devices you will need to do the following.
 
 1. Open up your **Project Settings** and go to **Platforms** > **Android**.
 
-	![](GDaydream_Arm64_Support_00.png)
+    ![](GDaydream_Arm64_Support_00.png)
 
-1. Under the **Build** section enable **Support arm64[aka arm64-v8a]** by clicking on the checkmark box next to its name.
+1. Under the **Build** section enable **Support arm64[aka arm64-v8a]** by clicking on the checkmark box next to its name and disable **Support arm7[aka armebi-v7a]**.
 
-	![](Platforms\GoogleVR\QuickStart\2\GVRQS_Build_arm64_Support_00.png)
+    ![](Platforms\GoogleVR\QuickStart\2\GVRQS_Build_arm64_Support_00.png)
 
-	[region:note]
+    [region:note]
     Enabling this will add significant time to your first cook as the arm64 libraries need to be compiled.
     [/region]
 
-1. Then under the **Advanced APKPackaging** section and enable the **Configure GoogleVR for low-latency rendering mode (scanline racing)** by clicking on the checkmark box next to its name.
+1. Then under the **Advanced APKPackaging** section and enable the following options.
+	* **Configure for deployment to Daydream**
+	* **Configure GoogleVR Deployment Mode**
+	* **Configure GoogleVR for low-latency rendering mode (scanline racing)**
 
-	![](GDaydream_Scanline_Racing_Support_00.png)
+    ![](GDaydream_Scanline_Racing_Support_00.png)
 
 Then next time you cook your UE4 project and deploy it to your Daydream device low latency rendering will be enabled.
 
-### Motion Controllers
-Working with Motion Controllers for Daydream devices is much like working with Motion Controllers for the HtC Vive or Oculus Rift. To quickly add Motion Controller functionality to your UE4 Daydream project you need to do the following.
+## Motion Controllers
+Working with Motion Controllers for Daydream devices is much like working with Motion Controllers for the HTC Vive or Oculus Rift. To quickly add Motion Controller functionality to your UE4 Daydream project you need to do the following.
 
 [region:note]
-For a more in-depth look at how to use Motion Controllers in UE4, make sure to check out the [Offical Motion Controller Documents](https://docs.unrealengine.com/latest/INT/Platforms/VR/MotionController/) page for more information.
+For a more in-depth look at how to use the Google Daydream Controller Motion Controllers in UE4, make sure to check out the [Offical Google Arm Model Documents](https://developers.google.com/vr/unreal/arm-model) page for more information on how to use the Arm Model.
 [/region]
 
-1. Open up your Player Pawn and search for Motion Controller Component from the **Add Components** list add the following items.
+1. First, open up your project's Plugins menu and under the Input section, enable the **Google VR Motion Controller** plugin.
 
-	* Motion Controller Component
-	* Static Mesh Component with Static Mesh
-	
-	[region:lightbox] 
-	[![](GVRQS_Daydream_MC_Setup_00.png)(w:700)](GVRQS_Daydream_MC_Setup_00.png)
+    ![](GDaydream_Plugins.png)
+
+    [region:note]
+    Please note that to use the Plugin you will need to re-start your UE4 Editor. Failing to do so will result in you not being able to use the Google VR Motion Controller. 
+    [/region]
+
+1. Next, open up your projects Pawn and add the **GoogleVRMotionController** Component to the Components list. When completed your Pawn's setup should match the following image. 
+
+    ![](GDaydream_MC_Addation.png)
+
+    [region:tip]
+    The Motion Controller component should use the same position and height of the VR Camera.
+    [/region]
+
+1. Now when you run your project your motion controllers should look similar to the following image.
+
+	[REGION:lightbox]
+	[![](unreal-controller-model.png)(w:724)](unreal-controller-model.png)
+	[/REGION]
+
+	[region:note]
+	If you are using the 4.14 binary build of UE4 the above content will be missing and you should instead use the standard Motion Controller component. All of this functionality is slated to be in the 4.15 release of UE4. If you need to this functionality now, you will have to download and compile the, [UE4.13 Google VR version]([REDACTED_PROFILE_URL]/UnrealEngine/tree/4.13-googlevr).
 	[/region]
 
-1. Next in the Event Graph use the Motion Contorler Blueprint functions to interact with the buttons on the Motion Controllers.
-	
-	![](Platforms\VR\MotionController\Working_With_MC.png)
+Now that you have a Google VR Motion Controller component on your Pawn, you can edit and adjust how the Motion Controller's function by searching for **Google VR** in the Blueprint Editor and using the nodes that you find. 
 
 ## Training Streams
 
+These Training Streams help to provide some additional resources on setting up your Motion Controllers and creating interactions in VR using Motion Controllers. 
+
 [REGION:topics third]
 [OBJECT:TopicCompactVideo]
-	[PARAMLITERAL:title]
-	Setting Up Motion Controllers
-	[/PARAMLITERAL]
-	[PARAMLITERAL:videoid]
-	6ALnsdQnkVQ
-	[/PARAMLITERAL]
+    [PARAMLITERAL:title]
+    Setting Up Motion Controllers
+    [/PARAMLITERAL]
+    [PARAMLITERAL:videoid]
+    6ALnsdQnkVQ
+    [/PARAMLITERAL]
 [/OBJECT]
 [OBJECT:TopicCompactVideo]
-	[PARAMLITERAL:title]
-	Creating Interactions in VR With Motion Controllers Part 1
-	[/PARAMLITERAL]
-	[PARAMLITERAL:videoid]
-	eRNtgFo6iU0
-	[/PARAMLITERAL]
+    [PARAMLITERAL:title]
+    Creating Interactions in VR With Motion Controllers Part 1
+    [/PARAMLITERAL]
+    [PARAMLITERAL:videoid]
+    eRNtgFo6iU0
+    [/PARAMLITERAL]
 [/OBJECT]
 [OBJECT:TopicCompactVideo]
-	[PARAMLITERAL:title]
-		Creating Interactions in VR With Motion Controllers Part 2
-	[/PARAMLITERAL]
-	[PARAMLITERAL:videoid]
-	utOahIZgKgc
-	[/PARAMLITERAL]
+    [PARAMLITERAL:title]
+        Creating Interactions in VR With Motion Controllers Part 2
+    [/PARAMLITERAL]
+    [PARAMLITERAL:videoid]
+    utOahIZgKgc
+    [/PARAMLITERAL]
 [/OBJECT]
 [/REGION]
 
diff --git a/Engine/Documentation/Source/Platforms/GoogleVR/GoogleVR.INT.udn b/Engine/Documentation/Source/Platforms/GoogleVR/GoogleVR.INT.udn
--- a/Engine/Documentation/Source/Platforms/GoogleVR/GoogleVR.INT.udn
+++ b/Engine/Documentation/Source/Platforms/GoogleVR/GoogleVR.INT.udn
@@ -13,6 +13,10 @@ tags:Platform
 prereq: Platforms/Android/GettingStarted/
 social-image:googlevr_social.png
 
+[region:note]
+Google has a provided a very detailed document that goes over everything you need to know to get the most out of Daydream. You can find the documents by going to the [Offical Google Daydream Unreal](https://developers.google.com/vr/unreal/) help site.
+[/region]
+
 ## Starting Out
 
 The GoogleVR Quick Start will walk you through getting your PC and Android-based smartphone to work with UE4 and Google VR.
@@ -25,3 +29,5 @@ Daydream is a part of Google VR that adds specific support for rendering optimiz
 
 [DIR(output:"listbutton" type:"overview" parent:"Platforms/GoogleVR")]
 
+
+
diff --git a/Engine/Documentation/Source/Platforms/GoogleVR/QuickStart/1/GoogleVRQuickStart_1.INT.udn b/Engine/Documentation/Source/Platforms/GoogleVR/QuickStart/1/GoogleVRQuickStart_1.INT.udn
--- a/Engine/Documentation/Source/Platforms/GoogleVR/QuickStart/1/GoogleVRQuickStart_1.INT.udn
+++ b/Engine/Documentation/Source/Platforms/GoogleVR/QuickStart/1/GoogleVRQuickStart_1.INT.udn
@@ -13,6 +13,10 @@ tags:Google VR
 
 [NAV]
 
+[region:note]
+Google has a provided a very detailed document that goes over everything you need to know to get the most out of Daydream. You can find the documents by going to the [Offical Google Daydream Unreal](https://developers.google.com/vr/unreal/) help site.
+[/region]
+
  In this part of the Google VR Quick Start, we will go over how to create a new UE4 project that has the best settings for Google VR development.
 
 1. From the UE4 Launcher make sure that you have UE4 version **4.12** or later downloaded and installed then press the **Launch** button to open the **Unreal Project Browser**.
diff --git a/Engine/Documentation/Source/Platforms/GoogleVR/QuickStart/2/GoogleVRQuickStart_2.INT.udn b/Engine/Documentation/Source/Platforms/GoogleVR/QuickStart/2/GoogleVRQuickStart_2.INT.udn
--- a/Engine/Documentation/Source/Platforms/GoogleVR/QuickStart/2/GoogleVRQuickStart_2.INT.udn
+++ b/Engine/Documentation/Source/Platforms/GoogleVR/QuickStart/2/GoogleVRQuickStart_2.INT.udn
@@ -13,102 +13,111 @@ tags:Google VR
 
 [NAV]
 
+[region:note]
+Google has a provided a very detailed document that goes over everything you need to know to get the most out of Daydream. You can find the documents by going to the [Offical Google Daydream Unreal](https://developers.google.com/vr/unreal/) help site.
+[/region]
+
  In this part of the Google VR Quick Start, we will go over setting up your UE4 project so that it can be used with Google VR.
 
 1. Once the UE4 Editor has loaded, from the **Main Toolbar**, open the **Edit** tab and select the **Plugins** menu item.
 
-	![](GVRQS_Pligin_Menu_00.png)
+    ![](GVRQS_Pligin_Menu_00.png)
 
 1. Inside of the Plugins menu, go the **Virtual Reality** section and make sure that **Google VR** is enabled.
 
-	![](GVRQS_Pligin_GVR_Enabled_00.png)
+    ![](GVRQS_Pligin_GVR_Enabled_00.png)
 
     [region:note]
     The Google VR plugin should be enabled by default. If it is not enabled, it can be enabled by clicking on the checkmark box next to **Enabled**. When you do this will be prompted to restart the UE4 Editor so that the changes can take place. Failing to re-start the UE4 editor after enabling the Google VR plugin could result in the plugin not working correctly.
     [/region]
 
-1. Inside of the **Content Browser** create a new Material, named **MAT_Floor** and open it up by double-clicking on it.
-
-	![](GVRQS_Floor_Material_00.png)
-
-1. Inside the Material Graph, add a **Constant** Material Expression node with a value of **0.5** and plug it's output into the **Base Color** input and then press the **Apply** button to compile the Material.
-
-	![](GVRQS_Floor_Material_Setup_00.png)
-
-1. Drag the **Mat_Floor** Material from the Content Browser to the **Floor** static mesh that is placed in the level to apply the Material to the Static Mesh.
-
-	[OBJECT:EmbeddedVideo]
-	[PARAMLITERAL:width]
-	640
-	[/PARAMLITERAL]
-	[PARAMLITERAL:height]
-	360
-	[/PARAMLITERAL]
-	[PARAMLITERAL:videoid]
-	k-1bWOSaXU8
-	[/PARAMLITERAL]
-	[/OBJECT]
-
 1. From the Main Toolbar, select the **File** option and then click on the **Save** option to bring up the **Save Level As** window.
 
-	![](GVRQS_Save_Level_00.png)
+    ![](GVRQS_Save_Level_00.png)
 
 1. From the **Save Level As** window under the **Name** section, name the level **EntryLevel** and then press the **Save** button.
 
-	![](GVRQS_Save_Level_01.png)
+    ![](GVRQS_Save_Level_01.png)
 
 1. Back on the **Main Toolbar** select the **Edit** menu option and then select the **Project Settings** menu option.
 
-	![](GVRQS_Project_Settings_00.png)
+    ![](GVRQS_Project_Settings_00.png)
 
 1. From the Project Settings menu go to the **Maps & Modes** section and under the **Default Maps** section, change the **Editor Startup Map** and **Game Default Map** to the EntryLevel map to make sure that the EntryLevel map is the map that will be loaded when the project is run.
 
-	![](GVRQS_Select_Level_00.png)
+    ![](GVRQS_Select_Level_00.png)
 
 1. Back in the Project Settings, under the **Target Hardware** section, make sure that **Target Hardware** class has been set to **Mobile / Tablet** and the graphics level set to **Scalable 3D or 2D**. 
 
-	![](GVRQS_Target_Hardware_00.png)
+    ![](GVRQS_Target_Hardware_00.png)
 
-	[region:note]
-	If you see the **Restart Editor** button under the **Pending Changes** section make sure to press it to re-start the editor so the changes that are made will be applied. Failing to do this could result in your project not working with Google VR.
-	![](GVRQS_Target_Hardware_01.png)
-	[/region]
+    [region:note]
+    If you see the **Restart Editor** button under the **Pending Changes** section make sure to press it to re-start the editor, so the changes that are made will be applied. Failing to do this could result in your project not working with Google VR.
+    ![](GVRQS_Target_Hardware_01.png)
+    [/region]
 
 1. In the **Rendering** section make sure that **Mobile HDR** is disabled as this is necessary for correct stereoscopic rendering.
 
-	![](GVRQS_Moile_HDR_Off_00.png)
+    ![](GVRQS_Moile_HDR_Off_00.png)
 
 1. Under the **Input** section in the **Mobile** category, remove the **DefaultVirtualJoysticks** from the **Default Touch Interface** by clicking on the small white triangle next to the DefaultVirtualJoysticks and selecting the **Clear** option from the drop-down list.
 
-	[OBJECT:EmbeddedVideo]
-	[PARAMLITERAL:width]
-	640
-	[/PARAMLITERAL]
-	[PARAMLITERAL:height]
-	360
-	[/PARAMLITERAL]
-	[PARAMLITERAL:videoid]
-	-Ylbi1TCKZs
-	[/PARAMLITERAL]
-	[/OBJECT]
+    [OBJECT:EmbeddedVideo]
+    [PARAMLITERAL:width]
+    640
+    [/PARAMLITERAL]
+    [PARAMLITERAL:height]
+    360
+    [/PARAMLITERAL]
+    [PARAMLITERAL:videoid]
+    -Ylbi1TCKZs
+    [/PARAMLITERAL]
+    [/OBJECT]
 
 1. In the **Platforms** section under the **Android** section press the **Configure Now** button in the **APKPackaging** and **Google Play Services** sections to ensure that our project can package properly.
 
-	![](GVRQS_Config_Now_00.png)
+    ![](GVRQS_Config_Now_00.png)
+
+1. In the **APKPackaging** section make sure you have the following options set:
+
+     * **Minimum SDK Version** and **Target SDK Version:** 24
+     * **Disable verify OBB on first start/update:** Checked
+     * **Enable FullScreen Immersive on KitKat and above devices:** Checked 
 
-1. Under the **APKPackaging** section change the **Minimum SDK Version** and **Target SDK Version** to **21** and enable the **Enable FullScreen Immersive on KitKat and above devices** option to hide any on-screen buttons that would otherwise be visible when the project is viewed in a VR headset.
+    ![](GVRQS_SDK_Version_00.png)
 
-	![](GVRQS_SDK_Version_00.png)
+1. Locate the **Build** section and uncheck the **Support armv7 [aka armeabi-v7a]** option and make sure that both **Support arm64[aka arm64-v8a]** and **Support OpenGL ES2** are checked.
 
-1. In the **Advanced APKPackaging** section and enable the **Configure GoogleVR for sustained performance mode** option.
+	![](GVRQS_Build_arm64_Support_00.png)
 
-	![](GVRQS_GoogleVR_Options_00.png)
+1. Next, in the **Advanced APKPackaging** section and enable the following options:
 
-	* **Sustained Performance:** Puts the device into sustained performance mode, which throttles performance to a level where the performance will not overheat the device, causing vastly decreased performance.  It is recommended that you enable this for apps that are meant to be run for significant periods of time.
+    * **Configure for deployment to Daydream:** Checked
+    * **Configure GoogleVR Deployment Mode:** Daydream & Cardboard
+    * **Configure GoogleVR for sustained performance mode:** Checked
+
+    ![](GVRQS_GoogleVR_Options_00.png)
+
+    [region:note]
+    The Configure GoogleVR Deployment Mode has several options that when selected will do that following.
+    * **Cardboard:** Configure GoogleVR to run in Cardboard-only mode.
+    * **Daydream:** Configure GoogleVR to run in Daydream-only mode. In this mode, the app won't be able to run on NonDaydream-ready phones.
+    * **Daydream & Cardboard:** Configure GoogleVR to run in Daydream mode on Daydream-ready phones and fallback to Cardboard mode on NonDaydream-ready phones.
+
+    [/region]
+
+1. In the **Android SDK** section under the **SDKConfig** section set both the **SDK API Level** and **NDK API Level** to **android-24** instead of their defaults to ensure that your project will be compiled against the most recent Android SDK.
+
+    ![](GVRQS_Android_SDK_Options_00.png)
+
+    [region:note]
+        For your UE4 Daydream projects to properly deploy to a Daydream device, you need to make sure that CodeWorks has installed the following components.
+        * **Android API 24**
+        * **Android NDK 12b**
+        * **JDK 1.8**
+    [/region]
 
-1. In the **Android SDK** section under the **SDKConfig** section set both the **SDK API Level** and **NDK API Level** to **android-21** instead of their defaults to ensure that your project will be compiled against the most recent Android SDK.
 
-	![](GVRQS_Android_SDK_Options_00.png)
+In the next section of the Google VR Quick Start, we will go over what you have to do to package your UE4 project, so it will run on a Google VR device.
 
-In the next section of the Google VR Quick Start, we will go over what you have to do to package your UE4 project so it will run on a Google VR device.
 [NAV]
\ No newline at end of file
diff --git a/Engine/Documentation/Source/Platforms/GoogleVR/QuickStart/GoogleVRQuickStart.INT.udn b/Engine/Documentation/Source/Platforms/GoogleVR/QuickStart/GoogleVRQuickStart.INT.udn
--- a/Engine/Documentation/Source/Platforms/GoogleVR/QuickStart/GoogleVRQuickStart.INT.udn
+++ b/Engine/Documentation/Source/Platforms/GoogleVR/QuickStart/GoogleVRQuickStart.INT.udn
@@ -28,6 +28,10 @@ The GoogleVR Quick Start will walk you through setting up your UE4 project to wo
 	[/PARAM]
 [/OBJECT]
 
+[region:note]
+Google has a provided a very detailed document that goes over everything you need to know to get the most out of Daydream. You can find the documents by going to the [Offical Google Daydream Unreal](https://developers.google.com/vr/unreal/) help site.
+[/region]
+
 
 
 
diff --git a/Engine/Documentation/Source/Platforms/Linux/AdvancedLinuxDeveloper/LinuxCrossCompileLegacy/LinuxCrossCompileLegacy.INT.udn b/Engine/Documentation/Source/Platforms/Linux/AdvancedLinuxDeveloper/LinuxCrossCompileLegacy/LinuxCrossCompileLegacy.INT.udn
new file mode 100644
--- /dev/null
+++ b/Engine/Documentation/Source/Platforms/Linux/AdvancedLinuxDeveloper/LinuxCrossCompileLegacy/LinuxCrossCompileLegacy.INT.udn
@@ -0,0 +1,151 @@
+Availability:Public
+Title: Cross-Compiling for Linux (Legacy)
+Crumbs:
+Description:For engine versions older than 4.14, this page shows users how to set up cross-compilation for the Linux Platform.
+Version: 4.14
+SkillLevel: Advanced
+type:reference
+parent:Platforms/Linux
+tags:Linux
+related:Programming/Development/CompilingProjects
+related:Programming/Development/BuildConfigurations
+related:Engine/Basics/Projects/Packaging
+related:Platforms/Linux/GettingStarted
+topic-image:LegacyRef_Topic.png
+
+<!--- %Globals:BinarySourceSelectVar% --->
+
+[TOC(start:2 end:3)]
+
+This reference is an archive for users who need to set up their cross-compilation toolchains in Unreal Engine 4 versions that were released prior to 4.14. 
+
+[REGION:note]
+If you're developing your project with Unreal Engine, version 4.14 (or newer), you'll want to refer to our [](Platforms/Linux/GettingStarted) documentation.
+[/REGION]
+
+[EXCERPT:Why_XC]
+## Why Cross-Compilation
+
+**Cross-compilation** makes it possible for game developers, working in a Windows-centric workflow, to target Linux. At this time, cross-compilation is only supported for Windows. 
+Mac users currently have to resort to [native compilation](https://wiki.unrealengine.com/Building_On_Linux). We support, test, and provide the libraries and toolchains for the Linux-x86_64 platform. 
+[/EXCERPT:Why_XC]
+## Getting the Toolchain
+
+Please use the following table to download the appropriate toolchain: 
+
+| Unreal Engine 4 (UE4) Version   | Toolchain                                                                                                                   | 
+| ------------------------------- | --------------------------------------------------------------------------------------------------------------------------- | 
+| **<= 4.8**                      | -v4 [clang-3.5.0-based](http://cdn.unrealengine.com/qfe/v4_clang-3.5.0_ld-2.24_glibc-2.12.2.zip)                            | 
+| **4.9 and 4.10**                | -v6 [clang-3.6.0-based](https://s3.amazonaws.com/unrealengine/qfe/v6_clang-3.6.0_ld-2.24_glibc-2.12.2.zip)                  | 
+| **4.11 thru 4.13**              | -v7 [clang-3.7.0-based](https://s3.amazonaws.com/unrealengine/CrossToolchain_Linux/v7_clang-3.7.0_ld-2.24_glibc-2.12.2.zip) | 
+| **4.14**                        | -v8 [clang-3.9.0-based](http://cdn.unrealengine.com/qfe/v8_clang-3.9.0-centos7.zip)                                         |
+
+We also provide the libraries and toolchains that allow you to compile for Linux ARM as well (original Raspberry Pi and up). Although, this will require you to make a (minor) code 
+change in [Unreal Build Tool](Programming/UnrealBuildSystem) (UBT).
+
+From this point on, we assume that you are targeting x86_64 Linux, though most of the following information will apply to compiling for ARM as well (except for the different toolchain). 
+Note that you can [build your own toolchain](https://wiki.unrealengine.com/Building_Linux_cross-toolchain) if you want different versions of the tools, or if you want to target a different architecture.
+
+## Setting up the Toolchain 
+
+Add an environment variable (`Control Panel->System->Advanced system settings->Advanced->Environment variables`) named `LINUX_ROOT`, the value of which, is the absolute path to your toolchain (without a trailing backslash):
+
+[REGION:raw]
+![](EnvironmentVar.png)
+[/REGION]
+
+After you make sure that the environment variable is set, regenerate UE4 project files (using `GenerateProjectFiles.bat`) and restart Visual Studio. After this, you should have "Linux" available for 
+Win32/Win64 configurations (you should now be able to cross-compile for it).
+
+## Setting up Packaging for Linux
+
+[OBJECT:ToggleButtonContent]
+	[PARAMLITERAL:category]
+	BinarySourceSelect
+	[/PARAMLITERAL]
+	[PARAMLITERAL:id]
+	binaryinstall
+	[/PARAMLITERAL]
+	[PARAMLITERAL:active]
+	active_button_content
+	[/PARAMLITERAL]
+	[PARAM:content]
+	%binary%
+	[/PARAM]
+[/OBJECT]
+[OBJECT:ToggleButtonContent]
+	[PARAMLITERAL:category]
+	BinarySourceSelect
+	[/PARAMLITERAL]
+	[PARAMLITERAL:id]
+	sourceinstall
+	[/PARAMLITERAL]
+	[PARAMLITERAL:active]
+	[/PARAMLITERAL]
+	[PARAM:content]
+	%source%
+	[/PARAM]
+[/OBJECT]
+
+[VAR:source]
+To package for the Linux platform, you will need to build the following targets:
+
+For Linux:
+
+* CrashReportClient
+* UE4Game (if you project is content-only, otherwise just build your project)
+
+For Windows:
+
+* The editor itself (this is needed to get Linux target modules built, so that Unreal Editor and Unreal Frontend know how to cook/package games for Linux).
+* UnrealPak and ShaderCompileWorker (you probably have those built anyway, but we're still mentioning it for completeness).
+
+[/VAR]
+
+[VAR:binary]
+Note that binary release can only package content-only projects. If you want to package a code-based project for Linux (or a project that has non-default third party plugins), you will have to set up a source build. 
+Otherwise, if you are fine with content-only, there's no need to build anything because Linux binaries of UE4Game and CrashReportClient are bundled.
+[/VAR]
+
+## Packaging for Linux
+
+The easiest way to package a project is to open it in the editor and then using `File->Package To->Linux` (assuming that you have a cross-toolchain installed in the previous step and Linux target modules are built for 
+the editor in question - if you don't see Linux in the list, then it's likely that one of the former is not true - see above). After some time (which depends on the project in question and is rather short for a sample project) you 
+will have game assets and binaries in the directory you chose to package to.
+
+The details of the packaging process can be seen by clicking the **Show Output Log** link. If this process fails with the error message, "unable to find UnrealPak or ShaderCompileWorker" see above about building them for the host (Windows) platform.
+
+Secure copy (`scp`) or otherwise copy it to a target machine (mounting a Samba share - if you know how to do that - may be better if the target machine is low on disk space, which also reduces iteration times), change mode (`chmod +x`) the target executable, 
+(which will be located in the `LinuxNoEditor/<ProjectName>/Binaries/Linux/` directory) and run it.
+
+### Changes for Linux-ARM Platforms
+
+[REGION:note]
+With the release of Unreal Engine, version 4.14, there is a new (and easy to follow) way of packaging projects for the ARM platform. 
+Check out our recently updated [](Platforms/Linux/GettingStarted) documentation to learn more.
+[/REGION]
+
+If you are using the Linux-ARM cross compile toolchain before running GenerateProjectFiles.bat below, edit the following file in the UE4 source code: 
+
+	.../UnrealEngine/Engine/Source/Programs/UnrealBuildTool/Linux/UEBuildLinux.cs 
+
+by commenting out the following line: 
+
+	static private string DefaultArchitecture = "x86_64-unknown-linux-gnu"; 
+
+and un-commenting the linux-arm architecture line just below it: 
+
+	//static private string DefaultArchitecture = "arm-unknown-linux-gnueabihf";
+
+An additional step is required for the packaged project to be runnable on Linux-ARM platforms. Assuming the packaged project is located at `\foo\bar\SunTemple_Linux`, open the following location:
+
+`\foo\bar\SunTemple_Linux\LinuxNoEditor\Engine\Binaries\Linux`
+
+and replace `libopenal.so.1` with the version from: 
+
+`...\UnrealEngine\Engine\Source\ThirdParty\OpenAL\1.15.1\lib\Linux\arm-unknown-linux-gnueabihf\libopenal.so`
+
+where `...\UnrealEngine` is the location of the UE4 source code. Please be sure to rename `libopenal.so` to `libopenal.so.1`.
+
+Now the project can be copied over to the target machine, `chmod+x` the target executable and run it as instructed above.
+
diff --git a/Engine/Documentation/Source/Platforms/Linux/BeginnerLinuxDeveloper/SettingUpAnIDE/1/SettingUpAnIDE_1.INT.udn b/Engine/Documentation/Source/Platforms/Linux/BeginnerLinuxDeveloper/SettingUpAnIDE/1/SettingUpAnIDE_1.INT.udn
--- a/Engine/Documentation/Source/Platforms/Linux/BeginnerLinuxDeveloper/SettingUpAnIDE/1/SettingUpAnIDE_1.INT.udn
+++ b/Engine/Documentation/Source/Platforms/Linux/BeginnerLinuxDeveloper/SettingUpAnIDE/1/SettingUpAnIDE_1.INT.udn
@@ -1,4 +1,4 @@
-Availability: Docs
+Availability: Public
 Title: 1 - Configuring your IDE to Build UE4
 Crumbs:
 Description:Learn how to configure your IDE to build UE4 in Debug and Release mode.
@@ -9,34 +9,19 @@ order: 1
 tags:Linux
 SkillLevel:Beginner
 
-[VAR:Steps]
-[OBJECT:Navigation]
-	[PARAM:previous]
-	[/PARAM]
-	[PARAM:current]
-	[/PARAM]
-	[PARAM:home]
-		[Home](Platforms/Linux/BeginnerLinuxDeveloper/SettingUpAnIDE)
-	[/PARAM]
-	[PARAM:next]
-		[Next Step](Platforms/Linux/BeginnerLinuxDeveloper/SettingUpAnIDE/2 "%Platforms/Linux/BeginnerLinuxDeveloper/SettingUpAnIDE/2:title%")
-	[/PARAM]
-[/OBJECT]
-[/VAR]
-
-%Steps%
-
-While some developers prefer to use text editors and command line debuggers to modify their game's source code, most developers 
-prefer to use an Integrated Development Environment (IDE) to edit, debug, and compile their code. If you're a game developer looking to setup 
+[NAV]
+
+While some developers like to use text editors and command line debuggers to modify their game's source code, most developers 
+prefer to use an Integrated Development Environment (IDE) to edit, debug, and compile their code. If you're a game developer looking to set up 
 an IDE to work with Unreal Engine 4 (UE4) in a Linux environment, this guide is for you. 
 
 Although we recognize that chosing an IDE is a matter of individual preference, we've configured UE4 to interface with three IDEs in Linux; 
 specifically, Qt Creator, KDevelop, and Codelite. We don't endorse one IDE over another; however, as an example for learning purposes, we've chosen to 
 show you how to configure Qt Creator to build and run UE4. By the end of this guide, you'll know how to configure your IDE to build and run UE4. Finally, 
-you'll also learn how to setup the IDE to build and run new C++ (CPP) projects that you create from the Unreal Editor.
+you'll also learn how to set up the IDE to build and run new C++ (CPP) projects that you create from the Unreal Editor.
 
 [REGION:note]
-Before we begin, please make sure that you've already setup your workflow using our [](Platforms/Linux/BeginnerLinuxDeveloper/SettingUpAnUnrealWorkflow). Also, please make sure that [Qt Creator](https://wiki.qt.io/Install_Qt_5_on_Ubuntu) is
+Before we begin, please make sure that you've already set up your workflow using our [](Platforms/Linux/BeginnerLinuxDeveloper/SettingUpAnUnrealWorkflow). Also, please make sure that [Qt Creator](https://wiki.qt.io/Install_Qt_5_on_Ubuntu) is
  installed on your machine.
 [/REGION]
 
@@ -129,7 +114,7 @@ Before we begin, please make sure that you've already setup your workflow using
 
 	![](SetupIDE_Step8.png)
 
-1.	Enter `UE4Editor` into the **Make arguments:** text field.
+1.	Enter `UE4Editor-Linux-Debug` into the **Make arguments:** text field.
 
 	![](SetupIDE_Step9.png)
 
@@ -137,7 +122,7 @@ Before we begin, please make sure that you've already setup your workflow using
 
 	![](SetupIDE_Step10.png)
 
-1.	Replace `clean` by entering `UE4Editor ARGS="-clean"` into the text field.
+1.	Replace `clean` by entering `UE4Editor-Linux-Debug ARGS="-clean"` into the **Make arguments:** text field.
 
 	![](SetupIDE_Step11.png)
 
@@ -151,12 +136,13 @@ Before we begin, please make sure that you've already setup your workflow using
 		Click for full image.
 	[/REGION]
 
-1.	Repeat Steps 6 through 11 to edit the build configuration settings for Release mode.
+1.	Repeat Steps 6 through 11 to edit the build configuration settings for Release mode. Please make sure that your settings match the following image:
 
 	![](SetupIDE_Step13.png)
 
 	[REGION:note]
-	Read through our [](Programming/Development/BuildConfigurations) to learn more about the different configurations you can set your build environment to.
+	If you'd like to set up additional build configurations for your project, please refer to the Qt Creator [Build Settings Documentation](http://doc.qt.io/qtcreator/creator-build-settings.html),
+	where they show you how to add as many build configurations as you might need.
 	[/REGION]
 
 ## End Result
@@ -169,6 +155,8 @@ Before we begin, please make sure that you've already setup your workflow using
 	Click for full image.
 [/REGION]
 
-At this point, you've configured UE4 to build in both Debug and Release mode. Now, you're ready to configure Qt Creator to run UE4.
+At this point, you've configured UE4 to build in both Debug and Release mode. Depending on your system's specifications, you may notice that UE4 runs slower while in Debug mode. For more information 
+on the performance characteristics of available build configurations, please read through our [](Programming/Development/BuildConfigurations) page. With this in 
+mind, you're ready to configure Qt Creator to run UE4.
 
-%Steps%
\ No newline at end of file
+[NAV]
\ No newline at end of file
diff --git a/Engine/Documentation/Source/Platforms/Linux/BeginnerLinuxDeveloper/SettingUpAnIDE/2/SettingUpAnIDE_2.INT.udn b/Engine/Documentation/Source/Platforms/Linux/BeginnerLinuxDeveloper/SettingUpAnIDE/2/SettingUpAnIDE_2.INT.udn
--- a/Engine/Documentation/Source/Platforms/Linux/BeginnerLinuxDeveloper/SettingUpAnIDE/2/SettingUpAnIDE_2.INT.udn
+++ b/Engine/Documentation/Source/Platforms/Linux/BeginnerLinuxDeveloper/SettingUpAnIDE/2/SettingUpAnIDE_2.INT.udn
@@ -1,4 +1,4 @@
-Availability: Docs
+Availability: Public
 Title: 2 - Configuring your IDE to Run UE4
 Crumbs:
 Description:Learn how to configure your IDE to run UE4 in Debug and Release mode.
@@ -9,23 +9,7 @@ order: 2
 tags:Linux
 SkillLevel:Beginner
 
-[VAR:Steps]
-[OBJECT:Navigation]
-	[PARAM:previous]
-		[Previous Step](Platforms/Linux/BeginnerLinuxDeveloper/SettingUpAnIDE/1 "%Platforms/Linux/BeginnerLinuxDeveloper/SettingUpAnIDE/1:title%")
-	[/PARAM]
-	[PARAM:current]
-	[/PARAM]
-	[PARAM:home]
-		[Home](Platforms/Linux/BeginnerLinuxDeveloper/SettingUpAnIDE)
-	[/PARAM]
-	[PARAM:next]
-		[Next Step](Platforms/Linux/BeginnerLinuxDeveloper/SettingUpAnIDE/3 "%Platforms/Linux/BeginnerLinuxDeveloper/SettingUpAnIDE/3:title%")
-	[/PARAM]
-[/OBJECT]
-[/VAR]
-
-%Steps%
+[NAV]
 
 At this point, you've configured UE4 to build in both Debug and Release mode. Now, you're ready to configure Qt Creator to run UE4.
 
@@ -91,4 +75,4 @@ At this point, you've configured UE4 to build in both Debug and Release mode. No
 
 You're now ready to build and run UE4 with Qt Creator in both Debug and Release mode. 
 
-%Steps%
\ No newline at end of file
+[NAV]
\ No newline at end of file
diff --git a/Engine/Documentation/Source/Platforms/Linux/BeginnerLinuxDeveloper/SettingUpAnIDE/3/SettingUpAnIDE_3.INT.udn b/Engine/Documentation/Source/Platforms/Linux/BeginnerLinuxDeveloper/SettingUpAnIDE/3/SettingUpAnIDE_3.INT.udn
--- a/Engine/Documentation/Source/Platforms/Linux/BeginnerLinuxDeveloper/SettingUpAnIDE/3/SettingUpAnIDE_3.INT.udn
+++ b/Engine/Documentation/Source/Platforms/Linux/BeginnerLinuxDeveloper/SettingUpAnIDE/3/SettingUpAnIDE_3.INT.udn
@@ -1,4 +1,4 @@
-Availability: Docs
+Availability: Public
 Title: 3 - Building UE4 from your IDE
 Crumbs:
 Description:Learn how to build Unreal Engine from your IDE.
@@ -9,23 +9,7 @@ order: 3
 tags:Linux
 SkillLevel:Beginner
 
-[VAR:Steps]
-[OBJECT:Navigation]
-	[PARAM:previous]
-		[Previous Step](Platforms/Linux/BeginnerLinuxDeveloper/SettingUpAnIDE/2 "%Platforms/Linux/BeginnerLinuxDeveloper/SettingUpAnIDE/2:title%")
-	[/PARAM]
-	[PARAM:current]
-	[/PARAM]
-	[PARAM:home]
-		[Home](Platforms/Linux/BeginnerLinuxDeveloper/SettingUpAnIDE)
-	[/PARAM]
-	[PARAM:next]
-		[Next Step](Platforms/Linux/BeginnerLinuxDeveloper/SettingUpAnIDE/4 "%Platforms/Linux/BeginnerLinuxDeveloper/SettingUpAnIDE/4:title%")
-	[/PARAM]
-[/OBJECT]
-[/VAR]
-
-%Steps%
+[NAV]
 
 Now that you've configured Qt Creator to build and run UE4, you're ready to run UE4 from the IDE. Before you run UE4 from the IDE, you'll want to open some output windows so that you can see 
 Unreal Build System at work.
@@ -48,12 +32,12 @@ Unreal Build System at work.
 
 	![](SetupIDE_Step15A.png)
 
-1.	Now, open the **Window** menu, hovering over the **Output Panes** sub-menu, clicking **Compile Output**. Alternatively, you can press the **Alt+4** keyboard combination to open the **Compile Output** window.
+1.	Now, open the **Window** menu, hovering over the **Output Panes** sub-menu, clicking **Compile Output**. Alternatively, you can press the **Alt+4** keyboard combination to open the **Compile Output** window pane.
 
 	![](SetupIDE_Step15B.png)
 
 	[REGION:note]
-	Opening the **Compile Output** window allows you to watch for any errors while building UE4 from source.
+	Opening the **Compile Output** window pane allows you to watch for any errors while building UE4 from source.
 	[/REGION]
 
 ## Building UE4
@@ -82,4 +66,4 @@ If you want to shorten the time it takes to compile the engine from source, we r
 machine that has at least eight (8) gigabytes of RAM with a multi-core processor having at least eight (8) cores (including 
 hyperthreading).
 
-%Steps%
\ No newline at end of file
+[NAV]
\ No newline at end of file
diff --git a/Engine/Documentation/Source/Platforms/Linux/BeginnerLinuxDeveloper/SettingUpAnIDE/4/SettingUpAnIDE_4.INT.udn b/Engine/Documentation/Source/Platforms/Linux/BeginnerLinuxDeveloper/SettingUpAnIDE/4/SettingUpAnIDE_4.INT.udn
--- a/Engine/Documentation/Source/Platforms/Linux/BeginnerLinuxDeveloper/SettingUpAnIDE/4/SettingUpAnIDE_4.INT.udn
+++ b/Engine/Documentation/Source/Platforms/Linux/BeginnerLinuxDeveloper/SettingUpAnIDE/4/SettingUpAnIDE_4.INT.udn
@@ -1,4 +1,4 @@
-Availability: Docs
+Availability:Public
 Title: 4 - Running UE4 from your IDE
 Crumbs:
 Description:Learn how to run Unreal Engine from your IDE.
@@ -9,26 +9,20 @@ order: 4
 tags:Linux
 SkillLevel:Beginner
 
-[VAR:Steps]
-[OBJECT:Navigation]
-	[PARAM:previous]
-		[Previous Step](Platforms/Linux/BeginnerLinuxDeveloper/SettingUpAnIDE/3 "%Platforms/Linux/BeginnerLinuxDeveloper/SettingUpAnIDE/3:title%")
-	[/PARAM]
-	[PARAM:current]
-	[/PARAM]
-	[PARAM:home]
-		[Home](Platforms/Linux/BeginnerLinuxDeveloper/SettingUpAnIDE)
-	[/PARAM]
-	[PARAM:next]
-		[Next Step](Platforms/Linux/BeginnerLinuxDeveloper/SettingUpAnIDE/5 "%Platforms/Linux/BeginnerLinuxDeveloper/SettingUpAnIDE/5:title%")
-	[/PARAM]
-[/OBJECT]
-[/VAR]
-
-%Steps%
+[NAV]
 
 Now that you've compiled the engine, you're ready to run the editor. During this step, you're going to use the IDE to run UE4.
 
+1.	If you want to view the status of your application while running UE4, you'll want to open the **Application Output** window pane. To open the **Application Output** window pane, press the **Alt+3** keyboard combination.
+
+	[REGION:lightbox]
+		[![](SetupIDE_Step17a.png)(w:720)](SetupIDE_Step17a.png)
+	[/REGION]
+	
+	[REGION:caption]
+		Click for full image.
+	[/REGION]
+
 1.	To run UE4 from the IDE, click the **Run** button. Alternatively, you can press the **Ctrl+R** keyboard combination to run UE4 from the IDE.
 
 	[REGION:lightbox]
@@ -51,4 +45,4 @@ Now that you've compiled the engine, you're ready to run the editor. During this
 
 Congratulations! You should now see UE4 running from your IDE.
 
-%Steps%
\ No newline at end of file
+[NAV]
\ No newline at end of file
diff --git a/Engine/Documentation/Source/Platforms/Linux/BeginnerLinuxDeveloper/SettingUpAnIDE/5/SettingUpAnIDE_5.INT.udn b/Engine/Documentation/Source/Platforms/Linux/BeginnerLinuxDeveloper/SettingUpAnIDE/5/SettingUpAnIDE_5.INT.udn
--- a/Engine/Documentation/Source/Platforms/Linux/BeginnerLinuxDeveloper/SettingUpAnIDE/5/SettingUpAnIDE_5.INT.udn
+++ b/Engine/Documentation/Source/Platforms/Linux/BeginnerLinuxDeveloper/SettingUpAnIDE/5/SettingUpAnIDE_5.INT.udn
@@ -1,4 +1,4 @@
-Availability: Docs
+Availability: Public
 Title: 5 - Creating a Blueprint Project from your IDE
 Crumbs:
 Description:Learn how to create and configure a Blueprint project from your IDE.
@@ -9,23 +9,7 @@ order: 5
 tags:Linux
 SkillLevel:Beginner 
 
-[VAR:Steps]
-[OBJECT:Navigation]
-	[PARAM:previous]
-		[Previous Step](Platforms/Linux/BeginnerLinuxDeveloper/SettingUpAnIDE/4 "%Platforms/Linux/BeginnerLinuxDeveloper/SettingUpAnIDE/4:title%")
-	[/PARAM]
-	[PARAM:current]
-	[/PARAM]
-	[PARAM:home]
-		[Home](Platforms/Linux/BeginnerLinuxDeveloper/SettingUpAnIDE)
-	[/PARAM]
-	[PARAM:next]
-		[Next Step](Platforms/Linux/BeginnerLinuxDeveloper/SettingUpAnIDE/6 "%Platforms/Linux/BeginnerLinuxDeveloper/SettingUpAnIDE/6:title%")
-	[/PARAM]
-[/OBJECT]
-[/VAR]
-
-%Steps%
+[NAV]
 
 You can create either Blueprint or CPP projects from the [](Engine/Basics/Projects/Browser) in UE4. In this section, we'll show you how to create a new Blueprint project from UE4. 
 For the purposes of this guide, let's suppose that you want to 
@@ -90,4 +74,4 @@ Creating any type of Blueprint project, with or without a project template, is n
 however, when you'll want to create and modify the source code driving 
 your game. In the next section, we'll show you how to create a new CPP project from UE4.
 
-%Steps%
\ No newline at end of file
+[NAV]
\ No newline at end of file
diff --git a/Engine/Documentation/Source/Platforms/Linux/BeginnerLinuxDeveloper/SettingUpAnIDE/6/SettingUpAnIDE_6.INT.udn b/Engine/Documentation/Source/Platforms/Linux/BeginnerLinuxDeveloper/SettingUpAnIDE/6/SettingUpAnIDE_6.INT.udn
--- a/Engine/Documentation/Source/Platforms/Linux/BeginnerLinuxDeveloper/SettingUpAnIDE/6/SettingUpAnIDE_6.INT.udn
+++ b/Engine/Documentation/Source/Platforms/Linux/BeginnerLinuxDeveloper/SettingUpAnIDE/6/SettingUpAnIDE_6.INT.udn
@@ -1,4 +1,4 @@
-Availability: Docs
+Availability: Public
 Title: 6 - Creating a CPP Project from your IDE
 Crumbs:
 Description:Learn how to create and configure a CPP project form your IDE.
@@ -14,22 +14,7 @@ Related: Programming/QuickStart
 Related: Engine/QuickStart
 Related: Support/Builds
 
-[VAR:Steps]
-[OBJECT:Navigation]
-	[PARAM:previous]
-		[Previous Step](Platforms/Linux/BeginnerLinuxDeveloper/SettingUpAnIDE/5 "%Platforms/Linux/BeginnerLinuxDeveloper/SettingUpAnIDE/5:title%")
-	[/PARAM]
-	[PARAM:current]
-	[/PARAM]
-	[PARAM:home]
-		[Home](Platforms/Linux/BeginnerLinuxDeveloper/SettingUpAnIDE)
-	[/PARAM]
-	[PARAM:next]
-	[/PARAM]
-[/OBJECT]
-[/VAR]
-
-%Steps%
+[NAV]
 
 By now, you've learned how to configure your IDE to run UE4 in a Linux environment. Also, you've learned how to create a new Blueprint project using a Project template. In this section, we'll show you 
 how to create a new CPP project from UE4. For the purposes of this demonstration, let's suppose that you want to create a new CPP game with a driving mechanic.
@@ -69,7 +54,11 @@ how to create a new CPP project from UE4. For the purposes of this demonstration
 	After UE4 creates your new CPP project, an instance of Qt Creator for your new CPP project will be launched, closing UE4Editor in the process.
 	[/REGION]
 
-1.	At this point, you should see Qt Creator's **Configure Project** menu in the newly launched instance of Qt Creator. Use what you've learned thus far to configure your new CPP project to have the following Build Settings for Debug and Release Mode:
+1.	At this point, you should see Qt Creator's **Configure Project** menu in the newly launched instance of Qt Creator. Use what you've learned thus far to configure your new CPP project to have the following Build Settings for Debug Mode;
+
+	![](SetupIDE_Step27a.png)
+
+	And Release Mode:
 
 	![](SetupIDE_Step27.png)
 
@@ -148,9 +137,9 @@ As you worked your way through this tutorial, you learned:
 Learn how to use Unreal Engine by referring to our [](GettingStarted) documentation!
 
 If you're looking to quickly get started with UE4, refer to the following tutorials:
+
 *	[](Engine/Content/QuickStart)
 *	[](Programming/QuickStart)
 *	[](Engine/QuickStart)
 
-
-%Steps%
\ No newline at end of file
+[NAV]
\ No newline at end of file
diff --git a/Engine/Documentation/Source/Platforms/Linux/BeginnerLinuxDeveloper/SettingUpAnIDE/SettingUpAnIDE.INT.udn b/Engine/Documentation/Source/Platforms/Linux/BeginnerLinuxDeveloper/SettingUpAnIDE/SettingUpAnIDE.INT.udn
--- a/Engine/Documentation/Source/Platforms/Linux/BeginnerLinuxDeveloper/SettingUpAnIDE/SettingUpAnIDE.INT.udn
+++ b/Engine/Documentation/Source/Platforms/Linux/BeginnerLinuxDeveloper/SettingUpAnIDE/SettingUpAnIDE.INT.udn
@@ -1,36 +1,17 @@
-Availability: Docs
-Title: How to Setup an IDE for UE4
+Availability: Public
+Title: How to Set up an IDE for UE4
 Crumbs:
-Description:Learn how to setup your IDE to build and run UE4 on Linux.
+Description:Learn how to set up your IDE to build and run UE4 on Linux.
 Platform:Linux
 Version: 4.14
 type:how to
-parent:Platforms/Linux/BeginnerLinuxDeveloper
+parent:Platforms/Linux
 order: 3
 tags:Linux
 SkillLevel:Beginner
 prereq:Platforms/Linux/BeginnerLinuxDeveloper/SettingUpAnUnrealWorkflow
+topic-image:IDE_HowToTopic.png
 
-[VAR:Params]
-[PARAM:image]
-	![%Platforms/Linux/BeginnerLinuxDeveloper/SettingUpAnIDE:title%](IDE_HowToTopic.png)
-[/PARAM]
-[PARAM:icon]
-	![](%ROOT%/tutorial_icon.png)(convert:false)
-[/PARAM]
-[PARAM:title]
-	%Platforms/Linux/BeginnerLinuxDeveloper/SettingUpAnIDE:title%
-[/PARAM]
-[PARAM:path]
-	[RELATIVE:Platforms/Linux/BeginnerLinuxDeveloper/SettingUpAnIDE]
-[/PARAM]
-[/VAR]
-
-[VAR:Topic]
-[OBJECT:Topic]
-%params%
-[/OBJECT]
-[/VAR]
 
 [REGION:imagetable]
 ![](CompiledEditor_LinuxIDE.png)(w:1280)  
@@ -39,7 +20,7 @@ At the end of this tutorial, you'll have an **Integrated Development Environment
 
 ## Goal
 
-The focus of this guide is to show you how to setup an IDE to build and run UE4. 
+The focus of this guide is to show you how to set up an IDE to build and run UE4. 
 
 ## Objectives
 
diff --git a/Engine/Documentation/Source/Platforms/Linux/BeginnerLinuxDeveloper/SettingUpAnUnrealWorkflow/1/SettingUpAnUnrealWorkflow_1.INT.udn b/Engine/Documentation/Source/Platforms/Linux/BeginnerLinuxDeveloper/SettingUpAnUnrealWorkflow/1/SettingUpAnUnrealWorkflow_1.INT.udn
--- a/Engine/Documentation/Source/Platforms/Linux/BeginnerLinuxDeveloper/SettingUpAnUnrealWorkflow/1/SettingUpAnUnrealWorkflow_1.INT.udn
+++ b/Engine/Documentation/Source/Platforms/Linux/BeginnerLinuxDeveloper/SettingUpAnUnrealWorkflow/1/SettingUpAnUnrealWorkflow_1.INT.udn
@@ -1,6 +1,6 @@
 Availability: Public
 Title: 1 - Required Setup
-Crumbs:%ROOT%, Platforms, Platforms/Linux, Platforms/Linux/BeginnerLinuxDeveloper, Platforms/Linux/BeginnerLinuxDeveloper/SettingUpAnUnrealWorkflow
+Crumbs:
 Description:Learn how to set up Git on Linux.
 Version: 4.13
 type:quick start
diff --git a/Engine/Documentation/Source/Platforms/Linux/BeginnerLinuxDeveloper/SettingUpAnUnrealWorkflow/2/SettingUpAnUnrealWorkflow_2.INT.udn b/Engine/Documentation/Source/Platforms/Linux/BeginnerLinuxDeveloper/SettingUpAnUnrealWorkflow/2/SettingUpAnUnrealWorkflow_2.INT.udn
--- a/Engine/Documentation/Source/Platforms/Linux/BeginnerLinuxDeveloper/SettingUpAnUnrealWorkflow/2/SettingUpAnUnrealWorkflow_2.INT.udn
+++ b/Engine/Documentation/Source/Platforms/Linux/BeginnerLinuxDeveloper/SettingUpAnUnrealWorkflow/2/SettingUpAnUnrealWorkflow_2.INT.udn
@@ -1,6 +1,6 @@
 Availability: Public
 Title: 2 - Downloading UE4 on Linux
-Crumbs:%ROOT%, Platforms, Platforms/Linux, Platforms/Linux/BeginnerLinuxDeveloper, Platforms/Linux/BeginnerLinuxDeveloper/SettingUpAnUnrealWorkflow
+Crumbs:
 Description:Learn how to fork and clone Unreal Engine source code.
 Version: 4.13
 type:quick start
diff --git a/Engine/Documentation/Source/Platforms/Linux/BeginnerLinuxDeveloper/SettingUpAnUnrealWorkflow/SettingUpAnUnrealWorkflow.INT.udn b/Engine/Documentation/Source/Platforms/Linux/BeginnerLinuxDeveloper/SettingUpAnUnrealWorkflow/SettingUpAnUnrealWorkflow.INT.udn
--- a/Engine/Documentation/Source/Platforms/Linux/BeginnerLinuxDeveloper/SettingUpAnUnrealWorkflow/SettingUpAnUnrealWorkflow.INT.udn
+++ b/Engine/Documentation/Source/Platforms/Linux/BeginnerLinuxDeveloper/SettingUpAnUnrealWorkflow/SettingUpAnUnrealWorkflow.INT.udn
@@ -1,46 +1,27 @@
 Availability: Public
-Title: Linux Quick Start Guide
-Crumbs:%ROOT%, Platforms, Platforms/Linux, Platforms/Linux/BeginnerLinuxDeveloper
+Title: Linux Quick Start
+Crumbs:
 Description:Learn how to download, build, and run UE4 on Linux.
 Platform:Linux
 Version: 4.13
 type:quick start
-parent:Platforms/Linux/BeginnerLinuxDeveloper
+parent:Platforms/Linux
 order: 2
 tags:Linux
 SkillLevel:Beginner
 prereq:GettingStarted/RecommendedSpecifications
 prereq:Platforms/Linux/BeginnerLinuxDeveloper/SupportedGfxDrivers
+topic-image:linuxQS_topic.png
 
-[VAR:Params]
-[PARAM:image]
-	![%Platforms/Linux/BeginnerLinuxDeveloper/SettingUpAnUnrealWorkflow:title%](linuxQS_topic.png)
-[/PARAM]
-[PARAM:icon]
-	![](%ROOT%/tutorial_icon.png)(convert:false)
-[/PARAM]
-[PARAM:title]
-	%Platforms/Linux/BeginnerLinuxDeveloper/SettingUpAnUnrealWorkflow:title%
-[/PARAM]
-[PARAM:path]
-	[RELATIVE:Platforms/Linux/BeginnerLinuxDeveloper/SettingUpAnUnrealWorkflow]
-[/PARAM]
-[/VAR]
-
-[VAR:Topic]
-[OBJECT:Topic]
-%params%
-[/OBJECT]
-[/VAR]
 
 [REGION:imagetable]
-![](CompiledEditor_Linux.png)(w:920)
+![](CompiledEditor_Linux.png)(w:920)  
 At the end of this tutorial, you'll have Unreal Engine 4 (UE4) running on your Linux machine.
 [/REGION]
 
 ## Goal
 
-The focus of this guide is to show you how to setup the primary workflow needed to get UE4 running on your Linux machine.
+The focus of this guide is to show you how to set up the primary workflow needed to get UE4 running on your Linux machine.
 
 ## Objectives
 
diff --git a/Engine/Documentation/Source/Platforms/Linux/BeginnerLinuxDeveloper/SupportedGfxDrivers/SupportedGfxDrivers.INT.udn b/Engine/Documentation/Source/Platforms/Linux/BeginnerLinuxDeveloper/SupportedGfxDrivers/SupportedGfxDrivers.INT.udn
--- a/Engine/Documentation/Source/Platforms/Linux/BeginnerLinuxDeveloper/SupportedGfxDrivers/SupportedGfxDrivers.INT.udn
+++ b/Engine/Documentation/Source/Platforms/Linux/BeginnerLinuxDeveloper/SupportedGfxDrivers/SupportedGfxDrivers.INT.udn
@@ -1,36 +1,17 @@
 Availability: Public
 Title: Supported Graphics Drivers
-Crumbs: %ROOT%, Platforms, Platforms/Linux, Platforms/Linux/BeginnerLinuxDeveloper
+Crumbs:
 Description:This page provides an overview of supported graphics drivers for Linux game developers.
 version: 4.13
-parent:Platforms/Linux/BeginnerLinuxDeveloper
-order: 1
+parent:Platforms/Linux
+order: 4
 tags:Linux
 type:Reference
 related:GettingStarted/RecommendedSpecifications
+topic-image:supportedgfxdrivers_topic.png
 
 [TOC(start:1 end:4)]
 
-[VAR:Params]
-[PARAM:image]
-	![%Platforms/Linux/BeginnerLinuxDeveloper:title%](supportedgfxdrivers_topic.png)
-[/PARAM]
-[PARAM:icon]
-	![](%ROOT%/tutorial_icon.png)(convert:false)
-[/PARAM]
-[PARAM:title]
-	%Platforms/Linux/BeginnerLinuxDeveloper/SupportedGfxDrivers:title%
-[/PARAM]
-[PARAM:path]
-	[RELATIVE:Platforms/Linux/BeginnerLinuxDeveloper/SupportedGfxDrivers]
-[/PARAM]
-[/VAR]
-
-[VAR:Topic]
-[OBJECT:Topic]
-%params%
-[/OBJECT]
-[/VAR]
 
 Unreal Engine requires compatible hardware and software support in order to run smoothly on your 
 machine. In an effort to help you determine your ideal development environment, we're providing 
