Copying //UE4/Dev-Documentation to //UE4/Samples-Main (Source: //UE4/Dev-Documentation @ 3372893)

diff --git a/Engine/Documentation/Source/Platforms/Linux/GettingStarted/LinuxCrossGettingStarted.INT.udn b/Engine/Documentation/Source/Platforms/Linux/GettingStarted/LinuxCrossGettingStarted.INT.udn
--- a/Engine/Documentation/Source/Platforms/Linux/GettingStarted/LinuxCrossGettingStarted.INT.udn
+++ b/Engine/Documentation/Source/Platforms/Linux/GettingStarted/LinuxCrossGettingStarted.INT.udn
@@ -1,134 +1,271 @@
 Availability:Public
-Title: Getting Started with Linux Projects
-Crumbs:%ROOT%, Platforms, Platforms/Linux
-Description:Setting up for development for the Linux platform.
-Version: 4.9
+Title: Cross-Compiling for Linux
+Crumbs:
+Description:Learn how to set up a cross-compilation toolchain in Windows, targeting the Linux platform.
+Version: 4.15
 SkillLevel: Advanced
-type:quick start
+type:overview
 parent:Platforms/Linux
 tags:Linux
-
-[VAR:Params]
-[PARAM:image]
-	![%Platforms/Linux/GettingStarted:title%](GettingStarted_Topic.png)
-[/PARAM]
-[PARAM:icon]
-	![](%ROOT%/tutorial_icon.png)(convert:false)
-[/PARAM]
-[PARAM:title]
-	%Platforms/Linux/GettingStarted:title%
-[/PARAM]
-[PARAM:path]
-	[RELATIVE:Platforms/Linux/GettingStarted]
-[/PARAM]
-[/VAR]
-
-[VAR:Topic]
-[OBJECT:Topic]
-%params%
-[/OBJECT]
-[/VAR]
+related:Programming/Development/CompilingProjects
+related:Programming/Development/BuildConfigurations
+related:Engine/Basics/Projects/Packaging
+related:Platforms/Linux/AdvancedLinuxDeveloper/LinuxCrossCompileLegacy
+related:Programming/Basics/ConfigurationFiles
+topic-image:GettingStarted_Topic.png
+prereq:Platforms/Linux/BeginnerLinuxDeveloper/SettingUpAnUnrealWorkflow
 
 <!--- %Globals:BinarySourceSelectVar% --->
 
-[TOC(start:2 end:3)]
+[TOC(start:2 end:4)]
+
+![](LinuxXC_Banner.png)(w:1200)
 
-## Why cross-compilation
+Recent improvements to the cross-compile toolchain make it possible for Unreal Engine 4 (UE4) developers to target a variety of platforms from their Windows 
+development environment. This guide tells you which platforms you can target with the cross-compilation toolchain, providing a link 
+to download the toolchain, and ultimately showing you how to set it up to cross-compile a UE4 project for Linux. If you're developing a project in Unreal Engine, version 4.15, you'll also 
+learn how to switch between target platforms.
 
-Cross-compilation was chosen in order to make it easier for game developers (who tend to have Windows-centric workflow) to target Linux. At the moment, this is only supported for Windows, Mac users currently have to resort to [native compilation](https://wiki.unrealengine.com/Building_On_Linux). We support and test Linux-x86_64, but we provide the libraries and the toolchain (clang-3.5.0-arm) which allow you to compile for Linux ARM as well (original Raspberry Pi and up), although this will require you to make a (minor) code change in UBT.
+[REGION:note]
+If you're developing your project with Unreal Engine, version 4.13 (or older), you'll want to refer to our [](Platforms/Linux/AdvancedLinuxDeveloper/LinuxCrossCompileLegacy) documentation.
+[/REGION]
 
-From this point on we assume that you are targeting x86_64 Linux; although the most of below will apply to compiling for ARM as well (except for the different toolchain). Note that you can [build your own](https://wiki.unrealengine.com/Building_Linux_cross-toolchain) toolchain if you want different versions of the tools or want to target a different architecture.
+[INCLUDE:Platforms/Linux/AdvancedLinuxDeveloper/LinuxCrossCompileLegacy#Why_XC]
 
-## Setting up the toolchain 
+## Setting up the Toolchain
 
-Add an environment variable (Control Panel->System->Advanced system settings->Advanced->Environment variables) named LINUX_ROOT, value of which is the absolute path to your toolchain, without trailing backslash:
+If you want to set up a toolchain from your Windows (host) computer to target any of the aforementioned platforms, you'll have to download the toolchain and 
+set your system's environment variable to package your application for Linux and ARM based environments. If this is your first time packaging UE4 projects for the 
+Linux platform, you'll also have to build the engine in Visual Studio for Linux. After working through the following steps, your host environment will be set up 
+to target any of the currently supported platforms. 
 
-[REGION:raw]
-![](image00.png)
+[REGION:note]
+Please make sure that you have a reliable Internet connection prior to following these steps. Depending on the speed of your 
+Internet connection, downloading the toolchain can take anywhere from a few minutes to a few hours.
 [/REGION]
 
-After you made sure that variable is set, regenerate UE4 project files (using GenerateProjectFiles.bat) and restart Visual Studio. After this, you should have "Linux" available among Win32/Win64 configurations, and you should be able to cross-compile for it.
+### Downloading the Toolchain
+
+1.	Click the [link](http://cdn.unrealengine.com/qfe/v8_clang-3.9.0-centos7.zip) to download our Clang 3.9.0 based toolchain.
+	
+1.	Now, create a new folder, naming it `Linux_CrossCompileToolChain`.
+
+	[REGION:lightbox]
+		[![](LinuxXC_Step2.png)(w:600)](LinuxXC_Step2.png)
+	[/REGION]
+
+	[REGION:caption]
+		Click for full image.
+	[/REGION]
+	
+	[REGION:note]
+	If you're using a different naming convention, substitute your preferred folder name with the one we've chosen for illustrative purposes.
+	[/REGION]
+
+1.	**Extract** the toolchain, being sure to set the extracted file's destination to the folder that you created in the previous step. 
+
+	![](LinuxXC_Step3.png)(w:600)
+
+	[REGION:caption]
+		Click for full image.
+	[/REGION]
+
+1.	Open the `[ROOT]/Linux_CrossCompileToolChain/v8_clang-3.9.0-centos7` folder to view the extracted files.
+
+	[REGION:lightbox]
+		[![](LinuxXC_Step4.png)(w:600)](LinuxXC_Step4.png)
+	[/REGION]
+
+	[REGION:caption]
+		Click for full image.
+	[/REGION]
+
+	#### Supported Platforms
+
+	Refer to the following table to see which platforms can be targeted with your version of UE4: 
+
+	| **Architecture**                                                    | **Sub-Architecture** | **Vendor** | **System** | **Application Binary Interface (abi)** | UE 4.14 Support | UE 4.15 Support  |
+	| ------------------------------------------------------------------- | -------------------- | ---------- | ---------- | -------------------------------------- | --------------- | ---------------- |
+	| **[AArch64](https://fedoraproject.org/wiki/Architectures/AArch64)** | N/A                  | Unknown    | Linux      | gnueabi                                | &#x2717;        | &#x2713;         |
+	| **[ARM](https://developer.arm.com/)**                               | N/A                  | Unknown    | Linux      | gnueabihf                              | &#x2717;        | &#x2713;         |
+	| **[i686](https://software.intel.com/en-us/)**                       | N/A                  | Unknown    | Linux      | gnu                                    | &#x2717;        | &#x2717;         |
+	| **[x86_64](http://developer.amd.com/)**                             | N/A                  | Unknown    | Linux      | gnu                                    | &#x2713;        | &#x2713;         |
+	
+	[REGION:note]
+	Please note that we haven't yet added Engine libraries for the i686 platform.
+	[/REGION]
+
+### Setting the Environment Variable
+
+1.	Prior to setting the `LINUX_MULTIARCH_ROOT` environment variable, double-click the **Setup** batch file (`Setup.bat`) to output the appropriate environment variable settings for your computer. 
+
+	[REGION:lightbox]
+		[![](LinuxXC_Step5.png)(w:600)](LinuxXC_Step5.png)
+	[/REGION]
+
+	[REGION:caption]
+		Click for full image.
+	[/REGION]
+
+1.	At this point, a new **OutputEnvVars** text file should be located inside of the `v8_clang-3.9.0-centos7` folder. 
+
+	[REGION:lightbox]
+		[![](LinuxXC_Step6.png)(w:600)](LinuxXC_Step6.png)
+	[/REGION]
+
+	[REGION:caption]
+		Click for full image.
+	[/REGION]
+
+	[REGION:tip]
+	The `OutputEnvVars.txt` file contains useful information for setting your system's environment variable.
+	[/REGION]
 
-## Setting up packaging for Linux
+1.	Press the **Windows Key** together with the **Pause Break Key** to open the **System Information** menu. 
+	Now, click the **Advanced System Settings** link to open the **Advanced** sub-menu, which is located in the **System Properties** menu.
 
-[OBJECT:ToggleButtonContent]
-	[PARAMLITERAL:category]
-	BinarySourceSelect
-	[/PARAMLITERAL]
-	[PARAMLITERAL:id]
-	binaryinstall
-	[/PARAMLITERAL]
-	[PARAMLITERAL:active]
-	active_button_content
-	[/PARAMLITERAL]
-	[PARAM:content]
-	%binary%
-	[/PARAM]
-[/OBJECT]
-[OBJECT:ToggleButtonContent]
-	[PARAMLITERAL:category]
-	BinarySourceSelect
-	[/PARAMLITERAL]
-	[PARAMLITERAL:id]
-	sourceinstall
-	[/PARAMLITERAL]
-	[PARAMLITERAL:active]
-	[/PARAMLITERAL]
-	[PARAM:content]
-	%source%
-	[/PARAM]
-[/OBJECT]
+	[REGION:lightbox]
+		[![](LinuxXC_Step8.png)(w:600)](LinuxXC_Step8.png)
+	[/REGION]
 
-[VAR:source]
-For packaging for Linux platform you will need to build the following targets:
+	[REGION:caption]
+		Click for full image.
+	[/REGION]
 
-For Linux:
+1.	Now, click the **Environment Variables...** button.
 
-* CrashReportClient
-* UE4Game (if you project is content-only, otherwise just build your project)
+	![](LinuxXC_Step9.png)
 
-For Windows:
+1.	To add a new **System Variable**, click the **New...** button under the **System variables** list.
 
-* The editor itself (this is needed to get Linux target modules built, so that Unreal Editor and Unreal Frontend know how to cook/package games for Linux).
-* UnrealPak and ShaderCompileWorker (you probably have those built anyway, but mentioning it for completeness).
+	![](LinuxXC_Step10.png)
 
-[/VAR]
+	[REGION:note]
+	If you haven't opened `OutputEnvVars.txt` yet, please do so before moving onto the next step.
+	
+	[REGION:lightbox]
+		[![](LinuxXC_Step10Note.png)(w:600)](LinuxXC_Step10Note.png)
+	[/REGION]
 
-[VAR:binary]
-Note that binary release can only package content-only projects. If you want to package a code-based project for Linux (or a project that has non-default third party plugins), you will have to set up a source build. Otherwise, if you are fine with content-only, there are no need to build anything as Linux binaries of UE4Game and CrashReportClient are bundled.
-[/VAR]
+	[REGION:caption]
+		Click for full image.
+	[/REGION]
 
-## Packaging for Linux
+	[/REGION]
 
-The easiest way to package a project is to open it in the editor and then using File->Package To->Linux (assuming that you have a cross-toolchain installed in the previous step and Linux target modules are built for the editor in question - if you don't see Linux in the list, then it's likely that one of the former is not true - see above). After some time (which depends on the project in question and is rather short for sample project) you will have game assets and binaries in the directory you chose to package to.
+1.	After the **New System Variable** dialog menu opens, copy the `LINUX_MULTIARCH_ROOT` environment variable setting into the text boxes for the **New System Variable** dialog.
 
-The details of the packaging process can be seen by clicking the "Show Output Log" link. If this process fails with error message "unable to find UnrealPak or ShaderCompileWorker" see above about building them for the host (windows) platform.
+	[REGION:lightbox]
+		[![](LinuxXC_Step11.png)(w:600)](LinuxXC_Step11.png)
+	[/REGION]
 
-Scp or otherwise copy it to a target machine (mounting a Samba share - if you know how to do that - may be better if target machine is low on disk space, this also reduces iteration times), chmod +x the target executable (which will be located in LinuxNoEditor/<ProjectName>/Binaries/Linux/ directory) and run it.
+	[REGION:caption]
+		Click for full image.
+	[/REGION]
 
-### Changes for Linux-ARM platforms
+1.	After returning to the **Evironment Variables** menu, click the **Ok** button to save the `LINUX_MULTIARCH_ROOT` environment variable.
 
-If you are using the Linux-ARM cross compile toolchain before running GenerateProjectFiles.bat below edit the following file in the UE4 source code: 
+	![](LinuxXC_Step14.png)
 
-.../UnrealEngine/Engine/Source/Programs/UnrealBuildTool/Linux/UEBuildLinux.cs 
+	[REGION:note]
+	Please note that if you already have the `LINUX_ROOT` environment variable set, you don't need to delete it because the engine will prioritize `LINUX_MULTIARCH_ROOT` over `LINUX_ROOT` during the packaging process.
+	[/REGION]
 
-by commenting out the following line: 
+1.	Now, click the **OK** button to close the **System Properties** menu.
 
-static private string DefaultArchitecture = "x86_64-unknown-linux-gnu"; 
+	![](LinuxXC_Step15.png)
+	
+	[REGION:note]
+	Now that you've added a new Environment Variable to your system, you'll need to close and restart all currently running instances of **Visual Studio** and **Command Prompt**.
+	[/REGION]
 
-and un-commenting the linux-arm architecture line just below it: 
+### Setting up the Linux Build Configuration
 
-//static private string DefaultArchitecture = "arm-unknown-linux-gnueabihf";
+If you've never packaged a project for Linux before, you'll want to follow these steps.
 
-An additional step is required for the packaged project to be runnable on Linux-ARM platforms. Assuming the packaged project is located at \foo\bar\SunTemple_Linux open the following location:
+1.	Opening the **Command Prompt**, run `[UE4 ROOT]/Setup.bat` to check for, update, and install prerequisite dependencies for building UE4 from **Visual Studio**.
 
-\foo\bar\SunTemple_Linux\LinuxNoEditor\Engine\Binaries\Linux 
+	![](LinuxXC_Step17.png)
 
-and replace libopenal.so.1 with the version from: 
+1.	Run `[UE4 ROOT]/GenerateProjectFiles.bat` to set up the project files for UE4.
 
-...\UnrealEngine\Engine\Source\ThirdParty\OpenAL\1.15.1\lib\Linux\arm-unknown-linux-gnueabihf\libopenal.so 
+	![](LinuxXC_Step18.png)
 
-where ...\UnrealEngine is the location of the UE4 source code. Please be sure to rename libopenal.so to libopenal.so.1.
+	[REGION:warning]
+	`GenerateProjectFiles.bat` must be run inside of the root folder that UE4 resides in (`[UE4 ROOT]`). If you run the batch file with [Unreal Game Sync](https://github.com/EpicGames/UnrealEngine/blob/release/Engine/Source/Programs/UnrealGameSync/Deployment.txt), Linux won't show up as a valid **Build Configuration** inside of **Visual Studio**.
+	[/REGION]
 
-Now the project can be copied over to the target machine, chmod+x the target executable and run it as instructed above.
+1.	Navigating to `[UE4 ROOT]`, double-click **UE4.sln** to open the UE4 solution in **Visual Studio**.
+
+	[REGION:lightbox]
+		[![](LinuxXC_Step19.png)(w:600)](LinuxXC_Step19.png)
+	[/REGION]
+
+	[REGION:caption]
+		Click for full image.
+	[/REGION]
+
+1.	If you haven't built UE4, go ahead and press the **F5** key to build and run UE4 from Visual Studio.
+
+	[REGION:lightbox]
+		[![](LinuxXC_Step20.png)(w:600)](LinuxXC_Step20.png)
+	[/REGION]
+
+	[REGION:caption]
+		Click for full image.
+	[/REGION]
+
+1.	Now, stop running the engine, and go back to **Visual Studio** to select the **Linux** solution platform in Visual Studio. You can build for Linux by pressing the following keyboard combination: **Ctrl+Shift+B**.
+
+	[REGION:lightbox]
+		[![](LinuxXC_Step21.png)(w:600)](LinuxXC_Step21.png)
+	[/REGION]
+
+	[REGION:caption]
+		Click for full image.
+	[/REGION]
+
+	[REGION:note]
+	At this point, you should see the toolchain being displayed inside of Visual Studio's **Output Window**.
+
+	[REGION:lightbox]
+		[![](LinuxXC_Step21Note.png)(w:600)](LinuxXC_Step21Note.png)
+	[/REGION]
+	
+	[REGION:caption]
+		Click for full image.
+	[/REGION]
+	[/REGION]
+
+## Targeting Specific Platforms
+
+[REGION:note]
+Because version 4.14 can only target the x86_64 platform, the following section is intended for developers working with version 4.15.
+[/REGION]
+
+If you'd like your project to target a specific platform, you'll need to edit your project's [Default Engine Configuration File](Programming/Basics/ConfigurationFiles) file. To do so, navigate to the 
+file's location (`[Project Directory]\Config`) and open `DefaultEngine.ini` to add the following lines:
+
+	[/Script/LinuxTargetPlatform.LinuxTargetSettings]
+	TargetArchitecture=X86_64UnknownLinuxGnu
+
+Feel free to set `TargetArchitecture` to any of the following values:
+
+* X86_64UnknownLinuxGnu
+* ArmUnknownLinuxGnueabihf
+* AArch64UnknownLinuxGnueabi
+
+## End Result
+
+Congratulations! Moving forward, you should be able to see the cross-compile toolchain being used while packaging your project for Linux.
+
+[REGION:lightbox]
+	[![](LinuxXC_EndResult.png)(w:600)](LinuxXC_EndResult.png)
+[/REGION]
+
+[REGION:caption]
+	Click for full image.
+[/REGION]
 
+ Now, you're ready to [compile](Programming/Development/CompilingProjects), [build](Programming/Development/BuildConfigurations), and [package](Engine/Basics/Projects/Packaging) projects for a variety of Linux platforms from your Windows (host) computer. 
+ 
\ No newline at end of file
diff --git a/Engine/Documentation/Source/Platforms/Linux/IntermediateLinuxDeveloper/AndroidOnLinux/1/AndroidOnLinux_1.INT.udn b/Engine/Documentation/Source/Platforms/Linux/IntermediateLinuxDeveloper/AndroidOnLinux/1/AndroidOnLinux_1.INT.udn
new file mode 100644
--- /dev/null
+++ b/Engine/Documentation/Source/Platforms/Linux/IntermediateLinuxDeveloper/AndroidOnLinux/1/AndroidOnLinux_1.INT.udn
@@ -0,0 +1,127 @@
+Availability: Docs
+Title: 1 - Required Setup
+Crumbs:
+Description:Learn how to install and set up the Android SDK for mobile game development with UE4 on Linux.
+Version: 4.14
+type:quick start
+parent:Platforms/Linux/IntermediateLinuxDeveloper/AndroidOnLinux
+order: 1
+tags:Linux
+SkillLevel:Intermediate
+Checkpoint:AndroidOnLinuxQS
+Contributor:yaakuro
+
+[nav]
+
+[EXCERPT:AndroidLinux_1]
+Unreal Engine 4 (UE4) uses a special version of the Android Software Development Kit (SDK), called CodeWorks for 
+Android 1R5. This version of CodeWorks for Android is the easiest way to set up the SDK and Native Development Kit (NDK) 
+needed to develop Android projects with UE4. If you have other versions of the Android SDK installed (or, older versions of CodeWorks for 
+Android), we recommend that you uninstall them and install [CodeWorks for Android](https://developer.nvidia.com/codeworks-android), using 
+the CodeWorks for Android installer that's being distributed with UE4.  
+
+[REGION:note]
+Please note that we've bundled the CodeWorks for Android installer with UE4 to make sure that your Android SDK will be compatible with Unreal Engine. 
+[/REGION]
+
+## Installing CodeWorks for Android
+
+[REGION:warning]
+UE4 should be closed before installing CodeWorks for Android 1R5. Also, please make sure that you have Java Development Kit 8 (JDK 8) installed.
+[/REGION]
+
+1.	Getting started, open the Terminal and navigate to the `[UE4Root]` folder that contains the compiled engine (using the `cd` command, changing the shell working directory).
+
+	![](LinuxAndroid_Step0.png)(w:900)
+
+1.	From inside of the `[UE4Root]` folder, navigate to the CodeWorks for Android 1R5 binary file (**CodeWorksforAndroid-1R5-linux-x64.run**) by entering the following string into the terminal: `cd Engine/Extras/AndroidWorks/Linux`.
+
+	![](LinuxAndroid_Step1.png)(w:900)
+
+1.	Now, enter `ls -l` into the terminal, displaying a long listing of the directory's contents.
+
+	![](LinuxAndroid_Step2.png)(w:900)
+
+	The file's permission settings tell you that **CodeWorksforAndroid-1R5-linux-x64.run** is not ready for execution.
+	Before you can run the CodeWorks for Android 1R5 binary, you need to change the file's access rights to allow for 'execute' permissions.
+
+1.	To change the file permission settings of `CodeWorksforAndroid-1R5-linux-x64.run`, enabling it to run like a normal executable, enter `chmod +x CodeWorksforAndroid-1R5-linux-x64.run` into the terimal.
+	
+	![](LinuxAndroid_Step3.png)(w:900)
+
+	[REGION:tip]
+	If you want to learn more about the `chmod` command, enter `man chmod` into the terminal to read the manual page covering how `chmod` works.
+	[/REGION]
+
+1.	Now, enter `ls -l` into the terminal.
+
+	![](LinuxAndroid_Step4.png)(w:900)
+
+	The file's permission settings highlighted in the image indicate that the CodeWorks for Android 1R5 binary file is ready for execution.
+
+1.	To run the CodeWorks for Android binary, enter `./CodeWorksforAndroid-1R5-linux-x64.run` into the terminal.
+
+	![](LinuxAndroid_Step5.png)(w:900)
+
+	[REGION:warning]
+	Please make sure that your computer has a valid Internet connection prior to moving onto the next step.
+	[/REGION]
+
+1.	After the CodeWorks for Android 1R5 binary unpacks the **NVIDIA CodeWorks for Android**  installer, you'll see the **Installer Welcome** window, prompting you to click either **Next >** or **Cancel**.
+
+	![](LinuxAndroid_Step6.png)(w:900)
+
+1.	From the **Installer Welcome** window, click the **Next >** button to begin the installation process.
+
+	![](LinuxAndroid_Step7.png)(w:900)
+
+1.	Now, from the **Android Developer's Welcome Window**, click the **Next >** button to continue the installation process.
+
+	![](LinuxAndroid_Step8.png)(w:900)
+
+1.	Specify the directory where CodeWorks for Android 1R5 should be installed and downloaded.
+
+	![](LinuxAndroid_Step9.png)(w:900)
+
+1.	Click the **Next >** button to download CodeWorks for Android 1R5 from NVIDIA, ultimately installing CodeWorks on your computer.
+
+	![](LinuxAndroid_Step10.png)(w:900)
+
+1.	Now, verify that you're installing **CodeWorks for Android 1R5** before clicking the **Next >** button from the **Android Component Manager** window.
+
+	![](LinuxAndroid_Step11.png)(w:900)
+
+1.	After the **Terms and Conditions** window appears, take some time to read through all of the **Terms and Conditions** before selecting the **Accept All** check box.
+
+	![](LinuxAndroid_Step12.png)(w:900)
+
+1.	Now, click the **Accept** button to begin the installation process.
+
+	![](LinuxAndroid_Step13.png)(w:900)
+
+1.	Once the installer completes the installation process, go ahead and click the **Next >** button to proceed.
+
+	![](LinuxAndroid_Step14.png)(w:900)
+
+1.	After the **Installation Complete** window opens, click the **Finish** button to complete the installation.
+
+	![](LinuxAndroid_Step15.png)(w:900)
+
+1.	Make sure to reboot your computer before moving onto the next step.
+
+	![](LinuxAndroid_Note16.png)(w:900)
+
+	[REGION:tip]
+	If you don't want to use Ubuntu's GUI, you can restart your computer by entering `sudo reboot` into the terminal. 
+	<!-- 
+	![](LinuxAndroid_Step16.png)(w:900)
+	-->
+	[/REGION]
+[/EXCERPT:AndroidLinux_1]
+
+## End Result
+
+At this point in the tutorial, you've installed the Android SDK, which is the first step in preparation for mobile game development with UE4 on Linux. 
+The next step shows you how to set up the development environment on your Android device. 
+
+[nav]
\ No newline at end of file
diff --git a/Engine/Documentation/Source/Platforms/Linux/IntermediateLinuxDeveloper/AndroidOnLinux/2/AndroidOnLinux_2.INT.udn b/Engine/Documentation/Source/Platforms/Linux/IntermediateLinuxDeveloper/AndroidOnLinux/2/AndroidOnLinux_2.INT.udn
new file mode 100644
--- /dev/null
+++ b/Engine/Documentation/Source/Platforms/Linux/IntermediateLinuxDeveloper/AndroidOnLinux/2/AndroidOnLinux_2.INT.udn
@@ -0,0 +1,109 @@
+Availability: Docs
+Title: 2 - Setting up your Android Device
+Crumbs:
+Description:Learn how to enable an Android's development environment.
+Version: 4.14
+type:quick start
+parent:Platforms/Linux/IntermediateLinuxDeveloper/AndroidOnLinux
+order: 2
+tags:Linux
+SkillLevel:Intermediate
+Checkpoint:AndroidOnLinuxQS
+Contributor:yaakuro
+
+[nav]
+
+Now that you've installed the Android SDK on your Linux computer, you're almost ready to start developing mobile games 
+for your Android device. However, before you can run Unreal Engine games on your Android device, you need to enable your 
+Android's development environment.
+
+[REGION:note]
+Make sure you have access to a USB cable prior to commencing the following steps.
+[/REGION]
+
+## Enabling Android's Developer Options
+[EXCERPT:AndroidLinux_2]
+1.	Connect the Android device to your Linux computer with a USB cable.
+
+1.	Open the **Settings** application on your Android device.
+
+	[REGION:lightbox]
+		[![](LinuxAndroid_Step17.png)(w:180)](LinuxAndroid_Step17.png)
+	[/REGION]
+	
+	[REGION:caption]
+		Click for full image.
+	[/REGION]
+
+1.	Tap **Developer Options** to open the Developer Options menu.
+
+	[REGION:lightbox]
+		[![](LinuxAndroid_Step18.png)(w:180)](LinuxAndroid_Step18.png)
+	[/REGION]
+	
+	[REGION:caption]
+		Click for full image.
+	[/REGION]	
+
+	[REGION:note]
+	If you don't see this item, you'll need to enable **Developer Options** on your device. For instructions on how to enable Developer Options on Android devices, 
+	visit [Android's User Guide](https://developer.android.com/studio/run/device.html).
+	[/REGION]
+
+## Enabling USB Debugging
+
+1.	From the **Developer Options** menu, enable **USB debugging**.
+
+	[REGION:lightbox]
+		[![](LinuxAndroid_Step19.png)(w:180)](LinuxAndroid_Step19.png)
+	[/REGION]
+	
+	[REGION:caption]
+		Click for full image.
+	[/REGION]
+
+1.	A window should appear, confirming your selection to enable USB debugging. After reading the prompt, go ahead and tap the **OK** button.
+
+	[REGION:lightbox]
+		[![](LinuxAndroid_Step20.png)(w:180)](LinuxAndroid_Step20.png)
+	[/REGION]
+	
+	[REGION:caption]
+		Click for full image.
+	[/REGION]
+
+	USB debugging is now enabled.  
+	![](LinuxAndroid_Step20Note.png)
+
+1.	Now, disconnect the Android device from its USB cable, plugging it back in after a few seconds. If you're connecting a device running Android 4.2.2 or higher, the following dialog will appear, asking whether you should accept an RSA key that allows debugging through your Linux computer. Go ahead and select the checkbox next to **Always allow from this computer** before tapping the **OK** button. 
+
+	[REGION:lightbox]
+		[![](LinuxAndroid_Step21.png)(w:180)](LinuxAndroid_Step21.png)
+	[/REGION]
+	
+	[REGION:caption]
+		Click for full image.
+	[/REGION]
+[/EXCERPT:AndroidLinux_2]
+
+## End Result
+
+At this point, you're ready to verify that your device can be debugged on your Linux computer. To verify that your device was setup properly, open the terminal and enter `adb devices`.
+
+![](LinuxAndroid_Step22.png)(w:900)
+
+You should see a list of devices that are attached to your Linux computer, which means that you have an Android device that's ready to run an Unreal Engine game.
+
+[REGION:note]
+	If you didn't see your device after you entered `adb devices` into the terminal, you may need to set up adb on your Linux computer. Refer to 
+	the [Android Debug Bridge User Guide](https://developer.android.com/studio/command-line/adb.html) to learn more about adb and how to set up the program on your computer. 
+	If you're certain that adb was already installed and configured on your computer, try the following troubleshooting techniques:
+
+	* You may need to disconnect and reconnect the Android device from your computer. 
+	* If you didn't select the checkbox next to **Always allow from this computer**, you may need to do so. 
+	* If the device is listed as being unauthorized, you may be to get the Android device to authorize your Linux computer for USB debugging. 
+	* In some cases, a correctly configured device might not be recognized if it's connected as a Media Device using the Media Transfer Protocol (MTP). 
+	* Also, a correctly configured device might not be recognized if it's connected as a Camera Device using the Picture Transfer Protocol (PTP).
+[/REGION]
+
+[nav]
\ No newline at end of file
diff --git a/Engine/Documentation/Source/Platforms/Linux/IntermediateLinuxDeveloper/AndroidOnLinux/3/AndroidOnLinux_3.INT.udn b/Engine/Documentation/Source/Platforms/Linux/IntermediateLinuxDeveloper/AndroidOnLinux/3/AndroidOnLinux_3.INT.udn
new file mode 100644
--- /dev/null
+++ b/Engine/Documentation/Source/Platforms/Linux/IntermediateLinuxDeveloper/AndroidOnLinux/3/AndroidOnLinux_3.INT.udn
@@ -0,0 +1,68 @@
+Availability: Docs
+Title: 3 - Creating your Project
+Crumbs:
+Description:Learn how to create an Android mobile game project in UE4.
+Version: 4.14
+type:quick start
+parent:Platforms/Linux/IntermediateLinuxDeveloper/AndroidOnLinux
+order: 3
+tags:Linux
+SkillLevel:Intermediate
+Checkpoint:AndroidOnLinuxQS
+Contributor:yaakuro
+
+[nav]
+
+By now, you've installed and configured the Android SDK, you've set up your Android's development environment, and Android Debug Bridge is able to list all devices that are connected to your Linux computer. 
+You're ready to create a new mobile project in UE4.
+
+## Creating your Mobile Project
+
+1.	Open the **Unreal Project Browser** and create a new Blueprint project with the following options.
+
+	* Select the **Third Person** Template
+	* Set the Target Hardware to be **Mobile / Tablet**
+	* Set the Graphics Level to **Scalable 3D or 2D**
+	* Select **No Starter Content**
+
+	[REGION:lightbox]
+		[![](LinuxAndroid_Step23.png)(w:720)](LinuxAndroid_Step23.png)
+	[/REGION]
+	
+	[REGION:caption]
+		Click for full image.
+	[/REGION]
+
+1.	Give the project a name, setting the Project's **Folder** location. For this tutorial, we're naming the project `LinuxAndroidProject`.
+
+	[REGION:lightbox]
+		[![](LinuxAndroid_Step24.png)(w:720)](LinuxAndroid_Step24.png)
+	[/REGION]
+	
+	[REGION:caption]
+		Click for full image.
+	[/REGION]
+
+1.	After you've set up your Project's settings, click the **Create Project** button.
+
+	[REGION:lightbox]
+		[![](LinuxAndroid_Step25.png)(w:720)](LinuxAndroid_Step25.png)
+	[/REGION]
+	
+	[REGION:caption]
+		Click for full image.
+	[/REGION]
+
+## End Result
+
+[REGION:lightbox]
+	[![](LinuxAndroid_Step3EndResult.png)(w:720)](LinuxAndroid_Step3EndResult.png)
+[/REGION]
+
+[REGION:caption]
+	Click for full image.
+[/REGION]
+
+Congratulations! You've just created a new UE4 project that's ready for use with a mobile device.
+
+[nav]
\ No newline at end of file
diff --git a/Engine/Documentation/Source/Platforms/Linux/IntermediateLinuxDeveloper/AndroidOnLinux/4/AndroidOnLinux_4.INT.udn b/Engine/Documentation/Source/Platforms/Linux/IntermediateLinuxDeveloper/AndroidOnLinux/4/AndroidOnLinux_4.INT.udn
new file mode 100644
--- /dev/null
+++ b/Engine/Documentation/Source/Platforms/Linux/IntermediateLinuxDeveloper/AndroidOnLinux/4/AndroidOnLinux_4.INT.udn
@@ -0,0 +1,109 @@
+Availability: Docs
+Title: 4 - Deploying your Level
+Crumbs:
+Description:Learn how to deploy a level onto your Android device from UE4.
+Version: 4.14
+type:quick start
+parent:Platforms/Linux/IntermediateLinuxDeveloper/AndroidOnLinux
+order: 4
+tags:Linux
+SkillLevel:Intermediate
+Checkpoint:AndroidOnLinuxQS
+Contributor:yaakuro
+
+[nav]
+
+Deploying a level means that you're going to [Cook](Engine/Deployment/Cooking) and [Launch](Engine/Deployment/Launching) a level from UE4 onto your Android device. 
+Because you're deploying a mobile application onto an Android device, there are a lot of details to consider, including (although, not limited to); your Linux computer's compilation environment, 
+your Android device's instruction set, and your device's rendering capabilities. As such, you'll review some of the most common pre-deployment considerations before deploying the **ThirdPersonExampleMap** level onto 
+your Android device. 
+
+## Pre-deployment Considerations on Linux
+
+Prior to deploying your application, refer to this list, which was written to help you avoid some common problems that might arise while you're deploying the game to your device.
+
+* Although it may seem like an obvious consideration, please make sure that the Android device is properly connected to your computer.
+* If you already haven't done so, verify that Android Debug Bridge (adb) lists your connected device(s). To learn more about adb, enter `adb --help` into the terminal. You can also refer to the [Android Debug Bridge User Guide](https://developer.android.com/studio/command-line/adb.html) to learn more about adb and how to set up the program on your computer. 
+* In order to properly build, cook, package, and deploy your game to Android, you'll need to have JDK 8 installed, with your `JAVA_HOME` environment variable set to the correct version of the JDK. To verify that `JAVA_HOME` is set properly, enter `echo $JAVA_HOME` into the terminal. If you don't see the correct JDK version, [download JDK 8](http://www.oracle.com/technetwork/java/javase/downloads/index.html), and install it on your Linux computer, making sure to set the `JAVA_HOME` environment variable to the newly installed JDK.
+* Make sure that you're building for the appropriate ARM architecture, because if you don't, your project won't deploy properly. After verifying the instruction set that'll work on your device, you can select the appropriate ARM architecture support from inside the Editor by navigating to **Edit > Project Settings... > Android > Build**, and selecting either **Support arm7** or **Support arm64**. Currently, we don't provide **Vulkan** rendering support from UE4 on Linux. 
+* Different Android devices uniquely support a variety of texture formats. If you don't know which texture formats are supported by your Android device, try setting your target texture to the Ericsson Texture Compression 1 (ETC 1) compression scheme, which should be supported by most Android devices. To learn more about supporting multiple textures, check out Android's documentation on [Creating Multiple APKs for Different GL Textures](https://developer.android.com/training/multiple-apks/texture.html). 
+
+#### Reviewing Available Texture Formats
+
+If you're interested in reading more about currently supported textures, the following chart provides a quick summary of the most commonly supported textures on Android devices:
+
+[INCLUDE:Platforms/Android/Reference#androidformats]
+
+## Deploying your Level
+
+Now that you've read through some of the most important considerations for deploying games to Android devices, you're ready to begin the deployment process.
+
+1.	First, make sure that the **ThirdPersonExampleMap** is open.
+
+	[REGION:lightbox]
+		[![](LinuxAndroid_Step26.png)(w:900)](LinuxAndroid_Step26.png)
+	[/REGION]
+	
+	[REGION:caption]
+		Click for full image.
+	[/REGION]
+
+1.	To open the **Devices** menu, click the arrow next to the **Launch** icon located in the **Main Toolbar**.
+
+	[REGION:lightbox]
+		[![](LinuxAndroid_Step27.png)(w:900)](LinuxAndroid_Step27.png)
+	[/REGION]
+	
+	[REGION:caption]
+		Click for full image.
+	[/REGION]
+
+1.	With the **Devices** menu open, select your Android Device.
+
+	[REGION:lightbox]
+		[![](LinuxAndroid_Step28.png)(w:900)](LinuxAndroid_Step28.png)
+	[/REGION]
+	
+	[REGION:caption]
+		Click for full image.
+	[/REGION]
+
+1.	At this point, your Unreal Engine will deploy the level to your Android device. 
+
+	[OBJECT:EmbeddedVideo]
+		[PARAMLITERAL:width]
+		640
+		[/PARAMLITERAL]
+		[PARAMLITERAL:height]
+		360
+		[/PARAMLITERAL]
+		[PARAMLITERAL:videoid]
+		afagQeEvf4Q
+		[/PARAMLITERAL]
+	[/OBJECT]
+
+	[REGION:note]
+	As you can see in the video, UE4 updates the progress of your application's deployment process in the Editor's lower right corner.
+	[/REGION]
+
+1.	When UE4 finishes cooking and launching your project, the level should be running on your device. You'll also be able to see an Unreal Engine icon, which should be accessible from your Android Applications Drawer.
+
+	![](LinuxAndroid_Step29.png)
+
+## Determining Supported Texture Formats
+
+If you're still wondering which texture formats are supported by your Android device, try out the following step:
+
+1.	While the game is running on your device, tap the screen with four (4) fingers.
+
+	![](LinuxAndroid_Step30.png)(w:900)
+
+	[REGION:caption]
+		A dialog should appear, displaying the texture formats that are supported by your Android device.
+	[/REGION]
+
+## End Result
+
+Congratulations! You should now have the Third Person Example Map running on your Android device.
+
+[nav]
\ No newline at end of file
diff --git a/Engine/Documentation/Source/Platforms/Linux/IntermediateLinuxDeveloper/AndroidOnLinux/5/AndroidOnLinux_5.INT.udn b/Engine/Documentation/Source/Platforms/Linux/IntermediateLinuxDeveloper/AndroidOnLinux/5/AndroidOnLinux_5.INT.udn
new file mode 100644
--- /dev/null
+++ b/Engine/Documentation/Source/Platforms/Linux/IntermediateLinuxDeveloper/AndroidOnLinux/5/AndroidOnLinux_5.INT.udn
@@ -0,0 +1,141 @@
+Availability: Docs
+Title: 5 - Packaging your Game
+Crumbs:
+Description:Learn how to package your mobile game for distribution (release).
+Version: 4.14
+type:quick start
+parent:Platforms/Linux/IntermediateLinuxDeveloper/AndroidOnLinux
+order: 5
+tags:Linux
+SkillLevel:Intermediate
+Checkpoint:AndroidOnLinuxQS
+Contributor:yaakuro
+
+[nav]
+
+Up to this point in the tutorial, you've created and deployed your mobile game project to an Android device; however, if you want to make your game distributable, you'll need to package the project. 
+The following steps show you how to package your project for distribution on the Android platform. Before we begin, it's important to note that you're going to package the project in **Development** mode. It's always a good practice to initially package your game in **Development** 
+mode for testing purposes. Once you're satisfied with your game, you can package the project in **Shipping** mode, which tells the engine to strip out console commands, statistics, and profiling tools. 
+
+[REGION:note]
+To learn more about the differences between the different build configurations made available in UE4, read through our [](Programming/Development/BuildConfigurations) page.
+[/REGION]
+
+##Packaging a Project
+[EXCERPT:AndroidLinux_5]
+1.	Opening the **File** menu, scroll down to **Package Project**, and hover your mouse curser over **Android** before selecting **Android (ETC1)**.
+
+	[REGION:lightbox]
+		[![](LinuxAndroid_Step31.png)(w:720)](LinuxAndroid_Step31.png)
+	[/REGION]
+	
+	[REGION:caption]
+		Click for full image.
+	[/REGION]
+
+1.	After the **Open Directory** menu appears, create a new folder entitled `Android_Builds` using the **New Folder** button.
+
+	[REGION:lightbox]
+		[![](LinuxAndroid_Step32.png)(w:720)](LinuxAndroid_Step32.png)
+	[/REGION]
+	
+	[REGION:caption]
+		Click for full image.
+	[/REGION]
+
+1.	Now, click the **Open** button to begin the packaging process.
+
+	[REGION:lightbox]
+		[![](LinuxAndroid_Step33.png)(w:720)](LinuxAndroid_Step33.png)
+	[/REGION]
+	
+	[REGION:caption]
+		Click for full image.
+	[/REGION]
+
+1.	You should see a progress indicator in the bottom right corner while UE4 packages your project.
+
+	[REGION:lightbox]
+		[![](LinuxAndroid_Step34.png)(w:720)](LinuxAndroid_Step34.png)
+	[/REGION]
+	
+	[REGION:caption]
+		Click for full image.
+	[/REGION]
+
+1.	At the end of the packaging process, you should see a message telling you that the packaging process has completed.
+
+	[REGION:lightbox]
+		[![](LinuxAndroid_Step35.png)(w:720)](LinuxAndroid_Step35.png)
+	[/REGION]
+	
+	[REGION:caption]
+		Click for full image.
+	[/REGION]
+
+1.	Navigate to the **Android_Builds** folder, where you should see a folder entitled **Android_ETC1**.
+
+	![](LinuxAndroid_Step36.png)
+
+1.	Open the **Android_ETC1** folder, where you'll find the files that were created during the packaging process.
+
+	![](LinuxAndroid_Step37.png)
+
+	You should see the **Installation** and **Uninstallation** shell scripts (\*.sh) along with the Android Application Package (\*.apk) that contains all of the files you need to 
+	install the project onto any Android device that supports the ETC1 compression scheme.
+
+1.	Now, open up the terminal from this folder's location by right-clicking inside of the folder and selecting **Open in Terminal**.
+
+	![](LinuxAndroid_Step38.png)
+
+1.	Enter `./Install_LinuxAndroidProject_Development-arm64.es2.sh` into the terminal to install the project onto your Android device.
+
+	[REGION:lightbox]
+		[![](LinuxAndroid_Step39.png)(w:900)](LinuxAndroid_Step39.png)
+	[/REGION]
+	
+	[REGION:caption]
+		Click for full image.
+	[/REGION]
+
+	[REGION:note]
+	If you're having trouble running the shell script, try entering `adb devices` into the terminal to verify that your device is indeed, connected to your computer.
+	[/REGION]
+
+1.	During the installation process, the shell script should flush the following output to your terminal. 
+
+	[REGION:lightbox]
+		[![](LinuxAndroid_Step40.png)(w:900)](LinuxAndroid_Step40.png)
+	[/REGION]
+	
+	[REGION:caption]
+		Click for full image.
+	[/REGION]
+
+	[REGION:note]
+	When the shell script has completed the installation process, you should see an **Installation successful** message.
+	[/REGION]
+[/EXCERPT:AndroidLinux_5]
+## End Result
+
+Now that the Installation script has finished installing the game onto your Android device, the project should now be visible inside of your device's application drawer.
+
+[REGION:lightbox]
+	[![](LinuxAndroid_Step5EndResult.png)(w:180)](LinuxAndroid_Step5EndResult.png)
+[/REGION]
+
+[REGION:caption]
+	Click for full image.
+[/REGION]
+
+Congratulations! You've learned:
+
+&#x2713; How to install and set up the Android SDK for mobile game development with UE4 on Linux.  
+&#x2713; How to enable an Android's development environment.  
+&#x2713; How to create an Android mobile game project in UE4.  
+&#x2713; How to deploy a level onto your Android device from UE4.  
+&#x2713; How to package your mobile game for distribution (release).  
+
+The next page has some suggestions on how you can apply your newly acquired skills.
+
+[nav]
\ No newline at end of file
diff --git a/Engine/Documentation/Source/Platforms/Linux/IntermediateLinuxDeveloper/AndroidOnLinux/6/AndroidOnLinux_6.INT.udn b/Engine/Documentation/Source/Platforms/Linux/IntermediateLinuxDeveloper/AndroidOnLinux/6/AndroidOnLinux_6.INT.udn
new file mode 100644
--- /dev/null
+++ b/Engine/Documentation/Source/Platforms/Linux/IntermediateLinuxDeveloper/AndroidOnLinux/6/AndroidOnLinux_6.INT.udn
@@ -0,0 +1,63 @@
+Availability: Docs
+Title: 6 - On Your Own!
+Crumbs:
+Description:Do some exercises on your own.
+Version: 4.14
+type:quick start
+parent:Platforms/Linux/IntermediateLinuxDeveloper/AndroidOnLinux
+order: 6
+tags:Linux
+SkillLevel:Intermediate 
+Checkpoint:AndroidOnLinuxQS
+Contributor:yaakuro
+
+[nav]
+
+Developing mobile games requires a high degree of attention, especially when it comes to device compatibility, mobile performance, and space limitations. 
+Applying best practices while you're designing mobile games will help you effectively refine your mobile game development skills. 
+After designing mobile content using the best practices you've just learned, learn how to preview and deploy your mobile game using an alternative deployment method. 
+Finally, at the end of this page, we've added a couple of exercises to help you practice your newly acquired mobile development skills.
+
+##Learning Some Best Practices
+
+Now that you know how to create, package, and deploy a template mobile UE4 project for your Android device; it's a good time to consider device compatibility, mobile performance, and APK best practices that could impact your ability to create mobile games from UE4 on Linux. 
+
+####Android Device Compatibility
+
+* [](Platforms\Android\DeviceCompatibility) 
+* [](Platforms\DeviceProfiles)
+
+####Mobile Performance
+
+* [](Platforms\Mobile\Performance)
+
+####Reducing APK Package Sizes
+
+There are times when you'll need to limit the size of your mobile game due to limitations imposed by targeted hardware or software platforms. If you encounter a situation where you need to reduce the size of your mobile game, you may need to reduce the size of your **Android Application Package** (APK), which is a custom file format that the Android operating system uses to package mobile applications for distribution. Check out the [Reducing APK Package Size Reference Page](Platforms/Android/ReducingAPKSize), where you'll learn ways to reduce the size of your packaged game.
+
+##Designing Mobile Content
+
+If you want to harness what you've learned about mobile game development with UE4 on Linux, you'll want to design game content that'll take advantage of your device's available rendering capabilities. Go ahead and read the following reference information:
+
+* [](Platforms\Mobile\Content) 
+* [](Platforms\Mobile\Lighting) 
+* [](Platforms\Mobile\Materials) 
+* [](Platforms\Mobile\PostProcessEffects) 
+* [](Platforms\Mobile\Textures)
+
+##Previewing your Mobile Game
+
+Prior to deploying your mobile game to an Android device, you might want to preview your game inside of the Editor. Read through our [](Platforms\Mobile\Previewer) Overview Page to learn how to preview your mobile game inside of the Editor.
+
+##Deploying Over a Wireless Network
+
+At this point, you've learned some best practices for designing mobile game content, you've previewed your mobile game with the Editor, and you're ready to deploy your game onto an Android device over a wireless network. Please read through our [Gear VR Debugging Reference Page](Platforms/GearVR/Debugging) to learn more about deploying a UE4 project to an Android device over a wireless network.
+
+##Doing Some Exercises On Your Own
+ 
+If you want to continue testing your mobile game development skills, try out a couple of exercises on your own:
+
+* Create a new mobile project using something other than the Third Person template, subsequently trying to get your newly created project running on your Android device. 
+* Try packaging your project using different Android Texture formats, documenting which ones will work for your specific device.
+
+[nav]
\ No newline at end of file
diff --git a/Engine/Documentation/Source/Platforms/Linux/IntermediateLinuxDeveloper/AndroidOnLinux/AndroidOnLinux.INT.udn b/Engine/Documentation/Source/Platforms/Linux/IntermediateLinuxDeveloper/AndroidOnLinux/AndroidOnLinux.INT.udn
new file mode 100644
--- /dev/null
+++ b/Engine/Documentation/Source/Platforms/Linux/IntermediateLinuxDeveloper/AndroidOnLinux/AndroidOnLinux.INT.udn
@@ -0,0 +1,59 @@
+Availability: Docs
+Title: Android On Linux Quick Start
+Crumbs:
+Description:Learn how to create, deploy, and package games to Android from Linux.
+Platform:Linux
+Version: 4.14
+type:quick start
+parent:Platforms/Linux
+order: 5
+tags:Linux
+SkillLevel:Intermediate
+prereq:GettingStarted/RecommendedSpecifications
+prereq:Platforms/Linux/BeginnerLinuxDeveloper/SettingUpAnUnrealWorkflow
+topic-image:androidLinuxQS_topic.png
+Checkpoint:AndroidOnLinuxQS
+Contributor:yaakuro
+
+[REGION:imagetable]
+![](DeployedAndroidGame_Linux.png)(w:900)  
+_At the end of this tutorial, you'll create, deploy, and package a game onto your Android device from Linux on Unreal Engine 4 (UE4)._
+[/REGION]
+
+If you're a developer looking to create and distribute games on Android devices from Linux; you've landed on the right page. 
+
+## Goal
+
+This guide teaches you how to configure an Android device with UE4 on Linux, ultimately showing you how to deploy and package mobile games on Android from Linux.
+
+## Objectives
+
+After working through this tutorial, you'll know the following: 
+
+* How to install and set up the Android SDK for mobile game development with UE4 on Linux. 
+* How to enable an Android's development environment. 
+* How to create an Android mobile game project in UE4. 
+* How to deploy a level onto your Android device from UE4. 
+* How to package your mobile game for distribution (release). 
+
+[OBJECT:TopicList]
+	[PARAM:icon]
+		![](%ROOT%/build_icon.png)(convert:false)
+	[/PARAM]
+	[PARAM:title]
+		Steps
+	[/PARAM]
+	[PARAM:description]
+	[/PARAM]
+	[PARAM:links]
+	[DIR(output:"checkpoint" parent:"Platforms/Linux/IntermediateLinuxDeveloper/AndroidOnLinux" org:"hierarchy")]
+	[/PARAM]
+[/OBJECT]
+
+[REGION:call_to_action]
+	[Click to Start](Platforms/Linux/IntermediateLinuxDeveloper/AndroidOnLinux/1 "%Platforms/Linux/IntermediateLinuxDeveloper/AndroidOnLinux/1:title%")
+[/REGION]
+
+[REGION:quote]
+Thanks to our community of talented contributors, we've successfully added Android support for Linux on UE4.
+[/REGION]
\ No newline at end of file
diff --git a/Engine/Documentation/Source/Platforms/Linux/Linux.INT.udn b/Engine/Documentation/Source/Platforms/Linux/Linux.INT.udn
--- a/Engine/Documentation/Source/Platforms/Linux/Linux.INT.udn
+++ b/Engine/Documentation/Source/Platforms/Linux/Linux.INT.udn
@@ -1,10 +1,10 @@
 Availability: Public
 Title: Linux Game Development
-Crumbs:%ROOT%, Platforms
+Crumbs:
 Description:Developing games for the Linux platform.
 Platform:Linux
 Navigation:platform
-Version: 4.13
+Version: 4.14
 type:landing
 parent:Platforms
 tags:Linux
@@ -16,9 +16,13 @@ you'll want to work your way through the Quick Start Guide before moving onto mo
 
 ## Getting Started
 
-%Platforms/Linux/GettingStarted:topic%
 %Platforms/Linux/BeginnerLinuxDeveloper/SettingUpAnUnrealWorkflow:topic%
-<!--%Platforms/Linux/BeginnerLinuxDeveloper/SettingUpAnIDE:topic%-->
+%Platforms/Linux/BeginnerLinuxDeveloper/SettingUpAnIDE:topic%
+
+## Advanced Topics
+
+%Platforms/Linux/GettingStarted:topic%
+%Platforms/Linux/AdvancedLinuxDeveloper/LinuxCrossCompileLegacy:topic%
 
 ## Reference
 
diff --git a/Engine/Documentation/Source/Platforms/Mobile/PostProcessEffects/MobilePostProcessEffects.INT.udn b/Engine/Documentation/Source/Platforms/Mobile/PostProcessEffects/MobilePostProcessEffects.INT.udn
--- a/Engine/Documentation/Source/Platforms/Mobile/PostProcessEffects/MobilePostProcessEffects.INT.udn
+++ b/Engine/Documentation/Source/Platforms/Mobile/PostProcessEffects/MobilePostProcessEffects.INT.udn
@@ -1,12 +1,13 @@
 Availability:Public
 Title:Post Process Effects on Mobile Platforms
-Crumbs:%ROOT%, Platforms, Platforms/Mobile
+Crumbs:
 Description:Using Post Process Effects on mobile games.
 Platform:Mobile
 Version: 4.14
 parent:Platforms/Mobile
 order:5
 type:overview
+skilllevel:Intermediate
 tags:mobile
 reference-image:Mobile_PostProcessEffect.gif
 
@@ -16,15 +17,13 @@ Mobile Post Processing is implemented as a separate code path from PC/Console Po
 To enable Post Processing for your mobile applications, please make sure that **Mobile HDR** is enabled in your **Project Settings** as most of the settings below require Mobile HDR to be enabled for the effect(s) to work.
 
 [region: lightbox]
-[![](MobileHDR_Enabled.png)(w:400 h:200)](MobileHDR_Enabled.png)
+[![](MobileHDR_Enabled.png)(w:400)](MobileHDR_Enabled.png)
 [/region]
 
 [region:caption]
 Click for full image.
 [/region]
 
-The following categories can be found in a Post Process volume (or a camera Actor). 
-
 ## Enabling Tonemapper for Mobile Devices
 To use the Tonemapper Post Processing options on Mobile Devices, you will need to set the value of **r.tonemapperfilm** from its default of **0** to **1**.  You can change, enable, or disable the Tonemapper inside of the Unreal Engine 4 (UE4) Editor in the following ways: 
 
@@ -101,128 +100,120 @@ Click for full image.
 ### White Balance
 |Post Process Property| Supported|R.tonemapperfilm Value| 
 |--------------------|-----------|--------------|
-|Temp|Yes|1|
-|Tint|Yes|1|
+|**Temp**|Yes|1|
+|**Tint**|Yes|1|
 
 ### Color Grading
 |Post Process Property| Supported|R.tonemapperfilm Value| 
 |--------------------|-----------|--------------|
-|Saturation|Yes|1|
-|Contrast|Yes|1|
-|Gamma|Yes|1|
-|Gain|Yes|1|
-|Offset|Yes|1|
+|**Saturation**|Yes|1|
+|**Contrast**|Yes|1|
+|**Gamma**|Yes|1|
+|**Gain**|Yes|1|
+|**Offset**|Yes|1|
 
 ### Film
 |Post Process Property| Supported|R.tonemapperfilm Value| 
 |--------------------|-----------|--------------|
-|Tint|Yes|0|
-|Tint Shadow|Yes|0|
-|Tint Shadow Blend|Yes|0|
-|Tint Shadow Amount|Yes|0|
-|Saturation|Yes|0|
-|Channel Mixer Red|Yes|0|
-|Channel Mixer Green|Yes|0|
-|Channel Mixer Blue|Yes|0|
-|Contrast|Yes|0|
-|Crush Shadows|Yes|0|
-|Crush Highlights|Yes|0|
-|Dynamic Range|Yes|0|
-|Slope|Yes|1|
-|Toe|Yes|1|
-|Shoulder|Yes|1|
-|Black Clip|Yes|1|
-|White Clip|Yes|1|
+|**Tint**|Yes|0|
+|**Tint Shadow**|Yes|0|
+|**Tint Shadow Blend**|Yes|0|
+|**Tint Shadow Amount**|Yes|0|
+|**Saturation**|Yes|0|
+|**Channel Mixer Red**|Yes|0|
+|**Channel Mixer Green**|Yes|0|
+|**Channel Mixer Blue**|Yes|0|
+|**Contrast**|Yes|0|
+|**Crush Shadows**|Yes|0|
+|**Crush Highlights**|Yes|0|
+|**Dynamic Range**|Yes|0|
+|**Slope**|Yes|1|
+|**Toe**|Yes|1|
+|**Shoulder**|Yes|1|
+|**Black Clip**|Yes|1|
+|**White Clip**|Yes|1|
 
 ### Scene Color
 |Post Process Property| Supported|R.tonemapperfilm Value| 
 |--------------------|-----------|--------------|
-|Scene Color Tint|Yes|1|
-|Fringe Intensity|Yes|0 & 1|
-|Vignette Intensity|Yes|0 & 1|
-|Grain Jitter |Yes|0 & 1|
-|Grain Intensity |Yes|0 & 1|
-|Color Grading Intensity |Yes|1|
-|Color Grading |Yes|1|
+|**Scene Color Tint**|Yes|1|
+|**Fringe Intensity**|Yes|0 & 1|
+|**Vignette Intensity**|Yes|0 & 1|
+|**Grain Jitter**|Yes|0 & 1|
+|**Grain Intensity**|Yes|0 & 1|
+|**Color Grading Intensity**|Yes|1|
+|**Color Grading**|Yes|1|
 
 ### Bloom
 |Post Process Property| Supported|R.tonemapperfilm Value| 
 |--------------------|-----------|--------------|
-|Intensity|Yes|0 & 1|
-|Threshold|Yes|0 & 1|
-|Size Scale |Yes|0 & 1|
-|#1 Size|No|n/a|
-|#2 Size|No|n/a|
-|#3 Size|No|n/a|
-|#4 Size|No|n/a|
-|#5 Size|No|n/a|
-|#1 Tint|No|n/a|
-|#2 Tint|No|n/a|
-|#3 Tint|No|n/a|
-|#4 Tint|No|n/a|
-|#5 Tint|No|n/a|
-|#6 Tint|No|n/a|
-|Dirt Mask Intensity|Yes|1|
-|Dirt Mask Tint |Yes|1|
-|Dirt Mask |Yes|1|
+|**Intensity**|Yes|0 & 1|
+|**Threshold**|Yes|0 & 1|
+|**Size Scale**|Yes|0 & 1|
+|**#1 Size**|No|n/a|
+|**#2 Size**|No|n/a|
+|**#3 Size**|No|n/a|
+|**#4 Size**|No|n/a|
+|**#5 Size**|No|n/a|
+|**#1 Tint**|No|n/a|
+|**#2 Tint**|No|n/a|
+|**#3 Tint**|No|n/a|
+|**#4 Tint**|No|n/a|
+|**#5 Tint**|No|n/a|
+|**#6 Tint**|No|n/a|
+|**Dirt Mask Intensity**|Yes|1|
+|**Dirt Mask Tint**|Yes|1|
+|**Dirt Mask**|Yes|1|
 
 ### Auto Exposure
 |Post Process Property| Supported|R.tonemapperfilm Value| 
 |--------------------|-----------|--------------|
-|Auto Exposure Histogram|Yes|0 & 1|
-|Auto Exposure Basic|Yes|0 & 1|
-|Low Percent|No|n/a|
-|High Percent|No|n/a|
-|Min Brightness|No|n/a|
-|Max Brightness|No|n/a|
-|Speed Up|No|n/a|
-|Speed Down|No|n/a|
-|Exposure Bias|Yes|0 & 1|
-|Histogram Log Min|No|n/a|
-|Histogram Log Max|No|n/a|
+|**Auto Exposure Histogram**|Yes|0 & 1|
+|**Auto Exposure Basic**|Yes|0 & 1|
+|**Low Percent**|No|n/a|
+|**High Percent**|No|n/a|
+|**Min Brightness**|No|n/a|
+|**Max Brightness**|No|n/a|
+|**Speed Up**|No|n/a|
+|**Speed Down**|No|n/a|
+|**Exposure Bias**|Yes|0 & 1|
+|**Histogram Log Min**|No|n/a|
+|**Histogram Log Max**|No|n/a|
 
 ### Depth Of Field
 |Post Process Property| Supported|R.tonemapperfilm Value| 
 |--------------------|-----------|--------------|
-|BokehDOF|Yes|0 |
-|GaussianDOF|Yes|0 |
-|CircleDOF|No|n/a |
-|High Quality Gaussian DoF on Mobile|Yes|1|
-|Aperture F - Stop|No|n/a|
-|Focal Distance |Yes|0|
-|Depth Blur km for 50%|No|n/a|
-|Depth Blur Radius|No|n/a|
-|Focal Region|No|n/a|
-|Near Transition Region|Yes|0|
-|Far Transition Region|Yes|0|
-|Scale|Yes|0|
-|Max Bokeh Size|No|n/a|
-|Near Blur Size|No|n/a|
-|Far Blur Size|No|n/a|
-|Shape|No|n/a|
-|Occlusion|No|n/a|
-|Color Threshold|No|n/a|
-|Sky Distance|No|n/a|
-|Vignette Size|No|n/a|
+|**BokehDOF**|Yes|0 |
+|**GaussianDOF**|Yes|0 |
+|**CircleDOF**|No|n/a |
+|**High Quality Gaussian DoF on Mobile**|Yes|1|
+|**Aperture F - Stop**|No|n/a|
+|**Focal Distance**|Yes|0|
+|**Depth Blur km for 50%**|No|n/a|
+|**Depth Blur Radius**|No|n/a|
+|**Focal Region**|No|n/a|
+|**Near Transition Region**|Yes|0|
+|**Far Transition Region**|Yes|0|
+|**Scale**|Yes|0|
+|**Max Bokeh Size**|No|n/a|
+|**Near Blur Size**|No|n/a|
+|**Far Blur Size**|No|n/a|
+|**Shape**|No|n/a|
+|**Occlusion**|No|n/a|
+|**Color Threshold**|No|n/a|
+|**Sky Distance**|No|n/a|
+|**Vignette Size**|No|n/a|
 
 
 ### Misc
 |Post Process Property| Supported|R.tonemapperfilm Value| 
 |--------------------|-----------|--------------|
-|Screen Percentage|No|n/a |
-|AA Methods|Yes|0 & 1|
+|**Screen Percentage**|No|n/a |
+|**AA Methods**|Yes|0 & 1|
 [region:note]
 The AA Method supports "None" (the fastest) and "TemporalAA" which provides a special mobile temporal AA which provides something around 2xSGSSAA with a little bit of judder in motion
 [/region]
 
-### Screen Space Reflections
-|Post Process Property| Supported|R.tonemapperfilm Value| 
-|--------------------|-----------|--------------|
-|Intensity|Yes|0 & 1|
-|Quality|Yes|0 & 1|
-|Max Roughness|Yes|0 & 1|
-
-
 ## Custom Post Process Materials
 
 Unreal Engine ships with a variety of Post Process effects, allowing you to customize the look and feel of your mobile applications. There are times, though, when you'll want to create 
@@ -250,11 +241,19 @@ You can now use the **Custom Depth** Post Process option with the Mobile renderi
 
 As this features requires Post-Processing to work correctly, **Mobile HDR** must be enabled. To enable Mobile HDR go to **Project Settings** > **Engine** > **Rendering** > **Mobile** and make sure that **Mobile HDR** is checked.
 
+[region: lightbox]
+[![](MobileHDR_Enabled.png)(w:400)](MobileHDR_Enabled.png)
+[/region]
+
+[region:caption]
+Click for full image.
+[/region]
+
 [region:warning]
 Please note that this feature also does not currently work while **Mobile MSAA** is enabled so Mobile MSAA will have to be disabled for you to see the effect on the device.
 [/region]
 
-## General Performance Tips
+## General Post Process Performance Tips
 
 * For best performance and quality, you should stick to using just Bloom and TemporalAA.
 
diff --git a/Engine/Documentation/Source/Platforms/VR/ContentSetup/VRContent.INT.udn b/Engine/Documentation/Source/Platforms/VR/ContentSetup/VRContent.INT.udn
--- a/Engine/Documentation/Source/Platforms/VR/ContentSetup/VRContent.INT.udn
+++ b/Engine/Documentation/Source/Platforms/VR/ContentSetup/VRContent.INT.udn
@@ -1,11 +1,12 @@
 Availability:Public
 Title:Virtual Reality Best Practices
-Crumbs:%ROOT%, Platforms, Platforms/VR
-Description:Information over developing for VR.
+Crumbs:
+Description:Information about developing for VR.
 Navigation:platform
 Version: 4.14
 parent:Platforms/VR
 type:overview
+tags:VR
 
 [VAR:Params]
 [PARAM:image]
@@ -49,7 +50,7 @@ platforms
 [TOC(start:2)]
 
 Developing content for Virtual Reality (VR) does require a few workflow adjustments to ensure that the user has the best 
-VR experience possible. The purpose of this guide to help point out some of the main things that you should aware of when 
+VR experience possible. The purpose of this guide is to help point out some of the main things that you should be aware of when 
 developing content for your VR projects. 
 
 ##	VR Project Settings
@@ -60,8 +61,8 @@ with **Scalable 3D or 2D** graphics and **No Starter Content** to ensure that yo
 
 ## VR .INI Settings
 
-The following .INI settings were taken from the UE4 powered VR demo Showdown. 
-These . INIsettings, if used, should be added to your projects **Config\ DefaultEngine.INI** file under **SystemSettings**. 
+The following .INI settings were taken from the UE4 powered, VR demo, Showdown. 
+These . INIsettings, if used, should be added to your project's **Config\DefaultEngine.INI** file, under **SystemSettings**. 
 
 	[SystemSettings]
 	r.VolumeMultiplier=1.0
@@ -92,14 +93,14 @@ These . INIsettings, if used, should be added to your projects **Config\ Default
 ## VR Instanced Stereo
 
 With the release of UE 4.11, Instanced Stereo Rendering was introduced to help lessen the performance impact of VR in UE4. 
-The feature can be enabled by checking the Instanced Stereo option that can be found under **Edit -> Project Settings -> Rendering -> VR**. 
+The feature can be enabled by checking the **Instanced Stereo** option that can be found under **Edit > Project Settings > Rendering > VR**. 
 
 ![](VR_ISR.png)
 
-After Instanced Stereo rendering has been enabled, it will require an engine restart and shader recompile.
- The base pass and early-z pass will work with static meshes, skeletal meshes, sprite particles and mesh particles with the feature turned on (both serial and parallel rendering paths). 
- Instanced stereo currently works on PC (DirectX) and PS4 with support for other platforms coming after the initial 4.11 release. 
- Here is a video showing Standard Stereo Rendering in VR versus Instanced Stereo Rendering in VR. 
+After **Instanced Stereo** rendering has been enabled, it will require an engine restart, and shader recompile.
+ The base pass, and early-z pass, works with static meshes, skeletal meshes, sprite particles, and mesh particles with the feature turned on (both serial and parallel rendering paths). 
+ **Instanced Stereo** currently works on PC (DirectX) and PS4 with support for other platforms coming after the initial 4.11 release. 
+ Here is a video showing Standard Stereo Rendering in VR, versus Instanced Stereo Rendering in VR. 
 
 [OBJECT:EmbeddedVideo]
 	[PARAMLITERAL:width]
@@ -115,37 +116,35 @@ After Instanced Stereo rendering has been enabled, it will require an engine res
 
 ### VR Multiview Support for Mobile
 
-Certain supported mobile devices can now support VR Multiview. Mobile Multiview is similar to Instanced Stereo rendering that is available to desktop PC's and works by providing an optimized path for Stereo Rendering on the mobile devices CPU. To enable this feature you will need to do that following:
+Certain (supported) mobile devices can now support VR Multiview. Mobile Multiview is similar to Instanced Stereo rendering that is available to desktop PC's, and works by providing an optimized path for Stereo Rendering on the mobile device's CPU. To enable this feature, you will need to do the following:
 
 [region:warning]
-Due to the experimental nature of this feature, it will only work on modern Mali-based GPU's. As this feature matures and support for more mobile devices come online, this warning will be removed.
+Due to the experimental nature of this feature, it will only work on modern Mali-based GPUs. As this feature matures, and support for more mobile devices come online, this warning will be removed.
 [/region]
 
-1. From the Main Toolbar go to **Edit** >  **Projects Settings** to bring up the Project's Settings options.
+1. From the Main Toolbar, go to **Edit** >  **Projects Settings**, to open the Editor's Project Settings.
 
-1. Then go to **Engine** > **Rendering** > **VR** and look for the **Mobile Multi-View (Experimental)** option. 
+1. Then, go to **Engine** > **Rendering** > **VR**, and look for the **Mobile Multi-View (Experimental)** option. 
 
-1. Enable the Mobile Multi-View (Experimental) option by click on the checkmark box next to its name and then restart the UE4 Editor for the changes to take effect.
+1. Enable the Mobile Multi-View (Experimental) option by click on the checkmark box next to its name, and then restart Editor for the changes to take effect.
 
     ![](VR_MobileMultiView.png)
 
-For this feature to work you will need your Android build settings to be setup like the following:
+For this feature to work, you will need your Android build settings to be setup like the following:
 
 * Make sure that **Support OpenGL ES2** is enabled in the **Android** > **Build** section. 
 * Make sure that both **Mobile HDR** and **Instanced Stereo Rendering** are disabled.
-* Currently Mobile Multiview will only work with modern Mali-based GPU's. 
-* If you package a project with this feature on but don't have a compatible GPU, it will be disabled at runtime. 
-
-
+* Currently, Mobile Multiview will only work with modern Mali-based GPU's. 
+* If you package a project with this feature on, but don't have a compatible GPU, it will be disabled at runtime. 
 
 ## VR Performance Profiling
 
 To track down assets that are too expensive for VR, you need to profile what your project is doing on the CPU and GPU as much as possible over the lifetime of your project. 
 
-* **GPU Profiling** - To active the GPU profiler, all you need to do is press  **CTRL + SHIFT + ,(Comma Key)** together while running your project. After the keys have been pressed, you will see a new window come up that should look similar to the image below. 
+* **GPU Profiling** - To activate the GPU profiler, all you need to do is press  **CTRL + SHIFT + ,(Comma Key)** together, while running your project. After the keys have been pressed, you will see a new window come up that should look similar to the image below. 
 	![](VR_GPU_Profiler.png)
 
-* **CPU Profiler** - Profiling what your project is doing on the CPU is a bit more involved than profiling for the GPU. To find out more on how to do this please refer to the [Performance Profiler](https://docs.unrealengine.com/latest/INT/Engine/Performance/Profiler/index.html) documents. 
+* **CPU Profiler** - Profiling what your project is doing on the CPU is a bit more involved than profiling for the GPU. To find out more on how to do this, please refer to the [Performance Profiler](https://docs.unrealengine.com/latest/INT/Engine/Performance/Profiler/index.html) documents. 
 
 ## VR Forward Rendering
 
@@ -157,31 +156,30 @@ To track down assets that are too expensive for VR, you need to profile what you
 Please see the [](Engine\Performance\ForwardRenderer) documentation for more information on this feature. 
 [/REGION]
 
-
 ## VR Post Process Settings
 
-Due to the demanding rendering requirements of VR, many of the advanced Post Processing features that are enabled by default should be disabled or your project could suffer
+Due to the demanding rendering requirements of VR, many of the advanced Post Processing features that are enabled by default should be disabled, or your project could suffer
  from performance problems. To accomplish this in your project, you will need to do the following.
 
  1. Add a Post Process(PP) volume to your level if there is not one already there.
- 1. Select the PP volume and in the **Post Process Volume** section enable the **Unbound** option so that the settings in the PP volume will be applied to the entire level. 
+ 1. Select the PP volume, and in the **Post Process Volume** section, enable the **Unbound** option so that the settings in the PP volume will be applied to the entire level. 
 
 	![](VR_PP_Unbound.png)
 
- 1. Expand the **Settings** of the **Post Process Volume** and then go through each section and disable any active PP settings by enabling that property by clicking on it and then set the value from the default, usually 1.0, to **0** to disable the feature.
+ 1. Expand the **Settings** of the **Post Process Volume**, and then go through each section to disable any active PP settings by enabling that property; clicking on it, and then setting the value from the default (usually 1.0) to **0**, thus disabling the feature.
 	
 	![](VR_Disable_PP_Settings.png)
 	[region:note]
-	When doing this you will not need to hit every section and set all the properties to 0. Instead first disable the really heavy hitting features like Lens Flares, Screen Space reflections, Temporal AA, SSAO, Bloom and anything else that might have an impact on performance.
+	When doing this, you won't need to hit every section and set all of the properties to 0. Instead, first disable the really heavy hitting features like Lens Flares, Screen Space reflections, Temporal AA, SSAO, Bloom and anything else that might have an impact on performance.
 	[/region]
 
 ## VR World Scale
 
-Getting the scale of your world correct is one of the most important things to help deliver the best user experience possible 
-on VR platforms. Having the wrong scale can lead to all kind of sensory issues for users and could even result in Simulation 
+Ensuring the correct scale of your world is one of the most important ways to help deliver the best user experience possible 
+on VR platforms. Having the wrong scale can lead to all kinds of sensory issues for users, and could even result in Simulation 
 Sickness. Objects are most easily viewed in VR when they are in a range of **0.75 to 3.5 Meters** from the player's camera. 
-Inside of UE4, **1 Unreal Unit (UU) is equal to 1 Centimeter (CM)**. This means that object's inside of Unreal are best viewed 
-when they are **75 UU to 350 UU** away from the player's camera when using VR. 
+Inside of UE4, **1 Unreal Unit (UU) is equal to 1 Centimeter (CM)**. This means that objects inside of Unreal are best viewed 
+when they are **75 UU to 350 UU** away from the player's camera (when using VR). 
  
 | Distance  | Distance in Unreal Units(UU) |
 | --- | --- |
@@ -191,7 +189,7 @@ when they are **75 UU to 350 UU** away from the player's camera when using VR.
 
 You can adjust the scale of your world using the **World to Meters** variable that is located under **World Settings**.
 Increasing or decreasing this number will make the user feel bigger or smaller in relation to the world around them. 
-Assuming your content was built with 1 Unreal Unit = 1 CM, setting **World To Meters** to **10** will make the world appear to be very big while setting 
+Assuming your content was built with 1 Unreal Unit = 1 CM, setting **World To Meters** to **10** will make the world appear to be very big, while setting 
 **World To Meters** to **1000** will make the world appear to be very small.
 
 ![](T_VR_World_To_Meters_00.png)
@@ -203,9 +201,9 @@ greatly affect a user's VR experience and in some cases, ruin the VR experience
 of your user having a bad VR experience, follow the best practices listed below **closely**. If you do not do this, your users 
 could end up having a very unpleasant VR experience.    
 
-* You MUST maintain framerate, and ideally, a little bit of a buffer to make sure you're always over the HMD's native framerate.
+* You must maintain framerate, and ideally, a little bit of a buffer to make sure you're always over the HMD's native framerate.
 Low framerates are another trigger for Simulation Sickness, so make sure to optimize your project as much as possible. 
-The following table shows the various HMD's UE4 supports and the target frame rates your VR project needs to run at on those devices.
+The following table shows the various HMD's UE4 supports, and the target frame rates your VR project needs to run at on those devices.
 | HMD Device  | Target Frame Rate |
 | --- | --- |
 | DK1 | 60 FPS|
@@ -215,106 +213,110 @@ The following table shows the various HMD's UE4 supports and the target frame ra
 | Gear VR | 60 FPS|
 | PSVR | Variable up to 120 FPS|
 
-* Developers make the **worst** test subjects because they are too used to using VR devices. Check your game as much as you can with as many different people as you can to make sure that you are not causing Simulation Sickness.
-* **_Avoid_** cinematic cameras or anything that takes control of the camera movements away from the player as this tends to be the biggest culprit of your user having a bad VR experience.
-* **_Do Not_** override the Field of View (FOV) manually, and **Do Not** expose this to the end user to edit. The value needs to match the physical geometry of the headset and lenses and should be automatically set through the device's SDK and internal configuration. If there is a mismatch, the world will appear to warp when you turn your head, leading to discomfort or nausea. 
-* **_Do Not_** have "Walking Bob" for a camera effect like in most first person games. Causing the camera to move up and down to mimic the movement of the human body will give player Simulation Sickness and ruin the VR experience. 
-* **_Do Not_** "Shake" the camera when trying to relay an event to the player. If a grenade goes off next to a player, a camera shake may make sense in non VR games, but in VR games, it can trigger Simulation Sickness very quickly.
+* Developers make the worst test subjects, because they are often used to using VR devices. Check your game as much as you can, with as many different people as you can, to make sure that you are not causing Simulation Sickness.
+* Avoid cinematic cameras, or anything that takes control of camera movements away from the player, as this tends to be the biggest culprit of your user having a bad VR experience.
+* Do Not override the Field of View (FOV) manually, and Do Not expose this to the end user for editing purposes. The value needs to match the physical geometry of the headset and lenses, which should be automatically set through the device's SDK and internal configuration. If there is a mismatch, the world will appear to warp when you turn your head, leading to discomfort or nausea. 
+* Do Not have "Walking Bob" for a camera effect (like in most first person games). Causing the camera to move up and down to mimic the movement of the human body will give player Simulation Sickness, ruin their VR experience. 
+* Do Not "Shake" the camera when trying to relay an event to the player. If a grenade goes off next to a player, a camera shake may make sense in non-VR games, but in VR games, it can trigger Simulation Sickness very quickly.
 * When you are designing worlds and levels for your VR game, make sure to use dimmer lights and colors than you normally would. Strong and vibrant lighting in VR games can cause Simulation Sickness to occur more quickly. Avoid this by using cooler shades and dimmer lights than you normally would.
-* Avoid stairs and use lifts instead. When the player is moving very quickly, especially up and down something like stairs, it can be very disorienting.
-* Players should all start out going full speed and not gradually accelerate to full speed and movement speed should always be at a constant rate of speed. 
-* **Do Not** use the Depth of Field or Motion Blur post process because they can greatly affect what the user is seeing and more importantly, they can give the user Simulation Sickness very quickly. 
-
+* Avoid stairs and use lifts instead. When the player is moving quickly, especially up and down stairs, it can be very disorienting.
+* Players should all start out going full speed and not gradually accelerate to full speed. Also, movement speed should always be at a constant rate of acceleration. 
+* Do Not use Depth of Field or Motion Blur post processes, because they can greatly affect what the user is seeing and, more importantly, they can give the user Simulation Sickness. 
 
-The above list is some of the main issues that will cause Simulation Sickness but there are many more. 
+Please note that this list is not an exhaustive compilation of root causes for Simulation Sickness; however, it should give you a good idea of what might be causing your players to fall ill while playing your VR game.
 
 ## VR Camera Setup
-How your VR camera is setup in UE4 will depend entirely on if your VR experience is seated or standing.
-In a seated experience you will need to artificially raise the camera origin for the character to be standing. 
-However for a standing experience, you must make sure the camera origin is at 0, on the ground. 
-You can achieve either by attaching a camera component to a scene component at the base of the character (ground level) or by setting the Eye Height to half of the 
-(negative)Cylinder Height of the collision capsule on the character.
+How your VR camera is setup in UE4 depends entirely on whether your VR experience is seated or standing.
+In a seated experience, you will need to artificially raise the camera origin for the character to be standing. 
+However, for a standing experience, you should make sure that the camera origin is at 0 (on the ground). 
+You can achieve either effects by attaching a camera component to a scene component at the base of the character (ground level), or by setting the Eye Height to half of the 
+(negative) Cylinder Height of the collision capsule on the character.
+
 ## VR Character Settings
 
 The setup for a character using a VR headset is slightly different than for a standard character. Things like character 
 Height, Width, Speed, and Camera Location all need to be slightly modified to accommodate a VR character.   
 
 [REGION:tip]
- When building objects for a VR world it is important that you make the scale of your digital object the same as their real 
- world counterpart. Making things bigger or smaller than they are in the real world can ruin the emersion that you are trying 
+ When building objects for a VR world, it is important that you make the scale of your digital object the same as their real 
+ world counterpart. Making things bigger or smaller than they are in the real world, can ruin the immersion that you are trying 
  to achieve. 
 [/REGION]
  
  **Character Height & Width**
 
-Character Height & Width should mimic real life measurements as much as possible. Using sizes that are too big or two small 
-could ruin the emersion that you are trying to achieve.  
+Character Height & Width should mimic real life measurements as much as possible. Using sizes that are too big or two small, 
+can ruin the immersion that you are trying to achieve.  
 
 |Property| UE4 Default | Recommended VR |
 |---| --- | --- |
-|Height:| 192 CM|176 CM|
-|Width:| 84 CM|68 CM|
+|Height:| 192 cm|176 cm|
+|Width:| 84 cm|68 cm|
 
 **Movement Speed**
 
-VR movement speed is a difficult property to recommend a setting for because the movement speed that you choose will mainly be 
-determined by the type of experience that you are trying to achieve. In the Elemental VR demo for example, the movement speed 
+VR movement speed is a difficult property to recommend a setting for, because the movement speed that you choose will mainly be 
+determined by the type of experience that you are trying to achieve. For example, in the Elemental VR demo, the movement speed 
 was cut to about 1/4 normal speed.  
 
 |Property| UE4 Default | Recommended VR |
 |---| --- | --- |
-Movement Speed:|60 M/S|24 M/S|
+Movement Speed:|60 m/s|24 m/s|
 
 **Camera Location**
 
-The VR camera needs to be positioned slightly lower than the base eye height to compensate for being at the characters eye level.  
+The VR camera needs to be positioned slightly lower than the base eye height to compensate for being at the character's eye level.  
 
 |Property| UE4 Default | Recommended VR |
 |---| --- | --- |
-Base Eye Height:|180 CM|160 CM|
+Base Eye Height:|180 cm|160 cm|
 
 
 
 ## VR Content Considerations
 
 When creating VR content, remember that users can look at that content from multiple angles. Here are few things that you might 
-have done in the past but need to avoid in VR:
+have done in the past, but need to avoid in VR:
 
-* **Scale** - The best thing to do about the scale of the objects in your VR world is to mimic reality as closely as you can. 
+* **Scale** - The best thing to do about the scale of the objects in your VR world, is to mimic reality as closely as you can. 
 Making objects bigger or smaller than their real world counterparts could lead to confusion or Simulation Sickness.
 
-* **Missing Polygon Faces** - In standard games, it is often acceptable and preferred to remove polygon faces from objects that cannot be seen by the player. 
-However in VR games, players have much more freedom to look around and this old practice can sometimes lead to players being able to see things that are not supposed to and should no longer be done. 
+* **Missing Polygon Faces** - In standard games, it is often acceptable (and preferred) to remove polygon faces from objects that cannot be seen by the player. 
+However, in VR games, players have much more freedom to look around, and this practice can sometimes lead to players being able to see things that they're not supposed to see.
 
-* ** Which Type of Lighting to use** - You should always use **Static lighting** and **Lightmaps** when making a VR project as this is the cheapest option to render. 
-If you need to use dynamic lighting make sure to limit the amount of dynamic lights to as few as possible and make sure that they never touch one another. 
-If you have an outdoor scene set your directional light to dynamic instead of stationary and then turn on Cascaded Shadow Maps(CSM )and set the settings to be as simple as possible while still giving you shadows. 
+* ** Which Type of Lighting to use** - You should always use **Static lighting** and **Lightmaps** when making a VR project... this is the cheapest option to render. 
+If you need to use dynamic lighting, make sure to limit the amount of dynamic lights to as few as possible, and make sure that they never touch one another. 
+If you have an outdoor scene, set your directional light to dynamic instead of stationary, and then turn on Cascaded Shadow Maps (CSM); adjusting the settings to be as simple as possible while still giving you shadows. 
 
-* ** VR & VFX** - Some VFX tricks like using [SubUV Textures](https://docs.unrealengine.com/latest/INT/Engine/Rendering/ParticleSystems/Reference/Modules/SubUV/index.html) to simulate fire or smoke do not hold up very well when viewed in VR. 
+* ** VR & VFX** - Some VFX tricks, like using [SubUV Textures](https://docs.unrealengine.com/latest/INT/Engine/Rendering/ParticleSystems/Reference/Modules/SubUV/index.html) to simulate fire or smoke, do not hold up very well when viewed in VR. 
 In many cases you are going to need to use static meshes instead of 2D particles to simulate VFX's like explosions or smoke trails.  
-Near field effects, or effects that happen very close to the camera work very well in VR but only when the effects are made up of Static Meshes particles.
+Near field effects, or effects that happen very close to the camera, work well in VR, but only when the effects are made up of Static Mesh particles.
 
-* **VR & Transparency** - In 3D graphics rendering transparency is extremely costly as transparency will generally have to be re-evaluated each frame to ensure that nothing has changed. 
-Because of this re-evaluation rendering Transparency in VR can be so costly that it's cost outweighs its benefits. 
-However to get around this issue you can use the **DitherTemporalAA** Material Function. 
-This Material Function will allow a Material to look like it is using transparency  as well as avoid common transparency issues like self-sorting.
+* **VR & Transparency** - In 3D graphics, rendering transparency is extremely costly, because transparency will generally have to be re-evaluated per-frame to ensure that nothing has changed. 
+Because of this re-evaluation, rendering Transparency in VR can be so costly, that its cost outweighs its benefits. 
+However, to get around this issue, you can use the **DitherTemporalAA** Material Function. 
+This Material Function will allow a Material to look like it is using transparency. Also, this will help you avoid common transparency issues, such as self-sorting.
 
 	![](VR_Dither_Trans_AA.png)
 
 * **Fake everything you can** - Finding clever ways to recreate expensive rendering options, like dynamic shadows or lighting, can be a huge win for hitting your performance goals in VR. 
-In Showdown having the characters cast dynamic shadows proved to be too expensive each frame and was cut. 
+In Showdown, having the characters cast dynamic shadows proved to be too expensive per-frame, which meant that dynamic shadows had to be cut from the project. 
 However, this made the characters look like they were floating while moving. 
-To fix this fake blob shadows were introduced that could dynamically adjust their position and intensity based on how close the character was to an object in the world. 
-This helped give the illusion that a character was casting shadows when they came close to the ground or other objects.
+To fix this, fake blob shadows were introduced that could dynamically adjust their position and intensity based on how close the character was to an object in the world. 
+This helped give the illusion that a character was casting shadows when they came close to the ground (or other objects).
 
 	[region:lightbox]
 	[![](VR_Fake_Shadow_Material.png)](VR_Fake_Shadow_Material.png)
 	[/region]
 
+	[region:caption]
+	Click for full image.
+	[/region]
+
 ## VR Blueprint Functions
  
  Throughout the Blueprint editor, you will find a number of different Blueprint nodes that can be used to interact with VR headsets and motion controllers. 
- Below you will find a description of what each of these nodes are and what they do.
+ Below, you will find a description of what each of these nodes are, and what they do.
 
 ### Head Mounted Display Blueprint Nodes
 
@@ -324,14 +326,14 @@ This helped give the illusion that a character was casting shadows when they cam
 Enable HMD|Switches to/from using HMD stereo rendering.|
 Enable Low Persistence Mode|Switches between low and full persistence mode.|
 Get Orientation And Position|Grabs the current orientation and position for the HMD.|
-Get Positional Tracking Camera Parameters|If the HMD has a positional tracking camera, this will return the game-world location of the camera, as well as the parameters for the bounding region of the tracking. This allows an in-game representation of the legal positional tracking range. All values will be zeroed if the camera is not available or the HMD does not support it.|
+Get Positional Tracking Camera Parameters|If the HMD has a positional tracking camera, this will return the game-world location of the camera, as well as the parameters for the bounding region of the tracking camera. This allows an in-game representation of the legal positional tracking range. All values will be zeroed if the camera is not available, or if the HMD does not support it.|
 Get Screen Percentage|Returns screen percentage to be used in VR mode.|
 Get World to Meters Scale|Returns the World to Meters scale, which corresponds to the scale of the world as perceived by the player.|
 Has Valid Tracking Position|If the HMD supports positional tracking, whether or not we are currently being tracked.|
 Is Head Mounted Display Enabled|Returns whether or not we are currently using the head mounted display.|
 Is In Low Persistence Mode| Returns _true_, if HMD is in low persistence mode; _false_ otherwise.|
-Reset Orientation And Position| Resets orientation by setting roll and pitch to 0, assuming that current yaw is forward direction and assuming current position as a 'zero-point'(for position tracking).|
-Set Clipping Planes| Sets the near and far clipping planes (NCP and FCP) for stereo rendering. This is similar to 'stereo ncp = fcp' consoule command, but NCP and FCP set by this command will no be saved to the .INI file|
+Reset Orientation And Position| Resets orientation by setting roll and pitch to 0, assuming that current yaw is forward direction and assuming current position as a 'zero-point' (for position tracking).|
+Set Clipping Planes| Sets the near and far clipping planes (NCP and FCP respectively) for stereo rendering. This is similar to the `stereo ncp = fcp` console command, but NCP and FCP set by this command will no be saved to the .INI file|
 Set World to Meters Scale| Sets the World to Meters scale, which changes the scale of the world as perceived by the player.|
 
 ### Steam VR Blueprint Nodes
@@ -339,38 +341,31 @@ Set World to Meters Scale| Sets the World to Meters scale, which changes the sca
 ![](VR_Steam_VR_BP.png)
 |Blueprint Node Name| Functionality | 
 |---| --- |
-|Get Hand Position and Orientation| Given a controller index and a hand, retursn the position and orientation of the controller.|
+|Get Hand Position and Orientation| Given a controller index and a hand, return the position and orientation of the controller.|
 |Get Tracked Device Position and Orientation|Grabs the current orientation and position for the HMD.|
-|Get Tracking Space| Gets the tracking space (e.g sitting or standing), which determines the location of the origin.|
+|Get Tracking Space| Gets the tracking space (for example, sitting or standing), which determines the location of the origin.|
 |Get Valid Tracked Device Ids| Returns an array of the currently tracked device IDs.|
-|Set Tracking Space| Sets the tracking space(e.g sitting or standing), which changes which space tracked positions are returned to.|
-
-
-
-
-
-
-
+|Set Tracking Space| Sets the tracking space (for example, sitting or standing), changing which space (tracked) positions are returned to.|
 
-## Known VR Issues 
+## Known VR Issues
 
-Due to the manner in which HMD's work some art techniques that are staples of Video Game development no longer have the impact that they once did. 
-Below you will find a list of features that might not work as expected in VR and possible workarounds to address this. 
+Due to the manner in which HMD's work, some art techniques that are staples of Video Game development, no longer have the impact that they once did. 
+Below, you will find a list of features that might not work as expected in VR, with possible workarounds to address this. 
 
-|Know Issue| Workaround | 
+|Know Issues| Workaround | 
 |---| --- |
-| Screen Space Reflections(SSR)|While SSR will still work in VR the reflections that they produce could have issues matching up to what it is reflecting in the world. Instead of using SSR use [Reflection Probes](Resources/Showcases/Reflections#reflectionenvironment) as they are much cheaper and do suffer less from reflection alignment issues.|
+| Screen Space Reflections (SSR)|While SSR will still work in VR the reflections that they produce could have issues matching up to what it is reflecting in the world. Instead of using SSR use [Reflection Probes](Resources/Showcases/Reflections#reflectionenvironment) as they are much cheaper and do suffer less from reflection alignment issues.|
 
 **Normal Mapping Issues**
 
 When viewing Normal maps on objects in VR, you will notice that they do not have the impact that they might have once had. 
 This is because Normal mapping does not account for having a binocular display or motion parallax. 
 Because of this, Normal maps will often come out looking flat when viewed with a VR device. 
-However that does not mean that you should not or will not need to use Normal maps, 
+However, that does not mean that you should not, or will not, need to use Normal maps; 
 it just means that you need to more closely evaluate if the data you are trying to convey in the Normal map would be better off being made out of geometry. 
-Below you will find some different techniques that can be used in place of Normal maps. 
+Below, you will find some different techniques that can be used in place of Normal maps. 
 
 |Shader Type | Description|
 |---|---|
-|[Parallax Mapping](http://en.wikipedia.org/wiki/Parallax_mapping)| Parallax mapping takes Normal mapping to the next level by accounting for depth cues Normal mapping does not. A Parallax mapping shader can better display depth information, making objects appear to have more detail than they do. This is because no matter what angle you look at, a Parallax map will always correct itself to show you the correct depth information from that view point. The best use of a Parallax map would be for cobblestone pathways and fine detail on surfaces. 
+|[Parallax Mapping](http://en.wikipedia.org/wiki/Parallax_mapping)| Parallax mapping takes Normal mapping to the next level, by accounting for depth cues that Normal mapping does not. A Parallax mapping shader can better display depth information, making objects appear to have more detail than they do. This is due to the fact that no matter what angle you look at, a Parallax map will always correct itself to show you the correct depth information from your view point. The best use of a Parallax map would be for cobblestone pathways and fine detail on surfaces. 
 
diff --git a/Engine/Documentation/Source/Platforms/VR/Profiling/Considerations/ConsiderationsAndInterpretations.INT.udn b/Engine/Documentation/Source/Platforms/VR/Profiling/Considerations/ConsiderationsAndInterpretations.INT.udn
--- a/Engine/Documentation/Source/Platforms/VR/Profiling/Considerations/ConsiderationsAndInterpretations.INT.udn
+++ b/Engine/Documentation/Source/Platforms/VR/Profiling/Considerations/ConsiderationsAndInterpretations.INT.udn
@@ -1,4 +1,4 @@
-Availability:Docs
+Availability:Public
 Title: VR Profiling Interpretations and Considerations
 Crumbs:%ROOT%, Platforms
 Description: Things to keep in mind when interpreting the data from profiling tools.
@@ -14,13 +14,13 @@ related: Platforms/VR/Profiling/Tools
 
 While [profiling tools](Platforms/VR/Profiling/Tools) can provide data about how your project is performing in VR, there are some points to consider when interpreting those values.
 
-It's worth bearing in mind that in certain CPU-bound cases on DX12 and Vulcan, the GPU timings can be affected by CPU waits (bubbles), so please consider that if you see unexpected results in cases where draw thread time is high. 
+Delays can appear in odd places in the GPU and CPU profilers, most often as occlusion or scene graph traversal time. If those numbers are abnormally high, it's possible that it's a false positive.
 
-Additionally, often when profiling, you'll notice that you pop between 90 Hz and 45 Hz.  The difference is accounted for because the Compositor acts much like a vsync.  If you miss framerate, it delays you until the next frame entirely.  So, you 
+[REGION:tip]
+Often when profiling, you'll notice that you pop between 90 Hz and 45 Hz.  The difference is accounted for because the Compositor acts much like a vsync.  If you miss framerate, it delays you until the next frame entirely.  So, you 
 tend to go down in brackets of 90 / n, where n is a whole number.  Because the fps reporting is an average, it won't always report as a whole number, unless the drop is sustained.  If you're bouncing back and forth between making 
 framerate and not, you may see some fraction of a jump.
-
-Note that these delays can appear in odd places in the GPU and CPU profilers, most often as occlusion or scene graph traversal time. If those numbers are abnormally high, it's possible that it's a false positive.
+[/REGION]
 
 If you'd like to remove this for testing, sometimes it's useful to run the game emulating stereo rendering, instead of running in the device itself.  To do so:
 
