Copying //UE4/Dev-Documentation to //UE4/Samples-Main (Source: //UE4/Dev-Documentation @ 3372893)

diff --git a/Engine/Documentation/Source/Platforms/VR/Profiling/Overview/Overview.INT.udn b/Engine/Documentation/Source/Platforms/VR/Profiling/Overview/Overview.INT.udn
--- a/Engine/Documentation/Source/Platforms/VR/Profiling/Overview/Overview.INT.udn
+++ b/Engine/Documentation/Source/Platforms/VR/Profiling/Overview/Overview.INT.udn
@@ -1,4 +1,4 @@
-Availability:Docs
+Availability:Public
 Title: VR Performance Testing
 Crumbs:%ROOT%, Platforms
 Description: Overview of the key principles for performance and profiling of UE4 virtual reality projects.
@@ -13,8 +13,8 @@ In general, maintaining VR frame rates is extremely challenging even on recommen
 frames can make your experience uncomfortable for the user.
 
 With that in mind, you need to leave yourself a buffer to ensure you can absorb frames which may run longer than average. We don't allow pipelining multiple frames due to the input/tracking latency it 
-would introduce. This means the game thread is only one frame ahead of the render thread at all times and the GPU and render thread are in sync on the same frame. This means that you don't get 11 ms, you get closer to 8.5.
- Long frame times in the game thread due to gameplay code overhead, the render thread due to draw call overhead, or the GPU due to transform or shading overhead can cause you to miss frame time. All three 
+would introduce. This means the game thread is only one frame ahead of the render thread at all times and the GPU and render thread are in sync on the same frame. You have slightly less than 11ms of
+GPU time, because of the overhead of reprojection. Long frame times in the game thread due to gameplay code overhead, the render thread due to draw call overhead, or the GPU due to transform or shading overhead can cause you to miss frame time. All three 
  need to be considered and balanced.
 
 If you follow the console variable settings in the [Best Practices guide](Platforms/Oculus/BestPractices), your render thread performance should be around where it needs to be for performance. From there, 
@@ -23,10 +23,12 @@ it's important to optimize your content and gameplay code to ensure that you mak
 General practices for simplifying content include the following.
 * Avoid dynamic lights and shadows. 
 * Avoid heavy use of translucency.
-* Instance where you can.
+* Instance in visible batches. If one element of your instanced group is visible, the whole group will be drawn.
 * Make LODs for everything.
 * Keep material complexity and number of materials per object low. 
 * Bake out anything that you can get away with. 
+* Avoid large geometry that can encompass the player.
+* Use precomputed visibility volumes where possible.
 
 [REGION:tip]
 See the [Bullet Train](https://s3.amazonaws.com/epiccms_epiccms_bucket/Resources/files/Going+Off+the+Rails.zip) and [Showdown](https://cdn2.unrealengine.com/Resources/files/UE4-Integration-and-Demos_OC-100270768.pptx) slide decks for 
@@ -37,7 +39,8 @@ Be sure to profile your work throughout the project and ensure you're always mai
 to meet the target performance a few days before shipping.
 
 Reducing draw calls really comes down to two techniques.  Make sure that you combine as many materials as possible and make sure to make heavy use of instanced static mesh components.  
-On PC, there's some potential in trying to automatically batch these things together, but on mobile, the cost is too high.
+Again, when instancing, remember that whatever is instanced will always be drawn in its entirety if one part is visible. This can greatly diminish the benefits of instancing in certain cases,
+so be sure to instance in batches where most of the objects in the batch will likely be visible or invisible from the same point.
 
 The other thing that you may consider doing is using cull volumes to ensure that non-visible rooms are completely culled out when outside of them.  Our occlusion is slightly more liberal in what it 
 chooses to draw, and cull volumes help make this a bit more conservative in those choices.  Additionally, you can toggle visibility on Actors (or entire sublevels) to get the same effect, if your content is built like that.
@@ -48,11 +51,13 @@ Another important thing to remember is that if something has an alpha of zero, i
 
 ## Getting General Stats
 
-When in the game, there are a few ways to get general stats:
+Both Oculus and SteamVR provide external tools to watch the performance. It is advisable to use these to see actual frame timings with the compositor overhead.
+
+From the UE4 side, when in the game, there are a few ways to get general stats:
 
 | Console Command | Description |
 | --------------- | ----------- |
-| stat unit | This will give you general game thread, draw thread, and GPU time, as well as overall frame timing.  This is most useful for gathering whether the total frametime is where it needs to be, and what the game thread time is, but it should not be used for Draw Thread and GPU time.  Oculus has a variable delay that it adds in order to try and time GPU work with its compositor. |
+| stat unit | This will give you general game thread, draw thread, and GPU time, as well as overall frame timing.  This is most useful for gathering whether the total frametime is where it needs to be, and what the game thread time is, but it should not be used for Draw Thread and GPU time. |
 | startfpschart / stopfpschart | If you want to get an idea of what percentage of time is spent above 90Hz, you can run these commands.  It'll capture and aggregate stats data over the window between when you started and stopped, and then dump a file that has bucketed framerate information.  Note that since sometimes the game reports slightly under 90Hz when actually at 90, it's best to check the 80+ buckets to determine the actual time spent at framerate. |
 | stat gpu | This command, added in 4.14,  gives similar stats to the GPU profiler, but in a form that you can watch and monitor from in the game.  Great for checking quick costs on GPU work. |
 
diff --git a/Engine/Documentation/Source/Platforms/VR/Profiling/Tools/ProfilingTools.INT.udn b/Engine/Documentation/Source/Platforms/VR/Profiling/Tools/ProfilingTools.INT.udn
--- a/Engine/Documentation/Source/Platforms/VR/Profiling/Tools/ProfilingTools.INT.udn
+++ b/Engine/Documentation/Source/Platforms/VR/Profiling/Tools/ProfilingTools.INT.udn
@@ -1,4 +1,4 @@
-Availability:Docs
+Availability:Public
 Title: VR Profiling Tools
 Crumbs:%ROOT%, Platforms
 Description:Internal and external tools you can use for profiling your Unreal Engine VR project
@@ -20,17 +20,6 @@ For GPU profiling, there are more options.  First of all, and most specifically,
 
 Additionally, you could go further drilling into more systemic problems using both GPUView and Nsight. 
 
-## RenderDoc
-
-RenderDoc is an external tool you can also use for profiling Unreal Engine projects. 
-To use it with UE4:
-
-1. Download the latest [RenderDoc](https://renderdoc.org/builds).
-1. In RenderDoc, configure it to launch your game (e.g. UE4Editor.exe, with appropriate command line arguments).
-1. Run `ToggleDrawEvents` from the command line of your app to get human readable events.
-1. Press F12 to capture a frame.
-
-
 ## Oculus PerfHud
 
 Oculus PerfHud is a tool that can be displayed in game, whether in editor or in packaged builds, that will help verify the actual CPU and GPU timings, excluding any throttling done by the application.  There's a great overview of the screens 
@@ -52,3 +41,23 @@ where N is:
 |	3	| Shows render timing info for application |
 |	4	| Shows render timing info for OVR compositor |
 |	5	| Shows SDK & HMD version Info |
+
+## SteamVR Frame Timing Tool
+
+The SteamVR Frame Timing will allow you see how your VR project is performing in both the UE4 Editor or in a packaged build. This tool will help verify the actual CPU and GPU timings, excluding any throttling done by the application. There's a great overview of how to open the tool as well as what all of the values relate to on the 
+[Valve Developer Community]( https://developer.valvesoftware.com/wiki/SteamVR/Frame_Timing) page. Please note that there is nothing inside of UE4 that needs to be done in order for the SteamVR Frame Timing tools to work. All of the tool functionality and use is done from the SteamVR Application.
+
+[region:note]
+Please note that there is nothing inside of UE4 that needs to be done in order for the SteamVR Frame Timing tools to work. 
+[/region]
+
+
+## RenderDoc
+
+RenderDoc is an external tool you can also use for profiling Unreal Engine projects. 
+To use it with UE4:
+
+1. Download the latest [RenderDoc](https://renderdoc.org/builds).
+1. In RenderDoc, configure it to launch your game (e.g. UE4Editor.exe, with appropriate command line arguments).
+1. Run `ToggleDrawEvents` from the command line of your app to get human readable events.
+1. Press F12 to capture a frame.
\ No newline at end of file
diff --git a/Engine/Documentation/Source/Platforms/VR/SplashScreens/VRSplashScreens.INT.udn b/Engine/Documentation/Source/Platforms/VR/SplashScreens/VRSplashScreens.INT.udn
--- a/Engine/Documentation/Source/Platforms/VR/SplashScreens/VRSplashScreens.INT.udn
+++ b/Engine/Documentation/Source/Platforms/VR/SplashScreens/VRSplashScreens.INT.udn
@@ -10,7 +10,7 @@ Skilllevel: Intermediate
 
 [TOC(start:2)]
 
-Unreal Engine 4 (UE4) offers support for displaying movies or images while a level loads when using Oculus, Gear VR, SteamVR, and PSVR Head Mounted Displays (HMD). The loading move or image will run on the rendering thread and can also help to mask frame rate issues while your content is loading. In the following document, we will take a look at how you can set these VR splash screens in your UE4 projects.
+Unreal Engine 4 (UE4) offers support for displaying movies or images while a level loads when using Oculus, Gear VR, SteamVR, and PSVR Head Mounted Displays (HMDs). The **loading** movie or image will run on the rendering thread and can also help to mask frame rate issues while your content is loading. In the following document, we will take a look at how you can set these VR splash screens in your UE4 projects.
 
 ## VR Splash Screen Supported Platforms 
 
@@ -23,7 +23,7 @@ VR Splash Screens will only work when used with the following VR platforms.
 
 ## VR Splash Screen Project Setup
 
-Before you can setup and use the VR splash screens you will first need to create a new **Game Mode** and **Pawn** Blueprint to make it easier to work with the VR splash screens. In the following section, we will go over what Blueprints need to be created and how they need to be setup.
+Before you can setup and use the VR splash screens you will first need to create a new **Game Mode** and **Pawn** Blueprint to make it easier to work with the VR splash screens. In the following section, we will go over what Blueprints need to be created and how they need to be set up.
 
 1. Right - click in the **Content Browser**, and from the **Create Basic Asset** list select the **Blueprint Class** option.
 
@@ -33,7 +33,7 @@ Before you can setup and use the VR splash screens you will first need to create
 
     ![](VRLS_Intro_02.png)
 
-1. Repeat the above process but this time select the **Pawn** option from the Pick Parent Class dialogue box and call this new Pawn, **VR_Pawn**.
+1. Repeat the above process, but this time select the **Pawn** option from the Pick Parent Class dialogue box and call this new Pawn, **VR_Pawn**.
 
     ![](VRLS_Intro_03.png)
 
@@ -41,7 +41,7 @@ Before you can setup and use the VR splash screens you will first need to create
 
     ![](VRLS_Intro_04.png)
 
-Now that the Pawn and Game Mode have been set up, in the next section we will take a look at how to setup the VR splash screens in the Pawn's Event Graph.
+Now that the Pawn and Game Mode have been set up, we will take a look at how to set up the VR splash screens in the Pawn's Event Graph.
 
 ## VR Splash Screen Setup
 
@@ -50,7 +50,7 @@ In the following section, we will go over how you set up a VR splash screen in y
 [region:note]
 Please note that the following is showing the simplest way to set VR splash screens up. While each project has unique needs, the overall setup for VR splash screens will be very similar to what is described.
 [/region]
-1. Open up the VR_Pawn and go to the **Event Graph**. Right - click inside the Event Graph and from the displayed menu, search for and add the following Blueprint nodes.
+1. Open up VR_Pawn and go to the **Event Graph**. Right-click inside the Event Graph and search for and add the following Blueprint nodes:
 
     * Enable Auto Loading Splash Screen
     * Set Splash Screen
@@ -66,14 +66,14 @@ Please note that the following is showing the simplest way to set VR splash scre
 
     ![](VRLS_Setup_03.png)
 
-1. The Set Splash Screen node is where you specify which Texture to use for the VR splash screen. To set a set a VR splash screen Texture, select a Texture in the Content Browser and then press the small white arrow in the Texture section of the Ste Splash Screen node to use that Texture as the VR splash screen.
+1. The Set Splash Screen node is where you specify which Texture to use for the VR splash screen. To set a VR splash screen Texture, select a Texture in the Content Browser and then press the small white arrow in the Texture section of the Set Splash Screen node to use that Texture as the VR splash screen.
 
     ![](VRLS_Setup_04.png)
     [region:note]
     For this example, a noise Texture was created in Photoshop and then imported into UE4 and used. Feel free to use any Texture you want for this section.
     [/region]
     [region:warning]
-    If you are targeting **SteamVR** you will need to make sure that you set the **Compression Settings** of the Texture you are using for the splash screen to **UserInterface2D(RGBA)** or it will not be displayed when the splash screen is called.
+    If you are targeting **SteamVR**, you will need to make sure that you set the **Compression Settings** of the Texture you are using for the splash screen to **UserInterface2D(RGBA)** or it will not be displayed when the splash screen is called.
 
     ![](VIL_SteamVR_Comp_Settings.png)
     [/region]
@@ -86,9 +86,9 @@ With the VR splash screen now set up in the next section, we will take a look at
 
 ## VR Splash Screen Use
 
-Now that we have the Pawn setup to use the VR splash screens, in this section we will go over how to display the VR Splash Screens when a level is loaded.
+Now that we have the Pawn set up to use the VR splash screens, we will go over how to display the VR Splash Screens when a level is loaded.
 
-1. For this to work you will need to create **two** levels that we can load. To create the levels go to the **Main Toolbar** and select the **File** > **New Level** option. When the **New Level** popup box is displayed, select the **Default** level option.  
+1. For this to work, you will need to create **two** levels that we can load. To create the levels, go to the **Main Toolbar** and select the **File** > **New Level** option. When the **New Level** popup box is displayed, select the **Default** level option.
 
     [REGION:lightbox]
     [![](VLI_NewLevel_00.png)(w:746)](VLI_NewLevel_00.png)
@@ -98,7 +98,7 @@ Now that we have the Pawn setup to use the VR splash screens, in this section we
     Click for full image.
     [/REGION]
 
-1. Inside of the **Content Browser** press the **Save All** button and name the new level **Level_00**. Once it has saved, **Right - Click** on Level_00 and from the context menu select the **Duplicate** option and name it **Level_01**. When all of that is completed, you should now have two new level in the Content Browser like in the image below. 
+1. Inside of the **Content Browser**, press the **Save All** button and name the new level **Level_00**. Once it has saved, right-click on Level_00, and select the **Duplicate** option. Name the duplicate **Level_01**. When all of that is completed, you should now have two new levels in the Content Browser, like in the image below:
 
     [REGION:lightbox]
     [![](VLI_NewLevel_01.png)(w:746)](VLI_NewLevel_01.png)
@@ -108,11 +108,11 @@ Now that we have the Pawn setup to use the VR splash screens, in this section we
     Click for full image.
     [/REGION]
 
-1. Open up each level and go to **World Settings** > **Game Mode** and the **Game Mode Override** to the VR_GameMode Blueprint that was created. This way we can be sure that the correct Pawn will be used when the level is loaded.
+1. Open up each level and go to **World Settings** > **Game Mode** and the **Game Mode Override** to the VR_GameMode Blueprint that was created. This way, we can be sure that the correct Pawn will be used when the level is loaded.
 
     ![](VLI_NewLevel_06.png)
 
-1. Open up the VR_Pawn Blueprint, if not already open, and in the Event, graph add key inputs for the **R** and **T** keys.
+1. Open up the VR_Pawn Blueprint, if not already open, and in the Event Graph, add key inputs for the **R** and **T** keys.
 
     [REGION:lightbox]
     [![](VLI_NewLevel_02.png)(w:746)](VLI_NewLevel_02.png)
@@ -122,7 +122,7 @@ Now that we have the Pawn setup to use the VR splash screens, in this section we
     Click for full image.
     [/REGION]
 
-1.  Right - Click on the Event Graph and search for the **Open Level** Blueprint node and add it to the Event Graph. Connect the **Pressed** output for the T key to the Open Level input and in the **Level Name** option input **Level_00**.
+1. Right-click on the Event Graph and search for the **Open Level** Blueprint node. Add it to the Event Graph. Connect the **Pressed** output for the T key to the Open Level input and enter **Level_00** in the **Level Name** field.
 
     [REGION:lightbox]
     [![](VLI_NewLevel_03.png)(w:746)](VLI_NewLevel_03.png)
@@ -132,7 +132,7 @@ Now that we have the Pawn setup to use the VR splash screens, in this section we
     Click for full image.
     [/REGION]
 
-1. Select the Open Level node and then press **CTRL + W** on the keyboard to duplicate it. Connect the Pressed output for the **R** key to the new Open Level node and change the name of the Level to **Level_01**.
+1. Select the Open Level node and then press **CTRL + W** on the keyboard to duplicate it. Connect the **Pressed** output for the R key to the new Open Level node and change the name of the Level to **Level_01**.
 
     [REGION:lightbox]
     [![](VLI_NewLevel_04.png)(w:746)](VLI_NewLevel_04.png)
@@ -152,7 +152,7 @@ Now that we have the Pawn setup to use the VR splash screens, in this section we
     Click for full image.
     [/REGION]
 
-1. Put your HMD on and then press the **R** key to load Level_01. While the level is being loaded you should see the image that you specified in the Texture input of the Set Splash Screen node being displayed. Once Level_01 loads press the **T** key to load Level_00, and you should now see the same image being displayed. 
+1. Put your HMD on and then press the **R** key to load Level_01. While the level is being loaded, you should see the image that you specified in the Texture input of the Set Splash Screen node. Once Level_01 loads, press the **T** key to load Level_00, and you should now see the same image being displayed.
 
     [OBJECT:EmbeddedVideo]
         [PARAMLITERAL:width]
@@ -170,7 +170,7 @@ With the ability to change levels on button presses now working, the next step w
 
 ## VR Splash Screen Node Reference
 
-The following breakdown covers what each one of the properties on the VR splash screen nodes does. When looking at the VR splash screen nodes, you will notice that there is a generic type and a type specific to Oculus. You should only use the Oculus-specific types when you are targeting Oculus HMD's. All other HMD's should use the generic type.
+The following breakdown covers what each one of the properties on the VR splash screen nodes does. When looking at the VR splash screen nodes, you will notice that there is a generic type and a type specific to Oculus. You should only use the Oculus-specific types when you are targeting Oculus HMDs. All other HMDs should use the generic type.
 
 ### Generic Splash Screens
 
@@ -196,10 +196,10 @@ The following breakdown covers what each one of the properties on the VR splash
 
     |Setting  Name| Description|
     |---------------|----------------------------------------------------|
-    |Texture|Texture to use fo the splash. Must be in B8R8G8A8 format.|
+    |Texture|Texture to use for the splash. Must be in B8R8G8A8 format.|
     |Scale|Scale of the Texture.|
     |Offest|Position from which to start rendering the Texture.|
-    |Show Loading Moive|Show a loading movie.|
+    |Show Loading Movie|Show a "loading" movie.|
     |Show on Set|Show the Texture specified when the compositor is called.|
 
 * Show Splash Screen
@@ -212,7 +212,7 @@ The following breakdown covers what each one of the properties on the VR splash
 
 ### Oculus Specific Splash Screens
 
-The following breakdown covers what each of the Oculus-specific splash screen nodes will do. Keep in mind that the following nodes should only be used when you are targeting Oculus HMD's. 
+The following breakdown covers what each of the Oculus-specific splash screen nodes will do. Keep in mind that the following nodes should only be used when you are targeting Oculus HMDs.
 
 * Add Loading Splash Screen
 
@@ -222,10 +222,10 @@ The following breakdown covers what each of the Oculus-specific splash screen no
     |---------------|----------------------------------------------------|
     |Texture|A Texture asset to be used for the splash screen. GearVR uses it as a path for loading icon; all other params are currently ignored by GearVR.|
     |Translation in Meters|Initial translation of the center of the splash screen in meters.|
-    |Rotation|Initial rotation of the splash screen, with the origin at the center of the splash screen.|
+    |Rotation|Initial rotation of the splash screen, with the pivot at the center of the splash screen.|
     |Size in Meters|Size, in meters, of the quad with the splash screen.|
-    |Delta Rotation|Incremental rotation, that is added each 2nd frame to the quad transform. The quad is rotated around the center of the quad.|
-    |Ckear Before Add |If true clears splash before adding a new one.|
+    |Delta Rotation|Incremental rotation added each 2nd frame to the quad's transform. The quad pivots around its center.|
+    |Clear Before Add|If true, clears splash before adding a new one.|
 
 * Clear Loading Splash Screen
 
diff --git a/Engine/Documentation/Source/Platforms/VR/StereoPanoramicCapture/QuickStart/CreatingViewerMaterial/CreatingViewerMaterial.INT.udn b/Engine/Documentation/Source/Platforms/VR/StereoPanoramicCapture/QuickStart/CreatingViewerMaterial/CreatingViewerMaterial.INT.udn
new file mode 100644
--- /dev/null
+++ b/Engine/Documentation/Source/Platforms/VR/StereoPanoramicCapture/QuickStart/CreatingViewerMaterial/CreatingViewerMaterial.INT.udn
@@ -0,0 +1,328 @@
+Availability: Public
+Title:3. Creating a Material for Viewing Stereo Images
+Crumbs: 
+Description:Here we will create and set up a new Material to view our stereo screenshots. 
+Version: 4.13
+SkillLevel: Intermediate
+parent:Platforms/VR/StereoPanoramicCapture/QuickStart
+type:quick start
+order:3
+tags:VR
+Tags:Cinematics
+Tags:GearVR
+Tags:Google VR	
+Tags:SteamVR
+Tags:Oculus Rift
+Tags:Virtual Reality
+Tags:Plugins
+Tags:Sequencer
+
+[NAV]
+
+In this section, we will create and set up a Material that will allow us to view the stereo images that were created in the previous section. 
+
+1. Locate the stereo images that were created in the previous section by going to your projects saved folder and then going to **StereoPanoramaImages\Saved\StereoPanorama\(Date & Time it was taken)**.
+
+    [REGION:lightbox]
+    [![](Platforms\VR\StereoPanoramicCapture\QuickStart\TakingStereoImages\SPHT_Rendered_Images.png)(w:473)](Platforms\VR\StereoPanoramicCapture\QuickStart\TakingStereoImages\SPHT_Rendered_Images.png)
+    
+    [/REGION]
+
+    [REGION:caption]
+    Click for full image.
+    [/REGION]
+
+1. Once you locate both images, select both images and then drag them from their location in Windows into the Content Browser which will import the images into your UE4 project.
+
+    [REGION:lightbox]
+    [![](SPHT_Import_Stereo_Images.png)(w:639)](SPHT_Import_Stereo_Images.png)
+    [/REGION]
+
+    [REGION:caption]
+    Click for full image.
+    [/REGION]
+
+1. Once both images have been imported, double-click on both of the images to open them up in the Texture Viewer and once opened, go to the **Level Of Detail** section and under **Mip Gen Settings**, change the setting from **From Texture Group** to **No Mipmaps**.
+
+    [REGION:lightbox]
+    [![](SPHT_Mip_Gen_Settings.png)(w:439)](SPHT_Mip_Gen_Settings.png)
+    [/REGION]
+
+    [REGION:caption]
+    Click for full image.
+    [/REGION]
+
+    [region:note]
+    Failing to do this could result in the textures showing a seam when viewed in stereo due to the distance this texture will be from the camera. 
+    [/region]
+
+1. Now Right - click in the Content Browser and from the **Create Basic Asset** section click on the Material option and give the newly created Material a name of **MAT_StereoViewer**.
+
+    [REGION:lightbox]
+    [![](SPHT_Create_Stereo_Material.png)(w:612)](SPHT_Create_Stereo_Material.png)
+    [/REGION]
+
+    [REGION:caption]
+    Click for full image.
+    [/REGION]
+
+1. Double-click on the MAT_StereoViewer Material to open it up and in the **Details** panel set the following options.
+
+    ![](SPHT_Material_Settings.png)
+
+    * Shading Model: Unlit
+    * Fully Rough: True
+    * Use Full Precision: True
+    * Use LIghtmap Directionality: False
+
+1. Now click on the **Copy Material Expression Graph** option that is below and copy the provided text and paste it into the Material Editor. After the text has been pasted into the Material Graph, make sure to wire up the output of the **If** node into the **Emissive Color** on the Main Material node. When completed your Material should look like the following image.
+
+    [Include:Platforms\VR\StereoPanoramicCapture\QuickStart\CreatingViewerMaterial\#SPMaterial]
+
+1. Now that the viewer Material has been set up and connected, locate the left and right eye Textures in the Content Browser and then load the Textures into the corresponding eye Texture input and then press the **Apply** AND **Save** buttons to compile and save the Material. 
+
+    [REGION:lightbox]
+    [![](SPHT_Completed_Material.png)(w:439)](SPHT_Completed_Material.png)
+    [/REGION]
+
+    [REGION:caption]
+    Click for full image.
+    [/REGION]
+
+
+With the Material now set up and ready to be used, in the next section, we will look at how to set up a level to display the Material when the level is run.
+
+
+
+[NAV]
+
+<!-----
+[EXCERPT:SPMaterial]    
+[OBJECT:Snippet]
+    [PARAMLITERAL:id]
+        SPMaterial
+    [/PARAMLITERAL]
+    [PARAMLITERAL:type]
+        material
+    [/PARAMLITERAL]
+    [PARAM:image]
+        ![](SPHT_Stereo_Viewer_Material.png)
+    [/PARAM]
+    [PARAM:snippet]
+                            Begin Object Class=MaterialGraphNode Name="MaterialGraphNode_0"
+               Begin Object Class=MaterialExpressionIf Name="MaterialExpressionIf_0"
+               End Object
+               Begin Object Name="MaterialExpressionIf_0"
+                  A=(Expression=MaterialExpressionMaterialFunctionCall'MaterialGraphNode_2.MaterialExpressionMaterialFunctionCall_0')
+                  B=(Expression=MaterialExpressionConstant'MaterialGraphNode_1.MaterialExpressionConstant_0')
+                  AGreaterThanB=(Expression=MaterialExpressionTextureSampleParameter2D'MaterialGraphNode_7.MaterialExpressionTextureSampleParameter2D_0',Mask=1,MaskR=1,MaskG=1,MaskB=1)
+                  ALessThanB=(Expression=MaterialExpressionTextureSampleParameter2D'MaterialGraphNode_8.MaterialExpressionTextureSampleParameter2D_1',Mask=1,MaskR=1,MaskG=1,MaskB=1)
+                  MaterialExpressionEditorX=64
+                  MaterialExpressionEditorY=144
+                  MaterialExpressionGuid=81A589084F901693C24D6188C7E9963C
+                  Material=PreviewMaterial'/Engine/Transient.MAT_Stereo_Image'
+               End Object
+               MaterialExpression=MaterialExpressionIf'MaterialExpressionIf_0'
+               NodePosX=64
+               NodePosY=144
+               NodeGuid=5B7FA1FE457D602E7A7E578891582634
+               CustomProperties Pin (PinId=A0EE7D184D1901AC89D7CCB7D7B2AF54,PinName="A",PinType.PinCategory="required",PinType.PinSubCategoryObject=None,PinType.PinSubCategoryMemberReference=(),PinType.bIsArray=False,PinType.bIsReference=False,PinType.bIsConst=False,PinType.bIsWeakPointer=False,LinkedTo=(MaterialGraphNode_2 176CDC95486F3BF845430099EF2BE7FE,),PersistentGuid=00000000000000000000000000000000,bHidden=False,bNotConnectable=False,bDefaultValueIsReadOnly=False,bDefaultValueIsIgnored=False,bAdvancedView=False,)
+               CustomProperties Pin (PinId=D37D419748FE491453C139AF821D8F11,PinName="B",PinType.PinCategory="optional",PinType.PinSubCategoryObject=None,PinType.PinSubCategoryMemberReference=(),PinType.bIsArray=False,PinType.bIsReference=False,PinType.bIsConst=False,PinType.bIsWeakPointer=False,LinkedTo=(MaterialGraphNode_1 62DD318C41832241F7AB73BA06FB692E,),PersistentGuid=00000000000000000000000000000000,bHidden=False,bNotConnectable=False,bDefaultValueIsReadOnly=False,bDefaultValueIsIgnored=False,bAdvancedView=False,)
+               CustomProperties Pin (PinId=44EC5EE54EACAD4AA705D8866FE16EB5,PinName="A > B",PinType.PinCategory="required",PinType.PinSubCategoryObject=None,PinType.PinSubCategoryMemberReference=(),PinType.bIsArray=False,PinType.bIsReference=False,PinType.bIsConst=False,PinType.bIsWeakPointer=False,LinkedTo=(MaterialGraphNode_7 6073EDCF453D52887F1ED597234FAF0C,),PersistentGuid=00000000000000000000000000000000,bHidden=False,bNotConnectable=False,bDefaultValueIsReadOnly=False,bDefaultValueIsIgnored=False,bAdvancedView=False,)
+               CustomProperties Pin (PinId=6E32A98741E5BFCB34BE82870AE85A02,PinName="A == B",PinType.PinCategory="optional",PinType.PinSubCategoryObject=None,PinType.PinSubCategoryMemberReference=(),PinType.bIsArray=False,PinType.bIsReference=False,PinType.bIsConst=False,PinType.bIsWeakPointer=False,PersistentGuid=00000000000000000000000000000000,bHidden=False,bNotConnectable=False,bDefaultValueIsReadOnly=False,bDefaultValueIsIgnored=False,bAdvancedView=False,)
+               CustomProperties Pin (PinId=F937930D47BA8EA98B9510B0C025B9A4,PinName="A < B",PinType.PinCategory="required",PinType.PinSubCategoryObject=None,PinType.PinSubCategoryMemberReference=(),PinType.bIsArray=False,PinType.bIsReference=False,PinType.bIsConst=False,PinType.bIsWeakPointer=False,LinkedTo=(MaterialGraphNode_8 CDF923064C8491C8DD875A9FD733B329,),PersistentGuid=00000000000000000000000000000000,bHidden=False,bNotConnectable=False,bDefaultValueIsReadOnly=False,bDefaultValueIsIgnored=False,bAdvancedView=False,)
+               CustomProperties Pin (PinId=33CE8B8B4E0E854F7C6164AD5735E5F2,PinName="Output",PinFriendlyName=" ",Direction="EGPD_Output",PinType.PinSubCategoryObject=None,PinType.PinSubCategoryMemberReference=(),PinType.bIsArray=False,PinType.bIsReference=False,PinType.bIsConst=False,PinType.bIsWeakPointer=False,LinkedTo=(MaterialGraphNode_Root_0 C45D0E1F4E227CF5D5EC2C9805D42CE9,),PersistentGuid=00000000000000000000000000000000,bHidden=False,bNotConnectable=False,bDefaultValueIsReadOnly=False,bDefaultValueIsIgnored=False,bAdvancedView=False,)
+            End Object
+            Begin Object Class=MaterialGraphNode Name="MaterialGraphNode_1"
+               Begin Object Class=MaterialExpressionConstant Name="MaterialExpressionConstant_0"
+               End Object
+               Begin Object Name="MaterialExpressionConstant_0"
+                  R=0.500000
+                  MaterialExpressionEditorX=-80
+                  MaterialExpressionEditorY=128
+                  MaterialExpressionGuid=3D90A5C443D14DF0CF49A68F6E13583F
+                  Material=PreviewMaterial'/Engine/Transient.MAT_Stereo_Image'
+               End Object
+               MaterialExpression=MaterialExpressionConstant'MaterialExpressionConstant_0'
+               NodePosX=-80
+               NodePosY=128
+               NodeGuid=6190E3714CEDD46F0F98D8B9F38E5751
+               CustomProperties Pin (PinId=62DD318C41832241F7AB73BA06FB692E,PinName="Output",PinFriendlyName=" ",Direction="EGPD_Output",PinType.PinSubCategoryObject=None,PinType.PinSubCategoryMemberReference=(),PinType.bIsArray=False,PinType.bIsReference=False,PinType.bIsConst=False,PinType.bIsWeakPointer=False,LinkedTo=(MaterialGraphNode_0 D37D419748FE491453C139AF821D8F11,),PersistentGuid=00000000000000000000000000000000,bHidden=False,bNotConnectable=False,bDefaultValueIsReadOnly=False,bDefaultValueIsIgnored=False,bAdvancedView=False,)
+            End Object
+            Begin Object Class=MaterialGraphNode Name="MaterialGraphNode_2"
+               Begin Object Class=MaterialExpressionMaterialFunctionCall Name="MaterialExpressionMaterialFunctionCall_0"
+               End Object
+               Begin Object Name="MaterialExpressionMaterialFunctionCall_0"
+                  MaterialFunction=MaterialFunction'/Engine/Functions/Engine_MaterialFunctions02/Utility/BreakOutFloat2Components.BreakOutFloat2Components'
+                  FunctionInputs(0)=(ExpressionInputId=0FC5C56B4335798F20CB928CF1C87421,Input=(Expression=MaterialExpressionScreenPosition'MaterialGraphNode_3.MaterialExpressionScreenPosition_0',InputName="Float2"))
+                  FunctionOutputs(0)=(ExpressionOutputId=6BC893C247FCFC567178CE944E213385,Output=(OutputName="R"))
+                  FunctionOutputs(1)=(ExpressionOutputId=9720C81E449DF8E5514FCB8A80FD5FA3,Output=(OutputName="G"))
+                  MaterialExpressionEditorX=-272
+                  MaterialExpressionGuid=0DB6A16C443E70F09CC35BAA076DD2C3
+                  Material=PreviewMaterial'/Engine/Transient.MAT_Stereo_Image'
+                  Outputs(0)=(OutputName="R")
+                  Outputs(1)=(OutputName="G")
+               End Object
+               MaterialExpression=MaterialExpressionMaterialFunctionCall'MaterialExpressionMaterialFunctionCall_0'
+               NodePosX=-272
+               NodeGuid=43D001084D713B8304D69BBC75A0F775
+               CustomProperties Pin (PinId=361F5A0A417E832EB7A5E48C380D9DA4,PinName="Float2 (V2)",PinType.PinCategory="optional",PinType.PinSubCategoryObject=None,PinType.PinSubCategoryMemberReference=(),PinType.bIsArray=False,PinType.bIsReference=False,PinType.bIsConst=False,PinType.bIsWeakPointer=False,LinkedTo=(MaterialGraphNode_3 A6AE01F64B365D34D81BD981E64175F1,),PersistentGuid=00000000000000000000000000000000,bHidden=False,bNotConnectable=False,bDefaultValueIsReadOnly=False,bDefaultValueIsIgnored=False,bAdvancedView=False,)
+               CustomProperties Pin (PinId=176CDC95486F3BF845430099EF2BE7FE,PinName="R",Direction="EGPD_Output",PinType.PinSubCategoryObject=None,PinType.PinSubCategoryMemberReference=(),PinType.bIsArray=False,PinType.bIsReference=False,PinType.bIsConst=False,PinType.bIsWeakPointer=False,LinkedTo=(MaterialGraphNode_0 A0EE7D184D1901AC89D7CCB7D7B2AF54,),PersistentGuid=00000000000000000000000000000000,bHidden=False,bNotConnectable=False,bDefaultValueIsReadOnly=False,bDefaultValueIsIgnored=False,bAdvancedView=False,)
+               CustomProperties Pin (PinId=C686EF7142811E46C988739585372C29,PinName="G",Direction="EGPD_Output",PinType.PinSubCategoryObject=None,PinType.PinSubCategoryMemberReference=(),PinType.bIsArray=False,PinType.bIsReference=False,PinType.bIsConst=False,PinType.bIsWeakPointer=False,PersistentGuid=00000000000000000000000000000000,bHidden=False,bNotConnectable=False,bDefaultValueIsReadOnly=False,bDefaultValueIsIgnored=False,bAdvancedView=False,)
+            End Object
+            Begin Object Class=MaterialGraphNode Name="MaterialGraphNode_3"
+               Begin Object Class=MaterialExpressionScreenPosition Name="MaterialExpressionScreenPosition_0"
+               End Object
+               Begin Object Name="MaterialExpressionScreenPosition_0"
+                  MaterialExpressionEditorX=-448
+                  MaterialExpressionEditorY=-32
+                  MaterialExpressionGuid=F5C4FBB24FB4BE834B8623BAD82AA1FF
+                  Material=PreviewMaterial'/Engine/Transient.MAT_Stereo_Image'
+               End Object
+               MaterialExpression=MaterialExpressionScreenPosition'MaterialExpressionScreenPosition_0'
+               NodePosX=-448
+               NodePosY=-32
+               NodeGuid=0F5E2A6C4507E7940A41E281E944C1D8
+               CustomProperties Pin (PinId=A6AE01F64B365D34D81BD981E64175F1,PinName="Output",PinFriendlyName=" ",Direction="EGPD_Output",PinType.PinSubCategoryObject=None,PinType.PinSubCategoryMemberReference=(),PinType.bIsArray=False,PinType.bIsReference=False,PinType.bIsConst=False,PinType.bIsWeakPointer=False,LinkedTo=(MaterialGraphNode_2 361F5A0A417E832EB7A5E48C380D9DA4,),PersistentGuid=00000000000000000000000000000000,bHidden=False,bNotConnectable=False,bDefaultValueIsReadOnly=False,bDefaultValueIsIgnored=False,bAdvancedView=False,)
+            End Object
+            Begin Object Class=MaterialGraphNode Name="MaterialGraphNode_4"
+               Begin Object Class=MaterialExpressionMultiply Name="MaterialExpressionMultiply_0"
+               End Object
+               Begin Object Name="MaterialExpressionMultiply_0"
+                  A=(Expression=MaterialExpressionCameraVectorWS'MaterialGraphNode_5.MaterialExpressionCameraVectorWS_0')
+                  B=(Expression=MaterialExpressionVectorParameter'MaterialGraphNode_9.MaterialExpressionVectorParameter_0',Mask=1,MaskR=1,MaskG=1,MaskB=1)
+                  MaterialExpressionEditorX=-720
+                  MaterialExpressionEditorY=304
+                  MaterialExpressionGuid=4F10B53C4ECBDE27F4BD38B46FACAD18
+                  Material=PreviewMaterial'/Engine/Transient.MAT_Stereo_Image'
+                  bCollapsed=False
+               End Object
+               MaterialExpression=MaterialExpressionMultiply'MaterialExpressionMultiply_0'
+               NodePosX=-720
+               NodePosY=304
+               NodeGuid=630A92064469FB428330238D474586C7
+               CustomProperties Pin (PinId=E68092674090694BB584AEAEF50A40E0,PinName="A",PinType.PinCategory="optional",PinType.PinSubCategoryObject=None,PinType.PinSubCategoryMemberReference=(),PinType.bIsArray=False,PinType.bIsReference=False,PinType.bIsConst=False,PinType.bIsWeakPointer=False,LinkedTo=(MaterialGraphNode_5 CF46DC7F490ECC1D0F932B8242298926,),PersistentGuid=00000000000000000000000000000000,bHidden=False,bNotConnectable=False,bDefaultValueIsReadOnly=False,bDefaultValueIsIgnored=False,bAdvancedView=False,)
+               CustomProperties Pin (PinId=F1DE55B14567D0B67759DEA9C934BA05,PinName="B",PinType.PinCategory="optional",PinType.PinSubCategoryObject=None,PinType.PinSubCategoryMemberReference=(),PinType.bIsArray=False,PinType.bIsReference=False,PinType.bIsConst=False,PinType.bIsWeakPointer=False,LinkedTo=(MaterialGraphNode_9 9C23E1264DF642DC356B788B08CB3693,),PersistentGuid=00000000000000000000000000000000,bHidden=False,bNotConnectable=False,bDefaultValueIsReadOnly=False,bDefaultValueIsIgnored=False,bAdvancedView=False,)
+               CustomProperties Pin (PinId=BABCD0D743049748AF3DA996E4665D7D,PinName="Output",PinFriendlyName=" ",Direction="EGPD_Output",PinType.PinSubCategoryObject=None,PinType.PinSubCategoryMemberReference=(),PinType.bIsArray=False,PinType.bIsReference=False,PinType.bIsConst=False,PinType.bIsWeakPointer=False,LinkedTo=(MaterialGraphNode_6 6018F5F544A36CE1C7E511989710AD1D,),PersistentGuid=00000000000000000000000000000000,bHidden=False,bNotConnectable=False,bDefaultValueIsReadOnly=False,bDefaultValueIsIgnored=False,bAdvancedView=False,)
+            End Object
+            Begin Object Class=MaterialGraphNode Name="MaterialGraphNode_5"
+               Begin Object Class=MaterialExpressionCameraVectorWS Name="MaterialExpressionCameraVectorWS_0"
+               End Object
+               Begin Object Name="MaterialExpressionCameraVectorWS_0"
+                  MaterialExpressionEditorX=-912
+                  MaterialExpressionEditorY=256
+                  MaterialExpressionGuid=BA70FEF141C3FC6429AA56B337E6F88C
+                  Material=PreviewMaterial'/Engine/Transient.MAT_Stereo_Image'
+               End Object
+               MaterialExpression=MaterialExpressionCameraVectorWS'MaterialExpressionCameraVectorWS_0'
+               NodePosX=-912
+               NodePosY=256
+               NodeGuid=8F0586154298BB5972CE548ACAD0E0D2
+               CustomProperties Pin (PinId=CF46DC7F490ECC1D0F932B8242298926,PinName="Output",PinFriendlyName=" ",Direction="EGPD_Output",PinType.PinSubCategoryObject=None,PinType.PinSubCategoryMemberReference=(),PinType.bIsArray=False,PinType.bIsReference=False,PinType.bIsConst=False,PinType.bIsWeakPointer=False,LinkedTo=(MaterialGraphNode_4 E68092674090694BB584AEAEF50A40E0,),PersistentGuid=00000000000000000000000000000000,bHidden=False,bNotConnectable=False,bDefaultValueIsReadOnly=False,bDefaultValueIsIgnored=False,bAdvancedView=False,)
+            End Object
+            Begin Object Class=MaterialGraphNode Name="MaterialGraphNode_6"
+               Begin Object Class=MaterialExpressionMaterialFunctionCall Name="MaterialExpressionMaterialFunctionCall_1"
+               End Object
+               Begin Object Name="MaterialExpressionMaterialFunctionCall_1"
+                  MaterialFunction=MaterialFunction'/Engine/Functions/Engine_MaterialFunctions01/Coordinates/LongLatToUV.LongLatToUV'
+                  FunctionInputs(0)=(ExpressionInputId=150D02534FA6F797DFD995BA692D7801,Input=(Expression=MaterialExpressionMultiply'MaterialGraphNode_4.MaterialExpressionMultiply_0',InputName="Vector"))
+                  FunctionOutputs(0)=(ExpressionOutputId=0EB334074DD42132558C7EA6C023EC50,Output=(OutputName="Result"))
+                  MaterialExpressionEditorX=-526
+                  MaterialExpressionEditorY=299
+                  MaterialExpressionGuid=2B540AED4336652F5CC59D9DF979B358
+                  Material=PreviewMaterial'/Engine/Transient.MAT_Stereo_Image'
+                  Outputs(0)=(OutputName="Result")
+               End Object
+               MaterialExpression=MaterialExpressionMaterialFunctionCall'MaterialExpressionMaterialFunctionCall_1'
+               NodePosX=-526
+               NodePosY=299
+               NodeGuid=3A24EE65400825A12AE9DC844E79E18F
+               CustomProperties Pin (PinId=6018F5F544A36CE1C7E511989710AD1D,PinName="Vector (V3)",PinType.PinCategory="required",PinType.PinSubCategoryObject=None,PinType.PinSubCategoryMemberReference=(),PinType.bIsArray=False,PinType.bIsReference=False,PinType.bIsConst=False,PinType.bIsWeakPointer=False,LinkedTo=(MaterialGraphNode_4 BABCD0D743049748AF3DA996E4665D7D,),PersistentGuid=00000000000000000000000000000000,bHidden=False,bNotConnectable=False,bDefaultValueIsReadOnly=False,bDefaultValueIsIgnored=False,bAdvancedView=False,)
+               CustomProperties Pin (PinId=C87298924529FAC8728933A296BA7682,PinName="Result",Direction="EGPD_Output",PinType.PinSubCategoryObject=None,PinType.PinSubCategoryMemberReference=(),PinType.bIsArray=False,PinType.bIsReference=False,PinType.bIsConst=False,PinType.bIsWeakPointer=False,LinkedTo=(MaterialGraphNode_7 A60D4B0340117A9FC89B09BA6D69B3AF,MaterialGraphNode_8 0EA30C1443845B538B9305BFF8C98510,),PersistentGuid=00000000000000000000000000000000,bHidden=False,bNotConnectable=False,bDefaultValueIsReadOnly=False,bDefaultValueIsIgnored=False,bAdvancedView=False,)
+            End Object
+            Begin Object Class=MaterialGraphNode Name="MaterialGraphNode_7"
+               Begin Object Class=MaterialExpressionTextureSampleParameter2D Name="MaterialExpressionTextureSampleParameter2D_0"
+               End Object
+               Begin Object Name="MaterialExpressionTextureSampleParameter2D_0"
+                  ParameterName="Left_Eye"
+                  ExpressionGUID=6A1A633E43E045D020DFC3A368709F4D
+                  Coordinates=(Expression=MaterialExpressionMaterialFunctionCall'MaterialGraphNode_6.MaterialExpressionMaterialFunctionCall_1')
+                  Texture=None
+                  MaterialExpressionEditorX=-266
+                  MaterialExpressionEditorY=174
+                  MaterialExpressionGuid=8B896BCB490997B105768DA336F9575C
+                  Material=PreviewMaterial'/Engine/Transient.MAT_Stereo_Image'
+               End Object
+               MaterialExpression=MaterialExpressionTextureSampleParameter2D'MaterialExpressionTextureSampleParameter2D_0'
+               NodePosX=-266
+               NodePosY=174
+               bHasCompilerMessage=True
+               bCanRenameNode=True
+               ErrorType=1
+               ErrorMsg="Param2D> Requires Texture2D"
+               NodeGuid=9A5B88C84A2C1D5D71F007999C9E64DB
+               CustomProperties Pin (PinId=A60D4B0340117A9FC89B09BA6D69B3AF,PinName="UVs",PinType.PinCategory="optional",PinType.PinSubCategoryObject=None,PinType.PinSubCategoryMemberReference=(),PinType.bIsArray=False,PinType.bIsReference=False,PinType.bIsConst=False,PinType.bIsWeakPointer=False,LinkedTo=(MaterialGraphNode_6 C87298924529FAC8728933A296BA7682,),PersistentGuid=00000000000000000000000000000000,bHidden=False,bNotConnectable=False,bDefaultValueIsReadOnly=False,bDefaultValueIsIgnored=False,bAdvancedView=False,)
+               CustomProperties Pin (PinId=E039327F417689ACF875B5A184D2C251,PinName="Tex",PinType.PinCategory="optional",PinType.PinSubCategoryObject=None,PinType.PinSubCategoryMemberReference=(),PinType.bIsArray=False,PinType.bIsReference=False,PinType.bIsConst=False,PinType.bIsWeakPointer=False,PersistentGuid=00000000000000000000000000000000,bHidden=False,bNotConnectable=False,bDefaultValueIsReadOnly=False,bDefaultValueIsIgnored=False,bAdvancedView=False,)
+               CustomProperties Pin (PinId=6073EDCF453D52887F1ED597234FAF0C,PinName="Output",PinFriendlyName=" ",Direction="EGPD_Output",PinType.PinCategory="mask",PinType.PinSubCategoryObject=None,PinType.PinSubCategoryMemberReference=(),PinType.bIsArray=False,PinType.bIsReference=False,PinType.bIsConst=False,PinType.bIsWeakPointer=False,LinkedTo=(MaterialGraphNode_0 44EC5EE54EACAD4AA705D8866FE16EB5,),PersistentGuid=00000000000000000000000000000000,bHidden=False,bNotConnectable=False,bDefaultValueIsReadOnly=False,bDefaultValueIsIgnored=False,bAdvancedView=False,)
+               CustomProperties Pin (PinId=1BFBF7694F3E966628AFC58E7C534726,PinName="Output2",PinFriendlyName=" ",Direction="EGPD_Output",PinType.PinCategory="mask",PinType.PinSubCategory="red",PinType.PinSubCategoryObject=None,PinType.PinSubCategoryMemberReference=(),PinType.bIsArray=False,PinType.bIsReference=False,PinType.bIsConst=False,PinType.bIsWeakPointer=False,PersistentGuid=00000000000000000000000000000000,bHidden=False,bNotConnectable=False,bDefaultValueIsReadOnly=False,bDefaultValueIsIgnored=False,bAdvancedView=False,)
+               CustomProperties Pin (PinId=3BFCC8DA419D81F993A7B391CF2FA32E,PinName="Output3",PinFriendlyName=" ",Direction="EGPD_Output",PinType.PinCategory="mask",PinType.PinSubCategory="green",PinType.PinSubCategoryObject=None,PinType.PinSubCategoryMemberReference=(),PinType.bIsArray=False,PinType.bIsReference=False,PinType.bIsConst=False,PinType.bIsWeakPointer=False,PersistentGuid=00000000000000000000000000000000,bHidden=False,bNotConnectable=False,bDefaultValueIsReadOnly=False,bDefaultValueIsIgnored=False,bAdvancedView=False,)
+               CustomProperties Pin (PinId=C25C6F6F42144A2FC2C5DB911ECD7C50,PinName="Output4",PinFriendlyName=" ",Direction="EGPD_Output",PinType.PinCategory="mask",PinType.PinSubCategory="blue",PinType.PinSubCategoryObject=None,PinType.PinSubCategoryMemberReference=(),PinType.bIsArray=False,PinType.bIsReference=False,PinType.bIsConst=False,PinType.bIsWeakPointer=False,PersistentGuid=00000000000000000000000000000000,bHidden=False,bNotConnectable=False,bDefaultValueIsReadOnly=False,bDefaultValueIsIgnored=False,bAdvancedView=False,)
+               CustomProperties Pin (PinId=4FACBD794BD17982F6AFE2A83001ED2E,PinName="Output5",PinFriendlyName=" ",Direction="EGPD_Output",PinType.PinCategory="mask",PinType.PinSubCategory="alpha",PinType.PinSubCategoryObject=None,PinType.PinSubCategoryMemberReference=(),PinType.bIsArray=False,PinType.bIsReference=False,PinType.bIsConst=False,PinType.bIsWeakPointer=False,PersistentGuid=00000000000000000000000000000000,bHidden=False,bNotConnectable=False,bDefaultValueIsReadOnly=False,bDefaultValueIsIgnored=False,bAdvancedView=False,)
+            End Object
+            Begin Object Class=MaterialGraphNode Name="MaterialGraphNode_8"
+               Begin Object Class=MaterialExpressionTextureSampleParameter2D Name="MaterialExpressionTextureSampleParameter2D_1"
+               End Object
+               Begin Object Name="MaterialExpressionTextureSampleParameter2D_1"
+                  ParameterName="Right_Eye"
+                  ExpressionGUID=20152A504733C1C1C6CCBFBEDE38D8C7
+                  Coordinates=(Expression=MaterialExpressionMaterialFunctionCall'MaterialGraphNode_6.MaterialExpressionMaterialFunctionCall_1')
+                  ConstMipValue=1
+                  Texture=None
+                  MaterialExpressionEditorX=-241
+                  MaterialExpressionEditorY=445
+                  MaterialExpressionGuid=0D48E23F4EF41B3303AB1AA8DC273216
+                  Material=PreviewMaterial'/Engine/Transient.MAT_Stereo_Image'
+               End Object
+               MaterialExpression=MaterialExpressionTextureSampleParameter2D'MaterialExpressionTextureSampleParameter2D_1'
+               NodePosX=-241
+               NodePosY=445
+               bHasCompilerMessage=True
+               bCanRenameNode=True
+               ErrorType=1
+               ErrorMsg="Param2D> Requires Texture2D"
+               NodeGuid=A93210F24FF9C8FE30CEECBBAF308703
+               CustomProperties Pin (PinId=0EA30C1443845B538B9305BFF8C98510,PinName="UVs",PinType.PinCategory="optional",PinType.PinSubCategoryObject=None,PinType.PinSubCategoryMemberReference=(),PinType.bIsArray=False,PinType.bIsReference=False,PinType.bIsConst=False,PinType.bIsWeakPointer=False,LinkedTo=(MaterialGraphNode_6 C87298924529FAC8728933A296BA7682,),PersistentGuid=00000000000000000000000000000000,bHidden=False,bNotConnectable=False,bDefaultValueIsReadOnly=False,bDefaultValueIsIgnored=False,bAdvancedView=False,)
+               CustomProperties Pin (PinId=6E12F03A4EBBD551878771B92104B952,PinName="Tex",PinType.PinCategory="optional",PinType.PinSubCategoryObject=None,PinType.PinSubCategoryMemberReference=(),PinType.bIsArray=False,PinType.bIsReference=False,PinType.bIsConst=False,PinType.bIsWeakPointer=False,PersistentGuid=00000000000000000000000000000000,bHidden=False,bNotConnectable=False,bDefaultValueIsReadOnly=False,bDefaultValueIsIgnored=False,bAdvancedView=False,)
+               CustomProperties Pin (PinId=CDF923064C8491C8DD875A9FD733B329,PinName="Output",PinFriendlyName=" ",Direction="EGPD_Output",PinType.PinCategory="mask",PinType.PinSubCategoryObject=None,PinType.PinSubCategoryMemberReference=(),PinType.bIsArray=False,PinType.bIsReference=False,PinType.bIsConst=False,PinType.bIsWeakPointer=False,LinkedTo=(MaterialGraphNode_0 F937930D47BA8EA98B9510B0C025B9A4,),PersistentGuid=00000000000000000000000000000000,bHidden=False,bNotConnectable=False,bDefaultValueIsReadOnly=False,bDefaultValueIsIgnored=False,bAdvancedView=False,)
+               CustomProperties Pin (PinId=F4A512534340958957BAC39A44E565FF,PinName="Output2",PinFriendlyName=" ",Direction="EGPD_Output",PinType.PinCategory="mask",PinType.PinSubCategory="red",PinType.PinSubCategoryObject=None,PinType.PinSubCategoryMemberReference=(),PinType.bIsArray=False,PinType.bIsReference=False,PinType.bIsConst=False,PinType.bIsWeakPointer=False,PersistentGuid=00000000000000000000000000000000,bHidden=False,bNotConnectable=False,bDefaultValueIsReadOnly=False,bDefaultValueIsIgnored=False,bAdvancedView=False,)
+               CustomProperties Pin (PinId=0080AD0C498AB9A1D974ED80486A3FEE,PinName="Output3",PinFriendlyName=" ",Direction="EGPD_Output",PinType.PinCategory="mask",PinType.PinSubCategory="green",PinType.PinSubCategoryObject=None,PinType.PinSubCategoryMemberReference=(),PinType.bIsArray=False,PinType.bIsReference=False,PinType.bIsConst=False,PinType.bIsWeakPointer=False,PersistentGuid=00000000000000000000000000000000,bHidden=False,bNotConnectable=False,bDefaultValueIsReadOnly=False,bDefaultValueIsIgnored=False,bAdvancedView=False,)
+               CustomProperties Pin (PinId=1B4746BE460A06F7E4C99091F01E070D,PinName="Output4",PinFriendlyName=" ",Direction="EGPD_Output",PinType.PinCategory="mask",PinType.PinSubCategory="blue",PinType.PinSubCategoryObject=None,PinType.PinSubCategoryMemberReference=(),PinType.bIsArray=False,PinType.bIsReference=False,PinType.bIsConst=False,PinType.bIsWeakPointer=False,PersistentGuid=00000000000000000000000000000000,bHidden=False,bNotConnectable=False,bDefaultValueIsReadOnly=False,bDefaultValueIsIgnored=False,bAdvancedView=False,)
+               CustomProperties Pin (PinId=2059FDD14284DE0A16798E92259F2C7A,PinName="Output5",PinFriendlyName=" ",Direction="EGPD_Output",PinType.PinCategory="mask",PinType.PinSubCategory="alpha",PinType.PinSubCategoryObject=None,PinType.PinSubCategoryMemberReference=(),PinType.bIsArray=False,PinType.bIsReference=False,PinType.bIsConst=False,PinType.bIsWeakPointer=False,PersistentGuid=00000000000000000000000000000000,bHidden=False,bNotConnectable=False,bDefaultValueIsReadOnly=False,bDefaultValueIsIgnored=False,bAdvancedView=False,)
+            End Object
+            Begin Object Class=MaterialGraphNode Name="MaterialGraphNode_9"
+               Begin Object Class=MaterialExpressionVectorParameter Name="MaterialExpressionVectorParameter_0"
+               End Object
+               Begin Object Name="MaterialExpressionVectorParameter_0"
+                  DefaultValue=(R=1.000000,G=-1.000000,B=-1.000000,A=1.000000)
+                  ParameterName="Image_Direction"
+                  ExpressionGUID=F1F2C7C348FA6F958EA8E099B55A37ED
+                  MaterialExpressionEditorX=-912
+                  MaterialExpressionEditorY=384
+                  MaterialExpressionGuid=6873E4014A670FD8059959BFE93D5CDA
+                  Material=PreviewMaterial'/Engine/Transient.MAT_Stereo_Image'
+               End Object
+               MaterialExpression=MaterialExpressionVectorParameter'MaterialExpressionVectorParameter_0'
+               NodePosX=-912
+               NodePosY=384
+               bCanRenameNode=True
+               NodeGuid=37303B2F48466367CBE440A181DB2539
+               CustomProperties Pin (PinId=9C23E1264DF642DC356B788B08CB3693,PinName="Output",PinFriendlyName=" ",Direction="EGPD_Output",PinType.PinCategory="mask",PinType.PinSubCategoryObject=None,PinType.PinSubCategoryMemberReference=(),PinType.bIsArray=False,PinType.bIsReference=False,PinType.bIsConst=False,PinType.bIsWeakPointer=False,LinkedTo=(MaterialGraphNode_4 F1DE55B14567D0B67759DEA9C934BA05,),PersistentGuid=00000000000000000000000000000000,bHidden=False,bNotConnectable=False,bDefaultValueIsReadOnly=False,bDefaultValueIsIgnored=False,bAdvancedView=False,)
+               CustomProperties Pin (PinId=6B8E6BCA4EF15BDF5D11B6B1CB4B1882,PinName="Output2",PinFriendlyName=" ",Direction="EGPD_Output",PinType.PinCategory="mask",PinType.PinSubCategory="red",PinType.PinSubCategoryObject=None,PinType.PinSubCategoryMemberReference=(),PinType.bIsArray=False,PinType.bIsReference=False,PinType.bIsConst=False,PinType.bIsWeakPointer=False,PersistentGuid=00000000000000000000000000000000,bHidden=False,bNotConnectable=False,bDefaultValueIsReadOnly=False,bDefaultValueIsIgnored=False,bAdvancedView=False,)
+               CustomProperties Pin (PinId=4F58ABD7410B38EA9235A497F25BD3C1,PinName="Output3",PinFriendlyName=" ",Direction="EGPD_Output",PinType.PinCategory="mask",PinType.PinSubCategory="green",PinType.PinSubCategoryObject=None,PinType.PinSubCategoryMemberReference=(),PinType.bIsArray=False,PinType.bIsReference=False,PinType.bIsConst=False,PinType.bIsWeakPointer=False,PersistentGuid=00000000000000000000000000000000,bHidden=False,bNotConnectable=False,bDefaultValueIsReadOnly=False,bDefaultValueIsIgnored=False,bAdvancedView=False,)
+               CustomProperties Pin (PinId=252E95C9465BAEED4430B5B9DA37851D,PinName="Output4",PinFriendlyName=" ",Direction="EGPD_Output",PinType.PinCategory="mask",PinType.PinSubCategory="blue",PinType.PinSubCategoryObject=None,PinType.PinSubCategoryMemberReference=(),PinType.bIsArray=False,PinType.bIsReference=False,PinType.bIsConst=False,PinType.bIsWeakPointer=False,PersistentGuid=00000000000000000000000000000000,bHidden=False,bNotConnectable=False,bDefaultValueIsReadOnly=False,bDefaultValueIsIgnored=False,bAdvancedView=False,)
+               CustomProperties Pin (PinId=B5B286C94A078205748BA1880BD0B1D1,PinName="Output5",PinFriendlyName=" ",Direction="EGPD_Output",PinType.PinCategory="mask",PinType.PinSubCategory="alpha",PinType.PinSubCategoryObject=None,PinType.PinSubCategoryMemberReference=(),PinType.bIsArray=False,PinType.bIsReference=False,PinType.bIsConst=False,PinType.bIsWeakPointer=False,PersistentGuid=00000000000000000000000000000000,bHidden=False,bNotConnectable=False,bDefaultValueIsReadOnly=False,bDefaultValueIsIgnored=False,bAdvancedView=False,)
+            End Object
+    [/PARAM]
+[/OBJECT]
+[/EXCERPT:SPMaterial]
+-->
\ No newline at end of file
diff --git a/Engine/Documentation/Source/Platforms/VR/StereoPanoramicCapture/QuickStart/LevelSetup/LevelSetup.INT.udn b/Engine/Documentation/Source/Platforms/VR/StereoPanoramicCapture/QuickStart/LevelSetup/LevelSetup.INT.udn
new file mode 100644
--- /dev/null
+++ b/Engine/Documentation/Source/Platforms/VR/StereoPanoramicCapture/QuickStart/LevelSetup/LevelSetup.INT.udn
@@ -0,0 +1,96 @@
+Availability: Public
+Title:4. Level Setup
+Crumbs:
+Description: In this section, we will go over how to set up a level to display the stereo screenshots.
+Version: 4.13
+SkillLevel: Intermediate
+parent:Platforms/VR/StereoPanoramicCapture/QuickStart
+type: quick start
+order:4
+tags:VR
+Tags:Cinematics
+Tags:GearVR
+Tags:Google VR	
+Tags:SteamVR
+Tags:Oculus Rift
+Tags:Virtual Reality
+Tags:Plugins
+Tags:Sequencer
+
+[NAV]
+
+In this step, we will go over how to set up a level that will display the stereo screenshots when the level is run.
+
+1. From the **Main Toolbar** go to the **File** option and then from the list that is displayed, select the **New Level** option.
+
+    ![](SPHT_New_Level.png)
+
+1. When the **New Level** dialogue box shows up, choose the **Empty Level** option to load a new empty level.
+
+    ![](SPHT_Blank_Level.png)
+
+1. Click on the link below to download the required Static Mesh that is needed to display the Material. 
+
+    [](SM_PhotoSphere.zip)
+
+	[region:note]
+	The Static Mesh you are downloading has had its UV’s setup to work correctly with Material that was created. While you can use any Static Mesh you want, downloading this file will ensure that you will get the correct results.
+	[/region]
+
+1. Once the file has downloaded, unzip the file and then import the Static mesh into the Content Browser.
+
+    [REGION:lightbox]
+    [![](SPHT_Import_Sky.png)(w:617)](SPHT_Import_Sky.png)
+    [/REGION]
+
+    [REGION:caption]
+    Click for full image.
+    [/REGION]
+
+    [region:note]
+    If you receive a warning about smoothing groups not being found in the file, you can safely ignore that warning as we do not need smoothing information for this object as it uses an unlit shader. 
+    [/region]
+
+1. After the **SM_PhotoSphere**  Static Mesh has been imported, drag the **SM_PhotoSphere** from the Content Browser into the level, positioning it so that it's location is a **0,0,0** in the world.
+
+    [REGION:lightbox]
+    [![](SPHT_Position_SM.png)(w:617)](SPHT_Position_SM.png)
+    [/REGION]
+
+    [REGION:caption]
+    Click for full image.
+    [/REGION]
+
+1. Once the SM_PhotoSphere has been positioned in the center of the level, select the MAT_StereoViewer Material in the Content Browser and then drag it to the **Materials** input in the Details panel.
+
+    [REGION:lightbox]
+    [![](SPHY_Apply_Material.png)(w:617)](SPHY_Apply_Material.png)
+    [/REGION]
+
+    [REGION:caption]
+    Click for full image.
+    [/REGION]
+
+1. Once the Material has been applied, press the **Play in VR** preview button and then put your HMD on to see what the images look like when viewed in stereo. If everything has been set up correctly, you should not be able to tell that the level is being rendered using just two images instead of 3D objects.
+
+    [OBJECT:EmbeddedVideo]
+        [PARAMLITERAL:width]
+        640
+        [/PARAMLITERAL]
+        [PARAMLITERAL:height]
+        360
+        [/PARAMLITERAL]
+        [PARAMLITERAL:videoid]
+        cg7xV85GRzo
+        [/PARAMLITERAL]
+    [/OBJECT]
+
+	[region:note]
+	The black sphere you are seeing is the shadow that is cast from the Static Mesh that is attached to the Default Pawn.  To get rid of this shadow you can either create a new Game Mode and Pawn or when you run the game you can press the player eject button, find the Default Pawn in the World Outliner and under the Static Mesh section, set the Static Mesh to none and the shadow will be gone.
+	[/region]
+
+
+
+With this simple set up, you can now view any images that you take using the Stereo Panoramic plugin by replacing the left and right eye images with the new images that you created.
+
+[NAV]
\ No newline at end of file
diff --git a/Engine/Documentation/Source/Platforms/VR/StereoPanoramicCapture/QuickStart/ProjectSetup/StereoPanoramicProjectSetup.INT.udn b/Engine/Documentation/Source/Platforms/VR/StereoPanoramicCapture/QuickStart/ProjectSetup/StereoPanoramicProjectSetup.INT.udn
new file mode 100644
--- /dev/null
+++ b/Engine/Documentation/Source/Platforms/VR/StereoPanoramicCapture/QuickStart/ProjectSetup/StereoPanoramicProjectSetup.INT.udn
@@ -0,0 +1,65 @@
+Availability: Public
+Title:1. Project Setup
+Crumbs:
+Description:Here we will create and setup a new project to capture stereo screenshots in. 
+Version:4.13
+SkillLevel: Intermediate
+parent:Platforms/VR/StereoPanoramicCapture/QuickStart
+type:quick start
+order:1
+tags:VR
+Tags:Cinematics
+Tags:GearVR
+Tags:Google VR	
+Tags:SteamVR
+Tags:Oculus Rift
+Tags:Virtual Reality
+Tags:Plugins
+Tags:Sequencer
+
+[NAV]
+
+In this section, we will create and then set up a new UE4 project to work out of.
+
+[region:note]
+While we will be creating a new blank project for this how-to, the information presented here can be applied to any UE4 project.
+[/region]
+
+## Steps
+
+1. From the **New Project** section of the Unreal Project Browser create a new Blank Blueprint base project with the following options and when completed, press the **Create Project** button to set up and load the new project.
+
+    * Blank Project Template
+    * Mobile / Tablet
+    * Scalable 3D or 2D
+    * No Starter Content
+
+    [REGION:lightbox]
+    [![](SPHT_Create_New_Project.png)(w:555)](SPHT_Create_New_Project.png)
+    [/REGION]
+
+    [REGION:caption]
+    Click for full image.
+    [/REGION]
+
+1. Once the project has been loaded go to the top menu bar and then go to **Edit** > **Plugins** to open up the plugins options.
+
+    ![](SPHT_Plugins_Menu.png)
+
+1. In the Plugins menu, click on the **Movie Capture** option and then in the **Stereo Panoramic Movie Capture** section, click on the checkmark box to enable the plugin and then press the **Restart Now** button to restart the UE4 editor.
+
+    [REGION:lightbox]
+    [![](SPHT_Enable_SPMP.png)(w:543)](SPHT_Enable_SPMP.png)
+    [/REGION]
+
+    [REGION:caption]
+    Click for full image.
+    [/REGION]
+
+1. Once the editor re-starts open up the console by pressing the backtick (`) on the keyboard and then input **SP.**. If the Stereo Panoramic Movie Plugin has been successfully enabled, you should see a list of commands you can choose from. If you do not see the command list make sure that you enabled and restarted the editor and that you are using Unreal Engine 4.13 or later.
+
+    ![](SPHT_Check_Plugin.png)
+
+With the Stereo Panoramic Move Plugin now enabled and verified to be working, we will next take a look at how to use the plugin to capture stereo images.
+
+[NAV]
\ No newline at end of file
diff --git a/Engine/Documentation/Source/Platforms/VR/StereoPanoramicCapture/QuickStart/StereoPanoramicHowTo.INT.udn b/Engine/Documentation/Source/Platforms/VR/StereoPanoramicCapture/QuickStart/StereoPanoramicHowTo.INT.udn
new file mode 100644
--- /dev/null
+++ b/Engine/Documentation/Source/Platforms/VR/StereoPanoramicCapture/QuickStart/StereoPanoramicHowTo.INT.udn
@@ -0,0 +1,39 @@
+Availability:Public
+Title: Stereo Panoramic Capture Tool Quick Start
+Crumbs:%ROOT%, Platforms, Platforms/VR
+Description:A look at how to use the Stereo Panoramic plugin to capture VR images and movies.
+Navigation:platform
+Version: 4.13
+parent:Platforms/VR/StereoPanoramicCapture
+order:1
+type:quick start
+Topic-Image:Platforms/VR/StereoPanoramicCapture/SPE_Topic_Image_00.png
+tags:VR
+Tags:Cinematics
+Tags:GearVR
+Tags:Google VR	
+Tags:SteamVR
+Tags:Oculus Rift
+Tags:Virtual Reality
+Tags:Plugins
+Tags:Sequencer
+
+![](Platforms\VR\StereoPanoramicCapture\SPE_Header_00.png)
+
+In the following example, we will take a look at how you can setup, create and view Stereo Panoramic screenshots all from inside UE4. When you have completed this project you will have created two Stereo Panoramic images and a Material that can be used to display those images in 3D. 
+
+
+## Steps 
+
+
+[DIR(output:"checkpoint" parent:"Platforms/VR/StereoPanoramicCapture/QuickStart")]
+
+[REGION:call_to_action]
+    [Click to Start](Platforms\VR\StereoPanoramicCapture\QuickStart\ProjectSetup "%Platforms\VR\StereoPanoramicCapture\QuickStart\ProjectSetup:title%")
+[/REGION]
+
+
+
+
+
+
diff --git a/Engine/Documentation/Source/Platforms/VR/StereoPanoramicCapture/QuickStart/TakingStereoImages/TakingStereoImages.INT.udn b/Engine/Documentation/Source/Platforms/VR/StereoPanoramicCapture/QuickStart/TakingStereoImages/TakingStereoImages.INT.udn
new file mode 100644
--- /dev/null
+++ b/Engine/Documentation/Source/Platforms/VR/StereoPanoramicCapture/QuickStart/TakingStereoImages/TakingStereoImages.INT.udn
@@ -0,0 +1,64 @@
+Availability: Public
+Title:2. Taking Stereo Images
+Crumbs:
+Description: In this section, we will go over how to capture stereo screenshots using the Stereo Panoramic plugin.
+Version:4.13
+SkillLevel: Intermediate
+parent:Platforms/VR/StereoPanoramicCapture/QuickStart
+type:quick start
+order:2
+tags:VR
+Tags:Cinematics
+Tags:GearVR
+Tags:Google VR	
+Tags:SteamVR
+Tags:Oculus Rift
+Tags:Virtual Reality
+Tags:Plugins
+Tags:Sequencer
+
+
+[NAV]
+
+In this step, we will use the Stereo Panoramic plugin to capture the images of the level in stereo 3D.
+
+## Steps
+
+1. From the **Main Toolbar** press the **Play** button and then navigate your player to an area that you want to capture. For this example, we will be leaving the player in the middle of the Static Mesh but feel free to position the player anywhere you want to.
+
+    ![](SPHT_Play_Game.png)
+
+1. Press the **Backtick** key (`) on the keyboard to open up the UE4 console and then input the following command. Once the command has been input press the Enter key to execute the command.
+
+    * **SP.PanoramicScreenshot**
+
+    ![](SPHT_Input_Command.png)
+
+    [region:note]
+    Instead of typing out SP.PanoramicScreenshot in its entirety; you can just type **SP.P** and then select SP.PanoramicScreenshot from the displayed list using the arrow keys. Once SP.PanoramicScreenshot is selected, press the Enter key to execute the command.
+    [/region]
+
+1. During the capture process, the editor might appear to be frozen or unresponsive for a few seconds or minutes.  This is due to the demanding rendering requirements of the Stereo Panoramic plugin. When the editor becomes responsive again, you will be able to find the screenshots in the following location in your projects folder.
+
+    * **StereoPanoramaImages\Saved\StereoPanorama\(Date & Time it was taken)**
+
+    [REGION:lightbox]
+    [![](SPHT_Rendered_Images.png)(w:473)](SPHT_Rendered_Images.png)
+    [/REGION]
+
+    [REGION:caption]
+    Click for full image.
+    [/REGION]
+
+    [region:note]
+    If the image previews look washed in thumbnail images or when viewed in an image viewer it is because the Stereo Panoramic plugin will add an alpha channel to the output.PNG file. 
+    [/region]
+
+	[region:tip]
+	To create the images that you see in the folder the Stereo Panoramic plugin is rendering the scene from the viewpoint of each eye that is rotated around a fixed point to simulate the head turning around the neck.
+	[/region]
+
+
+With the images now created, in the next section we will take a look at how to build a Material that will allow the images to be viewed using a VR HMD.
+
+[NAV]
diff --git a/Engine/Documentation/Source/Platforms/VR/StereoPanoramicCapture/Reference/StereoPanoramicCaptureReference.INT.udn b/Engine/Documentation/Source/Platforms/VR/StereoPanoramicCapture/Reference/StereoPanoramicCaptureReference.INT.udn
--- a/Engine/Documentation/Source/Platforms/VR/StereoPanoramicCapture/Reference/StereoPanoramicCaptureReference.INT.udn
+++ b/Engine/Documentation/Source/Platforms/VR/StereoPanoramicCapture/Reference/StereoPanoramicCaptureReference.INT.udn
@@ -1,13 +1,26 @@
-Availability:Docs
+Availability:Public
 Title:Stereo Panoramic Capture Reference
 Crumbs:%ROOT%, Platforms, Platforms/VR
 Description:An in-depth look at what each property for the Stereo Panoramic does
 Navigation:platform
 Version: 4.13
 parent:Platforms/VR/StereoPanoramicCapture
+order:3
 type:reference
 related: Gameplay/HowTo/UsingCameras/
 related: Engine/Performance/Scalability/ScalabilityReference/
+Topic-Image:SPE_Topic_Image_00.png
+tags:VR
+Tags:Cinematics
+Tags:GearVR
+Tags:Google VR	
+Tags:SteamVR
+Tags:Oculus Rift
+Tags:Virtual Reality
+Tags:Plugins
+Tags:Sequencer
+
+[TOC(start:2)]
 
 On the following page, you will find a detailed description of what options and controls are available for the Stereo Panoramic Capture plugin and what those options do to the image that is output. The options have been broken up into groups based on what those options do. To input a command press the backtick or tilde key (`) to open up the UE4 console and then type **SP.** followed by the command you want to enter.
 
diff --git a/Engine/Documentation/Source/Platforms/VR/StereoPanoramicCapture/StereoPanoramicCapture.INT.udn b/Engine/Documentation/Source/Platforms/VR/StereoPanoramicCapture/StereoPanoramicCapture.INT.udn
--- a/Engine/Documentation/Source/Platforms/VR/StereoPanoramicCapture/StereoPanoramicCapture.INT.udn
+++ b/Engine/Documentation/Source/Platforms/VR/StereoPanoramicCapture/StereoPanoramicCapture.INT.udn
@@ -1,4 +1,4 @@
-Availability:Docs
+Availability:Public
 Title:Stereo Panoramic Capture Tool
 Crumbs:%ROOT%, Platforms, Platforms/VR
 Description:A look at how to use the Stereo Panoramic plugin to capture VR images and movies.
@@ -6,6 +6,18 @@ Navigation:platform
 Version: 4.13
 parent:Platforms/VR
 type:landing
+SEO-Title:StereoPanoramicCaptureTool
+Topic-Image:SPE_Topic_Image_00.png
+tags:VR
+Tags:Cinematics
+Tags:GearVR
+Tags:Google VR	
+Tags:SteamVR
+Tags:Oculus Rift
+Tags:Virtual Reality
+Tags:Plugins
+Tags:Sequencer
+
 
 ![](SPE_Header_00.png)
 
@@ -17,9 +29,18 @@ Please note that the Stereo Panoramic Movie Capture plugin is still in an experi
 
 ## Starting Out
 
-The Stereo Panoramic Capture How To guide will walk you through getting your UE4 project setup to be able to capture and view stereo images.
+[DIR(end:"1" output:"topic" parent:"Platforms/VR/StereoPanoramicCapture" type:"quick start")]
+
+## Reference
+
+[DIR(end:"1" output:"listbutton" parent:"Platforms/VR/StereoPanoramicCapture" type:"!quick start")]
+
+<!--
+## Starting Out
+
+The Stereo Panoramic Capture Quick Start guide will walk you through getting your UE4 project setup to be able to capture and view stereo images.
 
-[DIR(end:"1" output:"listbutton" type:"how-to" parent:"Platforms/VR/StereoPanoramicCapture")]
+[DIR(end:"1" output:"listbutton" type:"quick start" parent:"Platforms/VR/StereoPanoramicCapture")]
 
 ## Stereo Panoramic Capture Reference 
 
@@ -33,20 +54,4 @@ The Stereo Panoramic Capture Tips & Tricks section will go over several tips and
 
 [DIR(end:"1" output:"listbutton" type:"overview" parent:"Platforms/VR/StereoPanoramicCapture")]
 
-
-
-
-
-
-
-
-
-
-
-
-
-
-
-
-
-
+-->
\ No newline at end of file
diff --git a/Engine/Documentation/Source/Platforms/VR/StereoPanoramicCapture/TipsAndTricks/TipAndTricks.INT.udn b/Engine/Documentation/Source/Platforms/VR/StereoPanoramicCapture/TipsAndTricks/TipAndTricks.INT.udn
--- a/Engine/Documentation/Source/Platforms/VR/StereoPanoramicCapture/TipsAndTricks/TipAndTricks.INT.udn
+++ b/Engine/Documentation/Source/Platforms/VR/StereoPanoramicCapture/TipsAndTricks/TipAndTricks.INT.udn
@@ -1,11 +1,22 @@
-Availability:Docs
+Availability:Public
 Title:Stereo Panoramic Capture Tips & Tricks
 Crumbs:%ROOT%, Platforms, Platforms/VR
 Description:Taking a look at how to fix some common issues with the Stereo Panoramic capture tool
 Navigation:platform
 Version: 4.13
 parent:Platforms/VR/StereoPanoramicCapture
+order:2
 type:overview
+Topic-Image:SPE_Topic_Image_00.png
+tags:VR
+Tags:Cinematics
+Tags:GearVR
+Tags:Google VR	
+Tags:SteamVR
+Tags:Oculus Rift
+Tags:Virtual Reality
+Tags:Plugins
+Tags:Sequencer
 
 [TOC(start:2)]
 
@@ -73,9 +84,10 @@ When viewing, images using the Stereo Image Material you might notice a seam whe
 To remove these seams, you will need to open up your Left and Right eye images in the UE4 Texture Viewer. Then in the **Level of Detail** section, change the value for **Mip Gen Settings** to **NoMipmaps**. 
 
 [REGION:lightbox]
-[![](Platforms\VR\StereoPanoramicCapture\HowTo\CreatingViewerMaterial\SPHT_Mip_Gen_Settings.png)(w:439)](Platforms\VR\StereoPanoramicCapture\HowTo\CreatingViewerMaterial\SPHT_Mip_Gen_Settings.png)
+[![](SPHT_Mip_Gen_Settings.png)(w:439)](SPHT_Mip_Gen_Settings.png)
 [/REGION]
 
+
 [REGION:caption]
 Click for full image.
 [/REGION]
diff --git a/Engine/Documentation/Source/Platforms/VR/VR.INT.udn b/Engine/Documentation/Source/Platforms/VR/VR.INT.udn
--- a/Engine/Documentation/Source/Platforms/VR/VR.INT.udn
+++ b/Engine/Documentation/Source/Platforms/VR/VR.INT.udn
@@ -50,7 +50,7 @@ vr
 
 ## Starting Out
 
-[DIR(output:"listbutton" type:"quick start" tags:"vr")]
+[DIR(output:"topic" type:"quick start" tags:"vr")]
 
 ## Virtual Reality Platforms
 
diff --git a/Engine/Documentation/Source/Programming/Development/BuildGraph/BuildGraph.INT.udn b/Engine/Documentation/Source/Programming/Development/BuildGraph/BuildGraph.INT.udn
--- a/Engine/Documentation/Source/Programming/Development/BuildGraph/BuildGraph.INT.udn
+++ b/Engine/Documentation/Source/Programming/Development/BuildGraph/BuildGraph.INT.udn
@@ -1,11 +1,14 @@
 Availability: Public
 Title: BuildGraph
 Crumbs: %ROOT%, Programming, Programming/Development
-Description: Customize your Installed Builds by using our BuildGraph scripting system.
+Description: Customize your builds with the BuildGraph scripting system.
 Version: 4.13
 type:landing
 parent:Programming/Development
+tags:Programming
 tags:BuildGraph
+topic-image:buildgraph_topic.png
+related:Programming/UnrealBuildSystem
 
 **BuildGraph** is a script-based build automation system that features graphs of building blocks common to Unreal Engine 4 (UE4) projects. 
 BuildGraph integrates with UnrealBuildTool, AutomationTool, and the editor, and can be extended and customized for your projects.
@@ -16,7 +19,7 @@ this project, then cook, run this test, and so on). When asked to build a target
 will execute all of the nodes in the graph required to make that happen.
 
 Unlike other build tools, BuildGraph is designed as a hybrid between a makefile-like scripting language and a build farm 
-configuration script. It allows annotations for the type of machine that steps are supposed to be executed on, providing a 
+configuration script. It enables annotations for the type of machine that steps are supposed to be executed on, providing a 
 list of recipients for failure notifications if a step fails, and groups nodes that should only be executed after an explicit 
 user trigger. It also tracks the creation of output files from task execution in a way that lends to graph execution being 
 distributed across a farm of machines (with nodes running in parallel where possible), and intermediate artifacts being 
@@ -54,7 +57,7 @@ Any element may be conditionally defined via the "If" attribute. See below for t
 BuildGraph is typically used for packaged games, so filtering and file manipulation are natively supported. Any attribute 
 that accepts a list of files may consist of a Perforce style wildcard (matching any number of "...", "*" and "?" patterns in 
 any location), a full path name, or a reference to a tagged collection of files. Attributes are denoted by prefixing a `#` 
-character. Files may be added to a tag set using the `<Tag>` task, which also allows the performance of union/difference style 
+character. Files may be added to a tag set using the `<Tag>` task, which also enables the performance of union/difference style 
 operations. Each node can declare multiple outputs in the form of a list of named tags, which other nodes can then depend on.
 
 Graphs may be executed in parallel as part of a build system. To do so, the initial graph configuration is generated by 
diff --git a/Engine/Documentation/Source/Programming/Development/BuildGraph/ScriptAnatomy/Conditions/BuildGraphScriptConditions.INT.udn b/Engine/Documentation/Source/Programming/Development/BuildGraph/ScriptAnatomy/Conditions/BuildGraphScriptConditions.INT.udn
--- a/Engine/Documentation/Source/Programming/Development/BuildGraph/ScriptAnatomy/Conditions/BuildGraphScriptConditions.INT.udn
+++ b/Engine/Documentation/Source/Programming/Development/BuildGraph/ScriptAnatomy/Conditions/BuildGraphScriptConditions.INT.udn
@@ -5,6 +5,7 @@ Description:Learn the syntax needed to write BuildGraph script conditions.
 version: 4.13
 parent:Programming/Development/BuildGraph
 type:Reference
+tags:Programming
 tags:BuildGraph
 
 [VAR:TopicCompact]
diff --git a/Engine/Documentation/Source/Programming/Development/BuildGraph/ScriptAnatomy/Elements/BuildGraphScriptElements.INT.udn b/Engine/Documentation/Source/Programming/Development/BuildGraph/ScriptAnatomy/Elements/BuildGraphScriptElements.INT.udn
--- a/Engine/Documentation/Source/Programming/Development/BuildGraph/ScriptAnatomy/Elements/BuildGraphScriptElements.INT.udn
+++ b/Engine/Documentation/Source/Programming/Development/BuildGraph/ScriptAnatomy/Elements/BuildGraphScriptElements.INT.udn
@@ -5,6 +5,7 @@ Description:Learn about the elements of a BuildGraph.
 version: 4.13
 parent:Programming/Development/BuildGraph
 type:Reference
+tags:Programming
 tags:BuildGraph
 
 [VAR:TopicCompact]
diff --git a/Engine/Documentation/Source/Programming/Development/BuildGraph/ScriptAnatomy/Tasks/BuildGraphScriptTasks.INT.udn b/Engine/Documentation/Source/Programming/Development/BuildGraph/ScriptAnatomy/Tasks/BuildGraphScriptTasks.INT.udn
--- a/Engine/Documentation/Source/Programming/Development/BuildGraph/ScriptAnatomy/Tasks/BuildGraphScriptTasks.INT.udn
+++ b/Engine/Documentation/Source/Programming/Development/BuildGraph/ScriptAnatomy/Tasks/BuildGraphScriptTasks.INT.udn
@@ -1,33 +1,13 @@
-Availability: Docs
+Availability:Docs
 Title: BuildGraph Script Tasks
-Crumbs: %ROOT%, Programming, Programming/Development, Programming/Development/BuildGraph, Programming/Development/BuildGraph/ScriptAnatomy
+Crumbs:
 Description: Learn how BuildGraph can help you create custom tasks.
 version: 4.13
 parent:Programming/Development/BuildGraph
-type:Reference
-related: Programming/Development/BuildGraph/ScriptAnatomy
+type:overview
+tags:Programming
 tags:BuildGraph
 
-[VAR:TopicCompact]
-[OBJECT:TopicCompact]
-	[PARAM:image]
-		![%Programming/Development/BuildGraph/ScriptAnatomy/Tasks%](tasks_topic.png)
-	[/PARAM]
-	[PARAM:icon]
-		![](%ROOT%/reference_icon.png)(convert:false)
-	[/PARAM]
-	[PARAM:title]
-		%Programming/Development/BuildGraph/ScriptAnatomy/Tasks:title%
-	[/PARAM]
-	[PARAM:description]
-		%Programming/Development/BuildGraph/ScriptAnatomy/Tasks:description%
-	[/PARAM]
-	[PARAM:path]
-		[RELATIVE:Programming/Development/BuildGraph/ScriptAnatomy/Tasks]
-	[/PARAM]
-[/OBJECT]
-[/VAR]
-
 [TOC(start:1 end:3)]
 
 [EXCERPT:BuildGraphScriptTasks]
diff --git a/Engine/Documentation/Source/Programming/Development/BuildingUnrealEngine/BuildingUnrealEngine.INT.udn b/Engine/Documentation/Source/Programming/Development/BuildingUnrealEngine/BuildingUnrealEngine.INT.udn
--- a/Engine/Documentation/Source/Programming/Development/BuildingUnrealEngine/BuildingUnrealEngine.INT.udn
+++ b/Engine/Documentation/Source/Programming/Development/BuildingUnrealEngine/BuildingUnrealEngine.INT.udn
@@ -7,9 +7,12 @@ Related: GettingStarted
 Related: Engine/Content/QuickStart
 Related: Programming/QuickStart
 Related: Engine/QuickStart
-Related: Support/Builds
+Related: Programming/UnrealBuildSystem
 Version: 4.11
 SkillLevel: Advanced
+topic-image:topic_buildingUE.png
+tags:Programming
+tags:Unreal Build System
 
 %Globals:osselect_linux%
 
diff --git a/Engine/Documentation/Source/Programming/Plugins/Plugins.INT.udn b/Engine/Documentation/Source/Programming/Plugins/Plugins.INT.udn
--- a/Engine/Documentation/Source/Programming/Plugins/Plugins.INT.udn
+++ b/Engine/Documentation/Source/Programming/Plugins/Plugins.INT.udn
@@ -1,6 +1,7 @@
 Availability:Public
 Title:Plugins
-Crumbs:%ROOT%, Programming
+Crumbs:
+Parent:Programming
 Description:How to create Unreal Engine plugins.
 Version: 4.9
 
@@ -70,10 +71,10 @@ into adding at a later date.
 Plugins are always located in your plugin's directory. In order for plugins to be found, they must be located in one of 
 the valid search paths for plugins in Unreal Engine. 
 
-| Plugin type      | Search path                                        |
-| ---------------- | -------------------------------------------------- |
-| Engine plugin    | /UE4 root/Engine/Plugins/My Engine Plugin/ |
-| Game plugin      | /My project/Plugins/My Game Plugin/        |
+| Plugin Type          | Search Path                                        |
+| -------------------- | -------------------------------------------------- |
+| **Engine**           | /UE4 Root/Engine/Plugins/My Engine Plugin/         |
+| **Game**             | /My Project/Plugins/My Game Plugin/                |
 
 You can also organize plugins into subdirectories under the base Plugins folder. The engine will scan all of your sub-folders 
 under the base Plugins folder for plugins to load, but it will never scan into subdirectories beneath where a plugin was already found.
diff --git a/Engine/Documentation/Source/Programming/ProgrammingGuide.INT.udn b/Engine/Documentation/Source/Programming/ProgrammingGuide.INT.udn
--- a/Engine/Documentation/Source/Programming/ProgrammingGuide.INT.udn
+++ b/Engine/Documentation/Source/Programming/ProgrammingGuide.INT.udn
@@ -2,6 +2,15 @@ Availability:Public
 Title:Programming Guide
 Crumbs: %ROOT%
 Description:Information for programmers developing with Unreal Engine.
+Related: Programming/Development/VisualStudioSetup
+Related: Programming/Development/BuildingUnrealEngine
+Related: Programming/Development/CompilingProjects
+Related: Programming/Development/CodingStandard
+Related: Programming/Development/ManagingGameCode
+Related: Programming/Development/Tools
+Related: Programming/Development/BuildConfiguration
+Related: Programming/Development/BuildGraph
+Related: Programming/Development/InstalledBuildReference
 version: 4.9
 parent:%ROOT%
 order:4
diff --git a/Engine/Documentation/Source/Programming/QuickStart/QuickStart.INT.udn b/Engine/Documentation/Source/Programming/QuickStart/QuickStart.INT.udn
--- a/Engine/Documentation/Source/Programming/QuickStart/QuickStart.INT.udn
+++ b/Engine/Documentation/Source/Programming/QuickStart/QuickStart.INT.udn
@@ -3,9 +3,11 @@ Title:Programming Quick Start
 Crumbs:%ROOT%, Programming
 Description:Create your first code project and add a new C++ class.
 prereq:GettingStarted/Installation
+prereq:Programming/Development/VisualStudioSetup
 prereq:Engine/Basics/Projects/Browser
 SkillLevel: Beginner
 Version: 4.9
+type:quick start
 
 [VAR:Topic]
 [OBJECT:Topic]
@@ -53,10 +55,7 @@ Version: 4.9
 	[/PARAMLITERAL]
 [/OBJECT]
 
-The primary goal of this quick start guide is to introduce you to Unreal Engine 4's (UE4) development environment. By the end of this guide, you'll know how to 
-setup and develop C++ Projects in UE4. This guide shows you how to create a new Unreal Engine project, add a new C++ class to it, compile the project, and add
-an instance of a new class to your level. By the time you reach the end of this guide, you'll be able to see your programmed **Actor** floating above a table 
-in the level.
+The primary goal of this quick start guide is to introduce you to Unreal Engine 4's (UE4) development environment. By the end of this guide, you'll know how to set up and develop C++ Projects in UE4. This guide shows you how to create a new Unreal Engine project, add a new C++ class to it, compile the project, and add an instance of a new class to your level. By the time you reach the end of this guide, you'll be able to see your programmed **Actor** floating above a table in the level.
 
 [OBJECT:TopicButtonList]
 	[PARAM:icon]
diff --git a/Engine/Documentation/Source/Programming/UnrealArchitecture/Reference/Interfaces/Interfaces.INT.udn b/Engine/Documentation/Source/Programming/UnrealArchitecture/Reference/Interfaces/Interfaces.INT.udn
--- a/Engine/Documentation/Source/Programming/UnrealArchitecture/Reference/Interfaces/Interfaces.INT.udn
+++ b/Engine/Documentation/Source/Programming/UnrealArchitecture/Reference/Interfaces/Interfaces.INT.udn
@@ -1,56 +1,65 @@
 Availability:Public
 Title: Interfaces
-Crumbs:%ROOT%, Programming, Programming/UnrealArchitecture/Reference
+Crumbs:
 Description:Reference to creating and implementing interfaces.
-Version: 4.9
+Tags: Programming
+Version: 4.14
+Parent:Programming/UnrealArchitecture/Reference
+Order:
+SkillLevel:Intermediate
+Type:overview
 
 [TOC(start:2)]
 
 
-Declaring an interface class is very similar to declaring a regular class. The only major difference in the syntax is the descriptors are passed to the `UINTERFACE` macro instead of the `UCLASS` macro.
+Interface classes are useful for ensuring that a set of (potentially) unrelated classes implement a common set of functions. This is very useful in cases where some game functionality may be shared by large, complex classes, that are otherwise dissimilar. For example, a game might have a system whereby entering a trigger volume can activate traps, alert enemies, or award points to the player. This might be implemented by a "ReactToTrigger" function on traps, enemies, and point-awarders. However, traps may be derived from `AActor`, enemies from a specialized `APawn` or `ACharacter` subclass, and point-awards from `UDataAsset`. All of these classes need shared functionality, but have no common ancestor other than `UObject`. In this case, an interface is recommended.
+
+## Interface Declaration
+
+Declaring an interface class is similar to declaring a normal Unreal class, but with two main differences. First, an interface class uses the UINTERFACE macro instead of the UCLASS macro, and inherits from `UInterface` instead of `UObject` directly.
 
 	UINTERFACE([specifier, specifier, ...], [meta(key=value, key=value, ...)])
-	class ClassName : ParentName
+	class UClassName : public UInterface
+	{
+		GENERATED_BODY()
+	};
 
-## Interface Declaration
+Second, the UINTERFACE class is not the actual interface; it is an empty class that exists only for visibility to Unreal Engine's reflection system. The actual interface that will be inherited by other classes must have the same class name, but with the initial "U" changed to an "I".
 
-In your .h file (e.g. `ItemInterface.h`):
+In your .h file (e.g. `ReactToTriggerInterface.h`):
 
     #pragma once
 
-    #include "ItemInterface.generated.h"
+    #include "ReactToTriggerInterface.generated.h"
 
     UINTERFACE(Blueprintable)
-    class UItemInterface : public UInterface
+    class UReactToTriggerInterface : public UInterface
     {
         GENERATED_BODY()
     };
 
-    class IItemInterface
+    class IReactToTriggerInterface
     {    
 		GENERATED_BODY()
 
     public:
-
-        UFUNCTION(BlueprintNativeEvent, BlueprintCallable, Category="Item")
-        FItemStruct GetItem() const;
+		/** React to a trigger volume activating this object. Return true if the reaction succeeds. */
+        UFUNCTION(BlueprintCallable, BlueprintCallable, Category="Trigger Reaction")
+        bool ReactToTrigger() const;
     };
 
-In your .cpp file (e.g. `ItemInterface.cpp`):
-
-    #include "ItemInterface.h"
+Functions can be written directly in the .h file, and this is commonly done when a function does nothing by default, or has a trivial behavior like returning false, zero, an empty string (or something similar). More complex functions can be written in the .cpp file to reduce compile times and pure virtual functions are supported. Although our sample function is trivial and could reasonably be written in the .h file, we'll write it in the .cpp file. Your .cpp file (e.g. `ReactToTriggerInterface.cpp`) should now contain this:
 
-    UItemInterface::UItemInterface(const FObjectInitializer& ObjectInitializer) : Super(ObjectInitializer)
-    {
-    }
-
-The declaration consists of a standard C++ class declaration for the class. Above the standard declaration, descriptors such as class specifiers and metadata are passed to the UINTERFACE macro. These are used to create the UInterface for the interface being declared, which can be thought of as the engine's specialized representation of the class. Also, the GENERATED_UINTERFACE_BODY() macro must be placed at the very beginning of the class body.
+    #include "ReactToTriggerInterface.h"
+	
+	bool IReactToTriggerInterface::ReactToTrigger() const
+	{
+		return false;
+	}
 
-For interfaces that are to be recognized by the engine, you must create two classes, a UINTERFACE class that inherits from UInterface, and a pure interface.
+The "U-prefixed" class needs no constructor or any other functions, while the "I-prefixed" class will contain all interface functions and is the one that will actually be inherited by your other classes.
 
-The UINTERFACE class should be prefixed with U, as it extends from UObject. The Interface should be prefixed with I.
-
-The `Blueprintable` specifier is required if you want to allow Blueprints to implement this interface.
+[REGION:NOTE]The `Blueprintable` specifier is required if you want to allow Blueprints to implement this interface.[/REGION:NOTE]
 
 ### Interface Specifiers
 
@@ -58,27 +67,35 @@ The `Blueprintable` specifier is required if you want to allow Blueprints to imp
 
 ## Implementing Your Interface in C++
 
-When declaring your class in C++, simply inherit from your interface (not the UInterface class though).
+To use your interface in a new class, simply inherit from your "I-prefixed" interface class (in addition to whatever `UObject`-based class you are using).
 
-    class AItem : public AActor, public IItemInterface
+    class ATrap : public AActor, public IReactToTriggerInterface
     {
         GENERATED_BODY()
 
     public:
-
-        FItemStruct GetItem_Implementation() const;
+        virtual bool ReactToTrigger() const override;
     };
 
-## Determining If A Given Actor Implements Your Interface
+## Determining If A Given Class Implements Your Interface
 
-For compatability with both C++ and Blueprint classes that implement your interface, use the following code:
+For compatability with both C++ and Blueprint classes that implement your interface, use either of the following functions:
 
-    obj->GetClass()->ImplementsInterface(UItemInterface::StaticClass());
+    bool bIsImplemented = OriginalObject->GetClass()->ImplementsInterface(UReactToTriggerInterface::StaticClass()); // bIsImplemented will be true if OriginalObject implements UReactToTriggerInterface.
+	IReactToTriggerInterface* ReactingObject = Cast<IReactToTriggerInterface>(OriginalObject); // ReactingObject will be non-null if OriginalObject implements UReactToTriggerInterface.
 
 [REGION:note]
-You must use the UInterface version of the class to check if the interface is implemented.
+Attempting to use `Cast` on the "U-prefixed" class will fail, while the `StaticClass` function is not implemented in the "I-prefixed" class and will not compile.
 [/REGION]
 
+## Casting To Other Unreal Types
+
+Unreal Engine's casting system supports casting from one interface to another, or from an interface to an Unreal type, where appropriate.
+
+	IReactToTriggerInterface* ReactingObject = Cast<IReactToTriggerInterface>(OriginalObject); // ReactingObject will be non-null if the interface is implemented.
+	ISomeOtherInterface* DifferentInterface = Cast<ISomeOtherInterface>(ReactingObject); // DifferentInterface will be non-null if ReactingObject is non-null and also implements ISomeOtherInterface.
+	AActor* Actor = Cast<AActor>(ReactingObject); // Actor will be non-null if ReactingObject is non-null and OriginalObject is an AActor or AActor-derived class.
+
 ## Blueprint Implementable Classes
 
-If you want to allow Blueprints to implement this interface, you must use the `Blueprintable` metadata specifier, and every function in your interface must be a `BlueprintNativeEvent` or a `BlueprintImplementableEvent`. Otherwise, a Blueprint class implementing the interface would be unable to implement all of the methods of the interface.
+If you want Blueprints to be able to implement this interface, you must use the `Blueprintable` metadata specifier. Every interface function that your Blueprint class is intended to override, must be a `BlueprintNativeEvent` or a `BlueprintImplementableEvent`. Functions marked as `BlueprintCallable` will still be able to be called, but not overridden. All other functions will be inaccessible from Blueprints.
diff --git a/Engine/Documentation/Source/Programming/UnrealArchitecture/Reference/Metadata/MetadataSpecifiers.INT.udn b/Engine/Documentation/Source/Programming/UnrealArchitecture/Reference/Metadata/MetadataSpecifiers.INT.udn
--- a/Engine/Documentation/Source/Programming/UnrealArchitecture/Reference/Metadata/MetadataSpecifiers.INT.udn
+++ b/Engine/Documentation/Source/Programming/UnrealArchitecture/Reference/Metadata/MetadataSpecifiers.INT.udn
@@ -8,7 +8,7 @@ When declaring classes, functions, and interfaces, metadata specifiers can be ad
 aspects of the engine and editor.
 
 [EXCERPT:MetadataUsage]
-The usage of metadata specifiers differs from regular class, function, and interface specifiers.
+The usage of metadata specifiers differs between regular classes, functions, and interfaces.
 [/EXCERPT:MetadataUsage]
 
 To use metadata specifies, use the meta specifier in a `UCLASS`, `UINTERFACE`, or `UFUNCTION` macro:
@@ -17,12 +17,14 @@ To use metadata specifies, use the meta specifier in a `UCLASS`, `UINTERFACE`, o
 
 ## Class Metadata Specifiers
 [EXCERPT:ClassMetadata]
+Classes can use the following metatag specifier:
 * [](Programming\UnrealArchitecture\Reference\Metadata\BlueprintSpawnableComponent)
 [/EXCERPT:ClassMetadata]
 
 
 ## Function Metadata Specifiers
 [EXCERPT:FunctionMetadata]
+Functions can use the following metatag specifiers:
 * [](Programming\UnrealArchitecture\Reference\Metadata\BlueprintInternalUseOnly)
 * [](Programming\UnrealArchitecture\Reference\Metadata\BlueprintProtected)
 * [](Programming\UnrealArchitecture\Reference\Metadata\DeprecatedFunction)
@@ -33,5 +35,6 @@ To use metadata specifies, use the meta specifier in a `UCLASS`, `UINTERFACE`, o
 
 ## Interface Metadata Specifiers
 [EXCERPT:InterfaceMetadata]
+Interfaces can use the following metatag specifier:
 * [](Programming\UnrealArchitecture\Reference\Metadata\CannotImplementInterfaceBlueprint)
 [/EXCERPT:InterfaceMetadata]
diff --git a/Engine/Documentation/Source/Programming/UnrealBuildSystem/Configuration/BuildConfiguration.INT.udn b/Engine/Documentation/Source/Programming/UnrealBuildSystem/Configuration/BuildConfiguration.INT.udn
--- a/Engine/Documentation/Source/Programming/UnrealBuildSystem/Configuration/BuildConfiguration.INT.udn
+++ b/Engine/Documentation/Source/Programming/UnrealBuildSystem/Configuration/BuildConfiguration.INT.udn
@@ -1,28 +1,26 @@
 Availability:Public
 Title:Configuring Unreal Build System
-Crumbs: %ROOT%, Programming, Programming/UnrealBuildSystem
-Description:Configuring Unreal Build System to control how the engine is built.
+Crumbs:
+Description:Configure Unreal Build System to control how the engine is built.
 Version: 4.14
+tags:Programming
+tags:Unreal Build System
+level:Intermediate
+parent:Programming/UnrealBuildSystem
+topic-image:topic_buildConfigProps.png
+
+[TOC(start:2 end:4)]
 
 UnrealBuildTool reads settings from XML configuration files in the following locations:
 
 	* Engine/Saved/UnrealBuildTool/BuildConfiguration.xml
-	* (User Folder)/AppData/Roaming/Unreal Engine/UnrealBuildTool/BuildConfiguration.xml
-	* (Documents Folder)/Unreal Engine/UnrealBuildTool/BuildConfiguration.xml
-
-These files are added to the generated UE4 project under the Config/UnrealBuildTool folder. An XML schema is also generated, giving IntelliSense information when editing in Visual Studio.
-
-## Categories
+	* *User Folder/AppData*/Roaming/Unreal Engine/UnrealBuildTool/BuildConfiguration.xml
+	* *My Documents*/Unreal Engine/UnrealBuildTool/BuildConfiguration.xml
 
-Settings are grouped into categories, and are represented by nested XML elements. In the example below, the bUseUnityBuild setting in the BuildConfiguration category is set to false.
+These files are added to the generated UE4 project under the Config/UnrealBuildTool folder.
 
-    <?xml version="1.0" encoding="utf-8" ?>
-    <Configuration xmlns="https://www.unrealengine.com/BuildConfiguration">
-	    <BuildConfiguration>
-			<bUseUnityBuild>false</bUseUnityBuild>
-	    </BuildConfiguration>
-    </Configuration>
+## UnrealBuildTool Configuration Properties
 
-The available categories and their settings are listed below.
+These properties will help you set and customize your build configuration(s).
 
 [INCLUDE:Programming/UnrealBuildSystem/Configuration/BuildConfigProperties]
