Copying //UE4/Dev-Documentation to //UE4/Samples-Main (Source: //UE4/Dev-Documentation @ 3372893)

diff --git a/Engine/Documentation/Source/Programming/UnrealBuildSystem/IWYUReferenceGuide/IWYUReferenceGuide.INT.udn b/Engine/Documentation/Source/Programming/UnrealBuildSystem/IWYUReferenceGuide/IWYUReferenceGuide.INT.udn
new file mode 100644
--- /dev/null
+++ b/Engine/Documentation/Source/Programming/UnrealBuildSystem/IWYUReferenceGuide/IWYUReferenceGuide.INT.udn
@@ -0,0 +1,126 @@
+Availability: Public
+Title: IWYU Reference Guide
+Crumbs:
+Description:An overview of the updated codebase for UE4, which uses an Include-What-You-Use (IWYU) dependency model.
+version: 4.15
+parent:Programming/UnrealBuildSystem
+type:Overview
+level:advanced
+tags:Programming
+tags:Unreal Build System
+related:Programming/UnrealBuildSystem
+prereq:Programming/Development
+prereq:Programming
+topic-image:topic_IWYU.png
+
+[TOC(start:1 end:4)]
+
+![](IWYU.png)
+
+Include-What-You-Use (**IWYU**), as the name implies, means that the Engine's source code only includes the dependencies that it needs to compile. The purpose of IWYU is to avoid including 
+monolithic header files, such as `Engine.h` or `UnrealEd.h`, thereby mitigating superfluous dependencies. The following reference guide tells you what it means to IWYU, including a 
+high-level explanation of how to enable IWYU, ensuring that your project adheres to IWYU conventions. Additionally, if you opt into using IWYU mode for your game project(s), you will learn 
+some general tips that will help you get the most out of working in IWYU mode. 
+
+[REGION:note]
+
+IWYU mode is disabled by default for games and game plugins; however, IWYU mode is enabled by default for the Engine and Engine plugins. 
+
+[/REGION]
+
+## What it Means to IWYU
+
+In previous versions of Unreal Engine 4 (UE4), the majority of engine functionality was included via large, module-centric header files, such as `Engine.h` and `UnrealEd.h` and fast compile times 
+were dependent on those files being compiled quickly through Precompiled Header (PCH) files. As the engine grew, this became a bottleneck. 
+
+With IWYU, every file includes only what it needs, and any PCH file we choose to use, purely acts as a layer of optimization on top of the underlying source files. They can be modified to minimize build times, 
+independently of changing the source files themselves, and without affecting whether the code compiles successfully or not. 
+
+When writing IWYU code, there are four specific conventions that we adopt: 
+
+1.	**All header files include their required dependencies.**
+1.	**\*.cpp files include their matching \*.h files first.**
+1.	**PCH files are no longer explicitly included.**
+1.	**Monolithic header files are no longer included.**
+
+###IWYU Conventions
+
+The following descriptions of IWYU conventions should give you a good idea about what it means to IWYU.
+
+1.	**All header files include their required dependencies.**
+	
+	*	There is a **CoreMinimal** header file containing a set of ubiquitous types (including FString, FName, TArray, etc.) from UE4's Core programming environment. 
+	
+	*	The `CoreMinimal` header file (located under the UE4 root directory at `\Engine\Source\Runtime\Core\Public\CoreMinimal.h`) is included first by most of the Engine's header files.
+
+	![](CoreMinimalHeader.png)
+
+	*	Within the `Core` module, most header files include the `CoreTypes.h` header file first. This only includes typedefs for primitive C++ types, UE4 build macros, and directives to configure the compile environment.
+
+	[REGION:tip]
+	The main takeaway is that every header file now includes everything that it needs to compile.
+	[/REGION]
+
+1.	**\*.cpp files include their matching \*.h file first.**
+
+	![](CPPFileIncludingHeader.png)
+
+	*	To verify that all of your source files include all of their required dependencies, compile your game project in non-unity mode with PCH files disabled.
+
+1.	**PCH files are no longer explicitly included.**
+
+	*	Although PCH files are still used, they're force-included on the compiler command line by UnrealBuildTool (UBT).
+
+1.	**Monolithic header files are no longer included.**
+
+	*	The compiler will emit a warning if your Engine code includes monolithic header files (such as `Engine.h` or `UnrealEd.h`).
+
+	[REGION:note]
+	Monolithic header files still exist in UE4 for compatibility with game projects, and (by default) a warning won't be emitted if your game projects include them.
+	[/REGION]
+
+##Enabling IWYU
+
+Prior to establishing IWYU conventions with the release of version 4.15, UE4 code typically included a PCH file at the top of every CPP file, which is contrary to what a developer 
+wants their IWYU compliant code to include. With the newly-established IWYU conventions, PCH files can be thought of as layers of compile-time optizations that are applied separately from 
+how the code was originally authored. So, rather than composing and including PCH files, we leave it to UBT to decide which PCH file to use (if any).
+
+If you want to enable IWYU mode, ensuring that a module complies with IWYU conventions, open the module's \*.build.cs file and set `PCHUsage` to `PCHUsageMode.UseExplicitOrSharedPCHs`. 
+
+![](ExamplePlugin_PCHUsage.png)
+
+Setting `PCHUsage` to `PCHUsageMode.UseExplicitOrSharedPCHs` creates an explicit PCH file for a module only if it has a `PrivatePCHHeaderFile` setting in the  module's \*.build.cs file. 
+Otherwise, the module will share a PCH with another module, saving the tool from generating more PCH files than necessary. Also, keep in mind that when you enable `UseExplicitOrSharedPCHs` 
+mode, the source file must include its matching header file. Alternatively, if you want a module to opt-out of complying with IWYU conventions, you can set `PCHUsage` to `PCHUsageMode.UseSharedPCHs`.
+
+![](ExamplePlugin_UseSharedPCHs.png)
+
+[REGION:note]
+
+After converting the engine's code base to an IWYU model, we've observed significant improvements with UE4 compile times. 
+
+[/REGION]
+
+###Running in IWYU Mode
+
+If you're running your game in IWYU mode, you'll need to ensure that your \*.cpp file(s) include their corresponding \*.h file(s) first. This is a useful practice, because it enables the compiler to 
+ensure that the \*.h file includes all of its required dependencies (when PCH files and unity builds are disabled). Unreal Build Tool (UBT) will emit a warning if you don't include the matching header file first 
+(for its corresponding CPP file). 
+
+If you want to disable the compiler from emitting warnings, you can set `bEnforceIWYU` to `false` in the module's \*.build.cs file.
+
+![](ModuleBuildCS_bEnforceIWYUfalse.png)
+
+[REGION:tip]
+
+If you want to disable the warning for an entire target, you can set `bEnforceIWYU` to false in the \*.target.cs file.
+
+[/REGION]
+
+## General Tips
+
+If you want your game to opt-in to IWYU, there are a few tips to keep in mind:
+
+1.	Include `CoreMinimal.h` at the top of each header file.
+1.	To verify that all of your source files include all of their required dependencies, compile your game project in non-unity mode with PCH files disabled.
+1.	If you need to access **UEngine** or **GEngine**, which are defined in `Runtime\Engine\Classes\Engine\Engine.h`, you can `#include Engine/Engine.h` (distinguishing from the monolithic header file, which is located at `Runtime\Engine\Public\Engine.h`).
diff --git a/Engine/Documentation/Source/Programming/UnrealBuildSystem/ProjectFileGenerator/ProjectFileGenerator.INT.udn b/Engine/Documentation/Source/Programming/UnrealBuildSystem/ProjectFileGenerator/ProjectFileGenerator.INT.udn
--- a/Engine/Documentation/Source/Programming/UnrealBuildSystem/ProjectFileGenerator/ProjectFileGenerator.INT.udn
+++ b/Engine/Documentation/Source/Programming/UnrealBuildSystem/ProjectFileGenerator/ProjectFileGenerator.INT.udn
@@ -1,8 +1,19 @@
 Availability:Public
 Title:Automatic Project File Generation
-Crumbs:%ROOT%, Programming, Programming/UnrealBuildSystem
+Crumbs:
+Parent:Programming/UnrealBuildSystem
 Description:Process for automatically generating project files for games and modules in the current workspace.
 Version: 4.9
+prereq:Programming/Development/BuildingUnrealEngine
+topic-image:topic_prjFileGen.png
+tags:Programming
+tags:Unreal Build System
+
+[TOC(start:1 end:2)]
+
+[REGION:warning]
+The following guide is meant for users who've downloaded and built UE4 from its source code, which we've made available on GitHub.
+[/REGION]
 
 If you download the latest Unreal Engine 4 (UE4) code, you might notice that there are no Visual Studio (VS) or Xcode project files included for compiling and running the engine or example games.  You will need to run a script that generates project files for you to load:
 
diff --git a/Engine/Documentation/Source/Programming/UnrealBuildSystem/UnrealBuildSystem.INT.udn b/Engine/Documentation/Source/Programming/UnrealBuildSystem/UnrealBuildSystem.INT.udn
--- a/Engine/Documentation/Source/Programming/UnrealBuildSystem/UnrealBuildSystem.INT.udn
+++ b/Engine/Documentation/Source/Programming/UnrealBuildSystem/UnrealBuildSystem.INT.udn
@@ -1,82 +1,66 @@
 Availability:Public
 Title: Unreal Build System
-Crumbs: %ROOT%, Programming
+Crumbs:
 Description:Unreal Engine's custom build system for compiling the engine and game projects.
-version:4.9
-[VAR:Topic]
-[OBJECT:Topic]
-	[PARAM:image]
-		![%Programming/UnrealBuildSystem:title%](Engine/build_topic.png)
-	[/PARAM]
-	[PARAM:icon]
-		![](%ROOT%/build_icon.png)(convert:false)
-	[/PARAM]
-	[PARAM:title]
-		%Programming/UnrealBuildSystem:title%
-	[/PARAM]
-	[PARAM:description]
-		%Programming/UnrealBuildSystem:description%
-	[/PARAM]
-	[PARAM:path]
-		[RELATIVE:Programming/UnrealBuildSystem]
-	[/PARAM]
-[/OBJECT]
-[/VAR]
+version:4.15
+parent:Programming
+tags:Programming
+tags:Unreal Build System
+type:overview
+topic-image:Engine/build_topic.png
+
+[TOC (start:1 end:3)]
 
 [REGION:fullwidth]
 ![Unreal Build System](build.png)
 [/REGION]
 
-UE4 is built using a small suite of custom tools designed to abstract cross-platform differences and allow engine features to be enabled, disabled, and built in a myriad different ways optimized for your game. 
+**Unreal Build System** refers to a collection of tools that automate the process of building Unreal Engine 4 (UE4). At a high level, UnrealBuildTool and UnrealHeaderTool support Unreal Build System.
+This guide provides and overview of UnrealBuildTool and UnrealHeaderTool, including a brief discussion about UE4's modular architecture. If you want to use this collection of build tools, we've included some guides that will help you set up and configure Unreal 
+Build System for your own projects. Finally, we've included some guides for advanced users, those wanting to learn more about the Include-What-You-Use (IWYU) dependency model and the BuildGraph scripting system.
 
-### UnrealBuildTool
+[REGION:note]
+If you would like to debug UnrealBuildTool or UnrealAutomationTool on Mac OS X, please install [MonoDevelop 4.2.2](http://monodevelop.com/Download) and [Mono 3.2.7](http://www.go-mono.com/mono-downloads/download.html). 
+[/REGION]
 
-**UnrealBuildTool** (often abbreviated as UBT) orchestrates the process of building games and programs based on the UE4 source code. Any game or program you can build is a [target](Programming\UnrealBuildSystem\TargetFiles), and is defined by rules given in a C# source file with the .target.cs extension.
+## UnrealBuildTool
 
-The building blocks for a target are [modules](Programming\UnrealBuildSystem\ModuleFiles). A module consists of C++ source code and a C# .build.cs file declaring how to build it. A module may depend on other modules to be built, or may declare include paths and macros for other modules that depend on it. The core of the engine is a large number of modules, which are augmented by modules you define for your own game. 
+**UnrealBuildTool** (UBT) is a custom tool that manages the process of building the UE4 source code across a variety of build configurations. Check out `BuildConfiguration.cs` to explore various 
+user-configurable build options.
 
-You can configure per-user settings for UnrealBuildTool through XML [configuration files](Programming\UnrealBuildSystem\Configuration).
+## UnrealHeaderTool
 
-[REGION:note]
-If you would like to debug UnrealBuildTool or UnrealAutomationTool on Mac OS X, you must install [MonoDevelop 4.2.2](http://monodevelop.com/Download) and [Mono 3.2.7](http://www.go-mono.com/mono-downloads/download.html). 
+**UnrealHeaderTool** (UHT) is a custom parsing and code-generation tool that supports the UObject system.  Code compilation happens in two phases:
+
+1.	UHT is invoked, which parses the C++ headers for Unreal-related class metadata and generates custom code to implement the various UObject-related features.
+1.	The normal C++ compiler is invoked to compile the results.
+
+When compiling, it is possible for either tool to emit errors, so be sure to look carefully.
+
+## Modular Architecture
+
+UE4 is split into many modules.  Each module has a `.build.cs` file that controls how it is built, including options for defining module dependencies, additional libraries, include paths, etc. By default, these modules are compiled into DLLs and loaded by a single executable. You can choose to build a monolithic executable in the `BuildConfiguration.cs` file.  
+
+## Using Unreal Build System
+
+Before using Unreal Build System, it is important to understand that the build process executes independently of any project files for the development environment, such as `.sln` or `.vcproj` files (for Visual Studio). These files are useful to have for editing purposes though, 
+so there is a tool being provided to generate them dynamically (based on the contents of your project directory tree). You can run this tool with the `GenerateProject.bat` file located in your `[Root UE4]` directory.
+
+[REGION:tip]
+Running `GenerateProject.bat` from time-to-time will keep your code editor up-to-date with files that are being added to (or being removed) from disk.
 [/REGION]
 
-### UnrealHeaderTool
-
-**UnrealHeaderTool** (often abbreviated as UHT) parses C++ source files and generates reflection data. This reflection data is used by the engine to allow you to edit objects you create in the editor, to load and save those classes into packages, marshall those objects over the network, and so on.
-
-### Generating Project Files
-The build process happens independently of any project files for the development environment, such as .sln or .vcproj files for Visual Studio.  These files are useful to have for editing purposes though, 
-so there is a provided tool to generate them dynamically based on the contents of your project directory tree. You can run this tool with the [GenerateProjectFiles.bat](Programming\UnrealBuildSystem\ProjectFileGenerator) file located in your root /UE4 directory.
- It is generally useful to run this batch file from time to time to keep your code editor in sync with code files you have added or removed on disk.
-
-[OBJECT:TopicButtonList]
-	[PARAM:icon]
-		![](%ROOT%/start_icon.png)(convert:false)
-	[/PARAM]
-	[PARAM:title]
-		Essentials
-	[/PARAM]
-	[PARAM:description]
-	[/PARAM]
-	[PARAM:links]
-		* [](Programming/UnrealBuildSystem/ProjectFileGenerator "%Programming/UnrealBuildSystem/ProjectFileGenerator:description%")
-		* [](Programming/Development/BuildingUnrealEngine "%Programming/Development/BuildingUnrealEngine:description%")
-	[/PARAM]
-[/OBJECT]
-
-[OBJECT:TopicButtonList]
-	[PARAM:icon]
-		![](%ROOT%/profile_icon.png)(convert:false)
-	[/PARAM]
-	[PARAM:title]
-		Setup & Configuration
-	[/PARAM]
-	[PARAM:description]
-	[/PARAM]
-	[PARAM:links]
-		* [](Programming/UnrealBuildSystem/Configuration "%Programming/UnrealBuildSystem/Configuration:description%")
-		* [](Programming/UnrealBuildSystem/TargetFiles "%Programming/UnrealBuildSystem/TargetFiles:description%")
-	[/PARAM]
-[/OBJECT]
+### Setup
+
+%Programming/UnrealBuildSystem/ProjectFileGenerator:topic%
+%Programming/Development/BuildingUnrealEngine:topic%
+
+### Configuration
+
+%Programming/UnrealBuildSystem/Configuration:topic%
+%Programming/UnrealBuildSystem/TargetFiles:topic%
+
+## Advanced Topics
 
+%Programming/UnrealBuildSystem/IWYUReferenceGuide:topic%
+%Programming/Development/BuildGraph:topic%
diff --git a/Engine/Documentation/Source/Resources/Showcases/Showcases.INT.udn b/Engine/Documentation/Source/Resources/Showcases/Showcases.INT.udn
--- a/Engine/Documentation/Source/Resources/Showcases/Showcases.INT.udn
+++ b/Engine/Documentation/Source/Resources/Showcases/Showcases.INT.udn
@@ -70,7 +70,7 @@ The corresponding documentation is designed to provide some insight into how eac
 	Water Planes
 	[/PARAMLITERAL]
 	[PARAMLITERAL:playlistid]
-	PLZlv_N0_O1gYq4tiS45UAgghUxoF916rH
+	PLZlv_N0_O1ga0aV9jVqJgog0VWz1cLL5f
 	[/PARAMLITERAL]
 	[PARAMLITERAL:videoid]
 	aJ1M0Ejd_ts
@@ -86,6 +86,9 @@ The corresponding documentation is designed to provide some insight into how eac
 	[PARAMLITERAL:title]
 	Features Tour 2014
 	[/PARAMLITERAL]
+	[PARAMLITERAL:playlistid]
+	PLZlv_N0_O1ga0aV9jVqJgog0VWz1cLL5f
+	[/PARAMLITERAL]
 	[PARAMLITERAL:videoid]
 	FcxA_xYnHZ8
 	[/PARAMLITERAL]
diff --git a/Engine/Documentation/Source/Resources/Showcases/TwinblastBust/TwinblastBust.INT.udn b/Engine/Documentation/Source/Resources/Showcases/TwinblastBust/TwinblastBust.INT.udn
--- a/Engine/Documentation/Source/Resources/Showcases/TwinblastBust/TwinblastBust.INT.udn
+++ b/Engine/Documentation/Source/Resources/Showcases/TwinblastBust/TwinblastBust.INT.udn
@@ -1,7 +1,7 @@
 Availability: Docs
-Title:Twinblast Bust
+Title:Photoreal Character Bust
 Crumbs:%ROOT%, Resources
-Description:An overview of the Twinblast Bust showcase.
+Description:An overview of the Photoreal Character Bust showcase.
 Related: Engine/Rendering/Materials
 Related: Engine/Rendering/Materials/PhysicallyBased
 Related: Engine/Rendering/Materials/LayeredMaterials
@@ -9,10 +9,10 @@ Version: 4.13
 
 [TOC(start:2 end:3)]
 
-![](TwinblastHeader.png)(w:700 convert:false)
 
+![](PhotorealHead.png)(w:700 convert:false)
 
-The purpose of this showcase is to demonstrate the use of high-quality character shader techniques, as demonstrated on a character from Epic's MOBA game Paragon. To view this showcase, simply open the project and press Play In Editor to see the cinematic pan-around. 
+The purpose of this showcase is to demonstrate the use of high-quality character shader techniques, similarly to those used on characters from Epic's MOBA game Paragon. To view this showcase, simply open the project and press Play In Editor to see the cinematic pan-around. 
 
 For more information on the tech used to produce this character, please refer to this [Unreal Engine Livestream - Tech & Techniques Behind Creating the Characters for Paragon](https://www.youtube.com/watch?v=toLJh5nnJA8).
 
@@ -20,21 +20,22 @@ For more information on the tech used to produce this character, please refer to
 
 The character's skin is created through the use of UE4's [Subsurface Profile shading model](Engine/Rendering/Materials/LightingModels/SubSurfaceProfile/). 
 
-![](TwinblastSkinMat.png)(w:700 convert:false)
+
+![](Skin.png)(w:700 convert:false)
 
 Note the use of [Material Functions](Engine/Rendering/Materials/Functions/) to set up the base for the skin Material. This is done as a reusable authoring approach to creating Materials for Paragon. In this way, artists can standardize approaches to generating certain types of surfaces, knowing that a fundamental change to a single Function will update all instances where that Function is being used. 
 
 ### Skin Shader Textures 
 
-The textures used on Twinblast's skin are all at 4K resolution and were originally produced from facial scans of the actor. The textures were then cleaned up and tweaked by artists at Epic. This skin setup utilizes five total texture maps: diffuse, roughness, specularity, scatter, and normal.
+The textures used on the character's skin are all at 4K resolution and were originally produced from facial scans of the actor. The textures were then cleaned up and tweaked by artists at Epic. This skin setup utilizes five total texture maps: diffuse, roughness, specularity, scatter, and normal.
 
 | Texture | Name | Description |
 |---------|------|-------------|
-|![](SkinDiffuse.png)(w:200 convert:false)|**Diffuse**|The diffuse map supplies the Base Color for the Material. At 4K, you can see the tiny capillaries just underneath the surface. Any darkening of wrinkles will accentuate the texture supplied by the normal map. (*TwinBlast_UE4_Demo_Head_D*)|
-|![](SkinRough.png)(w:200 convert:false)|**Roughness**|The roughness map is stored within the alpha channel of the diffuse texture. This is a common technique used to minimize the amount of textures in use. Note that the roughness increases wihtin pores and wrinkles. This causes those areas to appear less shiny, accentuating the look of depth provided by the diffuse and normal map. Also note that the hair texture on the head is pushed out to fully rough (1.0) this prevents any stray specular highlights from the scalp, which will give a much more significant sense of depth to the hair. (*TwinBlast_UE4_Demo_Head_D*)|
+|![](SkinDiffuse.png)(w:200 convert:false)|**Diffuse**|The diffuse map supplies the Base Color for the Material. At 4K, you can see the tiny capillaries just underneath the surface. Any darkening of wrinkles will accentuate the texture supplied by the normal map. (*UE4_Demo_Head_D*)|
+|![](SkinRough.png)(w:200 convert:false)|**Roughness**|The roughness map is stored within the alpha channel of the diffuse texture. This is a common technique used to minimize the amount of textures in use. Note that the roughness increases within pores and wrinkles. This causes those areas to appear less shiny, accentuating the look of depth provided by the diffuse and normal map. Also note that the hair texture on the head is pushed out to fully rough (1.0) this prevents any stray specular highlights from the scalp, which will give a much more significant sense of depth to the hair. (*UE4_Demo_Head_D*)|
 |![](SkinSpec.png)(w:200 convert:false)|**Specular**|The specular map scales the amount of specular highlight visible across the surface of the skin. It is important to note that the default value for specularity is 0.5. This map boosts specularity at areas where the skin may be stretched a bit tighter, and dampens it in areas where we don't want to see reflection, such as the center of pores and within wrinkles. (*CH2_exp04_merged_spec_f_FC_FINAL*)|
-|![](SkinScatter.png)(w:200 convert:false)|**Scatter**|The scatter map controls how much light is going to be scattered through the surface of the skin. Areas that are dark will exhibit very little scatter, such as the cheeks, while lighter areas will have higher amounts of perceived scatter, such as the nose and ears. The color of the scatter is managed by the [Subsurface Profile asset.](Engine/Rendering/Materials/LightingModels/SubSurfaceProfile/) (*TwinBlast_UE4_Demo_Head_BackScatter*)|
-|![](SkinNormal.png)(w:200 convert:false)|**Normal**|The normal map works as expected in typical material setups, supplying the tactile texture to the surface by perturbing pixel normals. There is nothing out of the ordinary in its setup for this example. (*TwinBlast_UE4_Demo_Head__normals*)|
+|![](SkinScatter.png)(w:200 convert:false)|**Scatter**|The scatter map controls how much light is going to be scattered through the surface of the skin. Areas that are dark will exhibit very little scatter, such as the cheeks, while lighter areas will have higher amounts of perceived scatter, such as the nose and ears. The color of the scatter is managed by the [Subsurface Profile asset.](Engine/Rendering/Materials/LightingModels/SubSurfaceProfile/) (*UE4_Demo_Head_BackScatter*)|
+|![](SkinNormal.png)(w:200 convert:false)|**Normal**|The normal map works as expected in typical material setups, supplying the tactile texture to the surface by perturbing pixel normals. There is nothing out of the ordinary in its setup for this example. (*UE4_Demo_Head__normals*)|
 
 
 ## Hair Shading
@@ -86,11 +87,12 @@ The Hair Shader uses the Masked blend mode instead of translucency. However, Mas
 
 ![](EdgeMaskGraph.png)(w:700 convert:false)
 
-Though not specifically part of the shader itself, it is worth noting that the Material used on Twinblast's hair - as well as the other characters of Paragon - uses what is called an "Edge Mask" to fade the hair away as it becomes edge-on to the camera. Since hair is rendered as a series of planes - see [Hair and Geometry](#hairgeometry) below - it is likely that from certain camera angles, one would see the edge of each sheet, breaking the hair effect.
+Though not specifically part of the shader itself, it is worth noting that the Material used on the example character's hair - as well as the other characters of Epic's game *Paragon* - uses what is called an "Edge Mask" to fade the hair away as it becomes edge-on to the camera. Since hair is rendered as a series of planes - see [Hair and Geometry](#hairgeometry) below - it is likely that from certain camera angles, one would see the edge of each sheet, breaking the hair effect.
+
+To help alleviate this problem, the Material calculates the camera's vector against the vertex normal such that as the surface turns perpendicularly away from the camera, it will also fade away to complete transparency. The other side of this approach, however, is that in doing so it is possible that the hair shader may show more of the scalp. This is also why many characters with thick heads of hair will often have a hair texture painted on the scalp, as seen in the image below.
 
-To help alleviate this problem, the Material calculates the camera's vector against the vertex normal such that as the surface turns perpendicularly away from the camera, it will also fade away to complete transparency. This is also why many characters with thick heads of hair will often have a hair texture painted on the scalp, as seen in the image below.
+![](HairPaint.png)(w:700 convert:false)
 
-![](TwinblastScalp.png)(w:700 convert:false)
 
 [REGION:tip]It should also be pointed out that this effect is tied to a Quality Switch node, so that on lower-powered platforms it can easily be switched off to trade back some shader performance. As always, it's best to test on target platform.[/REGION]
 
@@ -104,9 +106,12 @@ Creating hair using the techniques shown in this showcase requires a bit of unde
 
 #### Hair Geometry
 
-Hair geometry using UE4's hair shader is generally going to be constructed using a series of planar surfaces, which is a common approach in many real-time hair solutions. These can be authored in the DCC app of your choice.
+Hair geometry using UE4's hair shader is generally going to be constructed using a series of non-planar sheet surfaces, which is a common approach in many real-time hair solutions. These can be authored in the DCC app of your choice.
+
+[REGION:note]While there are no hard and fast rules for the detail of your hair geometry, it should be noted that this character uses approximately 800 individual hair sheets, with a total of approximately 18,000 triangles. Also note that the hair Material is set to be two-sided in the main Material node's properties.[/REGION]
+
+![](HairSheets.png)(w:700 convert:false)
 
-![](HairGeo.png)(w:700 convert:false)
 
 #### Hair Textures
 
@@ -115,11 +120,11 @@ In this use of the UE4 Hair shader, the final result is driven by five primary t
 
 | Texture | Name | Description |
 |---------|------|-------------|
-|![](Diffuse.png)(w:200 convert:false)|**Diffuse**|The Diffuse texture is precisely what one might expect in that it supplies the primary diffuse - or *base* - color of the hair itself. It is sometimes desirable to leave this texture uncolored and drive its color via parameters, especially in those cases where a character's hair may take on many different colors. (*TwinBlast_UE4_Demo_Hair_D*)|
-|![](Alpha.png)(w:200 convert:false)|**Alpha**|The Alpha texture supplies the areas of opacity for the hair, hiding geometry outside the body of strands. (*TwinBlast_UE4_Demo_Hair_A*)|
-|![](Root.png)(w:200 convert:false)|**Root**|The Root texture provides a mask which is used to change the color of the hair from the root to the tip, useful for lightening or darkening hair along its length. (*TwinBlast_UE4_Demo_Hair_Roots*)|
-|![](Depth.png)(w:200 convert:false)|**Depth**|The Depth texture is used by Pixel Depth offset to give the illusion of pushing the hair deeper into the hair volume. It can also be used as a basis to change color or shading values for hair at different depths, such as reducing overall specularity as hair falls further toward the scalp. (*TwinBlast_UE4_Demo_Hair_Depth*)|
-|![](ID.png)(w:200 convert:false)|**Unique ID**| The Unique ID texture simply provides a unique 0 to 1 (black to white) value for each strand of hair on a given piece of hair geometry. This is used to provide subtle variation within the hair itself. (*TwinBlast_UE4_Demo_Hair_ID*)|
+|![](Diffuse.png)(w:200 convert:false)|**Diffuse**|The Diffuse texture is precisely what one might expect in that it supplies the primary diffuse - or *base* - color of the hair itself. It is sometimes desirable to leave this texture uncolored and drive its color via parameters, especially in those cases where a character's hair may take on many different colors. (*UE4_Demo_Hair_D*)|
+|![](Alpha.png)(w:200 convert:false)|**Alpha**|The Alpha texture supplies the areas of opacity for the hair, hiding geometry outside the body of strands. (*UE4_Demo_Hair_A*)|
+|![](Root.png)(w:200 convert:false)|**Root**|The Root texture provides a mask which is used to change the color of the hair from the root to the tip, useful for lightening or darkening hair along its length. (*UE4_Demo_Hair_Roots*)|
+|![](Depth.png)(w:200 convert:false)|**Depth**|The Depth texture is used by Pixel Depth offset to give the illusion of pushing the hair deeper into the hair volume. It can also be used as a basis to change color or shading values for hair at different depths, such as reducing overall specularity as hair falls further toward the scalp. (*UE4_Demo_Hair_Depth*)|
+|![](ID.png)(w:200 convert:false)|**Unique ID**| The Unique ID texture simply provides a unique 0 to 1 (black to white) value for each strand of hair on a given piece of hair geometry. This is used to provide subtle variation within the hair itself. (*UE4_Demo_Hair_ID*)|
 
 
 
@@ -135,33 +140,40 @@ When using the Hair Shader you will notice some new properties available on the
 
 One of the primary reasons we refer to the Hair shader as an *approximation*, rather than a *simulation*, is that it is not actually simulating every individual strand of hair, nor is it simulating perfectly accurate lighting behavior. In the real world, as light bounces off of and transmits through follicles of hair, it will often encounter *other* strands of hair, repeating the same process potentially many times in succession. At this time there is no computing system capable of accurately producing the effects of such a simulation in realtime.
 
-Nevertheless, the way in which light scatters through a body of hair is still crucial to how realistic that hair appears to be, in a game just as much as in the real world. To control this, the Hair shader provides the Scatter property, which replaces Metallic on your primary shader node and is limited to values between 0.0 and 0.1. Scatter controls how much light can pass into and scatter through the hair *as a whole surface*, rather than per follicle. In simpler terms, this controls how much light passes through the entire body of your character's hair. 
+Nevertheless, the way in which light scatters through a body of hair is still crucial to how realistic that hair appears to be, in a game just as much as in the real world. To control this, the Hair shader provides the Scatter property, which replaces Metallic on your primary shader node and is limited to values between 0.0 and 1.0. Scatter controls how much light passes through the entire body of your character's hair, as if it were a single surface. 
 
-An important point to note about Scatter is that it will tend to be higher for lighter colored hair and darker for darker hair. This follows the physcal rules of the natural world, as darker hair will tend to absorb more light. In practical terms, if you are trying to create a blonde character, you will find that changing the diffuse texture or color is not enough; you will also have to increase the Scatter value as well.
+An important point to note about Scatter is that it will tend to be higher for lighter colored hair and darker for darker hair. This follows the physical rules of the natural world, as darker hair will tend to absorb more light. In practical terms, if you are trying to create a blonde character, you will find that changing the diffuse texture or color is not enough; you will also have to increase the Scatter value as well.
 
 
 [OBJECT:ComparisonSlider]
 [PARAM:before]
-![Scatter value of 0.0](LowScatter.png)(w:700 convert:false)
+![Scatter value of 0.0](Scatter_0.png)(w:700 convert:false)
 [/PARAM]
 [PARAM:after]
-![Scatter value of 1.0](HighScatter.png)(w:700 convert:false)
+![Scatter value of 1.0](Scatter_1.png)(w:700 convert:false)
 [/PARAM]
 [/OBJECT]
-[REGION:caption]In this example the Diffuse texture has not been modified from its default dark brown color. This illustrates how various tones of hair can be generated just by changing the Scatter value.[/REGION]
+[REGION:caption]In this example the Root and Tip colors were set to a pale blonde color, and Random Variation was set to 0.0. Scatter was then used to adjust how much light was approximated to move through the hair. This illustrates how various tones of hair can be generated just by changing the Scatter value.[/REGION]
 
 #### Tangent
 
-The Tangent property replaces the Normal property on the Hair shader. Instead of a typical normal map, the hair shader calculates a tangent. The tangent is calculated as a vector that runs parallel to each strand of hair, pointed back toward the root. In Twinblast's hair shader, this vector is given a random offset in the Z-Axis between 0.3 and -0.3, using the Unique ID texture. This produces a vector which has a random direction within an arc, and aids in providing variation within the anisotropic specular highlights.
+The Tangent property replaces the Normal property on the Hair shader. The tangent is calculated as a vector that runs parallel to each strand of hair, pointed back toward the root. The purpose of the Tangent property is to aid in proper anisotropic specularity. If you are unfamiliar, anisotropic specularity is what happens when light bounces off a micro-grooved surface, such as brushed metal. 
+
+![](AnisoVsNonAniso.png)(w:700 convert:false)
+[REGION:caption]The sphere on the left is anisotropically shaded, while the one on the right is not. Notice how the anisotropic specular highlight stretches along the surface.[/REGION]
+
+The Tangent property is used primarily as a vector to control the direction of that the anisotropic highlight will stretch.
 
 ![](HairTangent.png)(w:700 convert:false)
 [REGION:caption]In this image the yellow line represents the tangent along a strand of hair, pointing back toward the root.[/REGION]
 
-Tangents can be mapped in one of two ways: automatically or with the help of a flow map. The automatic method simply requires that that the textures for each sheet of hair be oriented such that the root is toward the top and the tips are toward the bottom. So long as your hair is short, and no sheets are bending around and twisting too much, this approach is adequate. This is the method used on Twinblast's hair.
+[REGION:tip]In the example character's hair shader, this vector is given a random offset in the Z-Axis between 0.3 and -0.3, using the Unique ID texture. This produces a vector which has a random direction within an arc, and aids in providing variation within the anisotropic specular highlights such as you would see in clusters of real hair.[/REGION]
+
+Tangents can be mapped in one of two ways: automatically or with the help of a flow map. The automatic method simply requires that that the textures for each sheet of hair be oriented such that the root is toward the top and the tips are toward the bottom. So long as your hair is short, and no sheets are bending around and twisting too much, this approach is adequate. This is the method used on the example character's hair.
 
 ##### Flow Maps for Tangency
 
-The other method requires the creation of a flow map. This is useful if your character's hair is long and is bending and needs to appear to bend and curl more than the actual geometry does, or if the individual texture portions for the hair are oriented in different ways (not top-to-bottom as described above). The flow map will represent the direction the hair is moving in tangent space, or along the surface. Within the Twinblast Bust project, you can find an unused flow map named *T_SparrowHair_Flow*. Below is a comparison of that flowmap and the result of the specularity. 
+The other method requires the creation of a flow map. This is useful if your character's hair is long and is bending and needs to appear to bend and curl more than the actual geometry does, or if the individual texture portions for the hair are oriented in different ways (not top-to-bottom as described above). The flow map will represent the direction the hair is moving in tangent space, or along the surface. Within the the Photoreal Character Bust project, you can find an unused flow map named *T_Hair_Flow*. Below is a comparison of that flowmap and the result of the specularity. 
 
 [OBJECT:ComparisonSlider]
 [PARAM:before]
@@ -194,15 +206,17 @@ Pixel Depth Offset (PDO) is not a property unique to the Hair shader. In layman'
 
 ## Eye Shading
 
-![](EyeShading.png)(w:700 convert:false)
+![](EyeShading2.png)(w:700 convert:false)
 
 Eyes are often referred to as "the windows to the soul." Few aspects of a character have as much influence on the viewer's perception of a character as a real person as that character's eyes. The Eye shader in UE4 is designed to bring realistic rendering of an eye surface, exposing artistic control over each of the biological parts of the eye itself.
 
-[REGION:warning]The Eye shader in UE4 is highly technical and has been developed such that there are **very strong dependencies** between the shader code, the Material, the shape of the geometry, and the geometry's UV layout. Building an eye from scratch using this shader is not recommended without extensive experience. To that end, we strongly recommend, for those creating realistic humanoid eyes, that you extract the eye geometry from this example and use the eye Material **as-is**, replacing the necessary textures to suit your purpose. [/REGION]
+[REGION:warning]The Eye shader in UE4 is highly technical and has been developed such that there are **very strong dependencies** between the shader code, the Material, the shape of the geometry, and the geometry's UV layout. Building an eye from scratch using this shader is **not** recommended without extensive experience in shader development. To that end, we strongly recommend, for those creating realistic humanoid eyes, that you extract the eye geometry from this example and use the eye Material **as-is**, replacing the necessary textures to suit your purpose. [/REGION]
+
+[REGION:tip]Located within the Photoreal Character project is a Static Mesh named **PhotorealEyeGeo**. This can be exported from the project into the DCC app of your choice via FBX to be used in your own projects.[/REGION]
 
 ### Eye Biology
 
-In order to fully understand the features of the eye, it is best to know a little bit of eye biology. Each of the parts you see here are represented in this eye setup, and all can be modified in some way using the exposed parameters in the Material Instance applied to Twinblast's eyes (*MI_TwinBlast_EyeRefractive_Bust*).
+In order to fully understand the features of the eye, it is best to know a little bit of eye biology. Each of the parts you see here are represented in this eye setup, and all can be modified in some way using the exposed parameters in the Material Instance applied to the example character's eyes (*MI_EyeRefractive_Bust*).
 
 Here is a quick refresher:
 
@@ -215,29 +229,91 @@ Here is a quick refresher:
 |4|Pupil|The pupil is the dark spot at the center of the eye. It is the hole through which light passes so it can be picked up by the rods and cones of the retina.|
 |5|Cornea|The cornea is the clear, fluid-filled dome that rests over the surface of the iris. |
 
-### Refraction within the Eye
+### Refraction within the Eye - Then and Now
 
 Because the eyeball is filled with fluid, it is given to refracting any light that passes through it. This can be seen in the real world when looking at an eye from multiple angles. The iris and pupil will be distorted by refraction as they are viewed through the cornea. Traditional approaches to solving this for games and film was to create two separate eye surfaces: one that supplied the sclera, iris, and pupil; and another surface on top that supplied the cornea and overall wetness of the eye. This allowed for refraction as the underlying surface was viewed through the wet layer shell. Such an approach was used in Epic's *A Boy and His Kite* tech demo on the boy's eyes. The diagram below shows a representation of the two surfaces.
 
 ![](KiteBoyEyes.png)(w:700 convert:false)
 
-Using our new Eye shader, however, an even more realistic effect can be achieved with a *single surface*. The refraction of light through the surface of the cornea is now handled entirely within the shader, alleviating the need for underlying geometry. You can experiment with this yourself by opening the *MI_TwinBlast_EyeRefractive_Bust* Material Instance and adjusting the **Depth Scale** property.
+Using our new Eye shader, however, an even more realistic effect can be achieved with a *single surface*. The refraction of light through the surface of the cornea is now handled entirely within the shader, alleviating the need for underlying geometry. You can experiment with this yourself by opening the *MI_EyeRefractive_Bust* Material Instance and adjusting the **Refraction On/Off** property.
 
 [OBJECT:ComparisonSlider]
 [PARAM:before]
-![Without Refraction (Depth Scale at 0.0)](EyeNotRefracting.png)(w:700 convert:false)
+![Without Refraction (Refraction On/Off at 0.0)](NotRefracting.png)(w:700 convert:false)
 [/PARAM]
 [PARAM:after]
-![With Refraction (Depth Scale at 1.2)](EyeRefracting.png)(w:700 convert:false)
+![With Refraction (Refraction On/Off at 1.0)](Refracting.png)(w:700 convert:false)
 [/PARAM]
 [/OBJECT]
 
+### Eye Shader Parameters
+
+Due to the complexity of the eye shader and our recommendation that you use it as-is, we have included a breakdown of the artistic parameters that are built into the Material. The following are the eye shader parameters that can be tweaked from within *MI_EyeRefractive_Bust* Material Instance Constant. Wherever possible, we will also include suggested parameter ranges.
+
+![](MI_EyeRefractive_Bust.png)
+
+|Name|Description|Suggested Value Range|
+|----|-----------|-------------------|
+|[REGION:tablesection]Static Switch Parameter Values[/REGION]|||
+|**UseEyeBulge**|This property is intended to adjust the refraction properties of the shader based on whether or not the eyeball geometry uses a physically modeled bulge for the cornea.|n/a|
+|[REGION:tablesection]Vector Parameter Values[/REGION]|||
+|**Sclera inner color**| This is the color of the sclera where it meets the iris.|n/a|
+|**Sclera outer color**| This is the color of the sclera at the outer edge of the eye.| n/a|
+|**SecondaryEnvBalance**|This controls how much the [Secondary Environment](#eyeshadersecondaryenvironment) is visible on the surface of the eye.|0.00 - 0.03|
+|**SecondaryEnvRotationAxis**|Since the [Secondary Environment](#eyeshadersecondaryenvironment) is a faked reflection, it must be manually rotated around the eye. This property controls the axis about which it will rotate.| <0.0, 0.0, 1.0> (Z-Axis)|
+|[REGION:tablesection]Scalar Parameter Values[/REGION]|||
+|**Depth Scale**|This controls the depth of the refraction of the iris underneath the cornea. |1.0-1.4|
+|**Flatten Normal**|This value controls how much flattening of the eye's normal map is taking place, focused particularly on the sclera.|0.9 - 1.1|
+|**IoR**|Index of refraction of the fluids underneath the cornea. Controls how much refraction takes place.|1.336|
+|**Iris Concavity Power**| Used alongside *Iris Concavity Scale* to control the shape and amount of light caustics that are calculated on the surface of the iris as light passes through the cornea. This will generally only be visible in an actual lit scene, and can be difficult to visualize within the Material Instance Editor.|0.2 - 0.3|
+|**Iris Concavity Scale**| Used alongside *Iris Concavity Power* to control the shape and amount of light caustics that are calculated on the surface of the iris as light passes through the cornea. This will generally only be visible in an actual lit scene, and can be difficult to visualize within the Material Instance Editor.|0.07 - 0.15|
+|**Iris UV Radius**| Controls the overall size of the iris on the eyeball.|0.14 - 0.15|
+|**Iris Brightness**| Controls the brightness of the iris. | *Depends on iris color* |
+|**Iris Roughness**| This value drives how shininess of the cornea, the surface directly over the iris. | 0.0 - 0.1 |
+|**Limbus Dark Scale**| Controls the size of the darkening ring of the limbus. | 2.0 - 2.15 |
+|**Limbus Pow***| Controls the overall darkening of the limbus area. Overdriven values will darken the entire corneal area. | 15.0 - 22.0 |
+|**Limbus UV Width Color| Controls the sample size for the limbus, or how much of the eye surface will be alloted for showing the limbus. This is a fairly delicate setting and should be adjusted with care. | 0.025 - 0.040 |
+|**Limbus UV Width Shading**| Controls how much light will affect the shading of the limbus. Very delicate setting, adjust with care. | 0.040 - 0.050 |
+|**Normal UV Scale**| Controls the scale of the normal map texture used across the surface of the eye. | 1.8 - 2.6 |
+|**Pupil Scale**| Controls the size of the pupil. This is the setting you would use to dilate the eyes. | 0.5 - 1.0 |
+|**Refraction On/Off**| Blends between refracting and non-refracting versions of the shader. | 1.0 |
+|**Scale By Center**| Adjusts the scale of the entire iris/pupil area from its center. | 0.8 - 1.2 |
+|**Sclera Brightness**| Controls the brightness of the sclera, or white of the eye. | 0.9 - 1.3 |
+|**Sclera Roughness**| Controls the Material Roughness value of the sclera| 0.0 - 0.1 |
+|**Secondary Env Rotation**| This rotates the [secondary environment cubemap](#eyeshadersecondaryenvironment) about the axis defined using the SecondaryEnvRotationAxis property. | n/a |
+|**Shadow Hardness**| Controls the sharpness of the blend between the sclera inner and outer color. Used with Shadow Radius to help drive shading across the surface of the sclera, approximating the effect of the eyelid casting subsurface scattered shadows across the surface of the eyeball. | 0.1 - 0.2 |
+|**Shadow Radius**| Controls the size of the blend between the inner and outer colors of the sclera. Used with Shadow Hardness. | 0.5 - 0.85 |
+|**Specularity Iris**| Controls specularity level across the cornea (iris and pupil). | 0.25 - 0.8 |
+|**Specularity Sclera**| Controls the specularity level across the sclera. | 0.25 - 0.8 |
+|**Veins**| Controls the visibility of veins in the eye, as provided by the Sclera Color Map. | *Depends on Sclera Color Map and desired vein amount.* |
+|[REGION:tablesection]Texture Parameter Values[/REGION]|||
+|**Iris Color**| See [Eye Textures](#eyetextures) below. | n/a |
+|**Normal Map**| See [Eye Textures](#eyetextures) below. | n/a |
+|**Sclera Color**| See [Eye Textures](#eyetextures) below. | n/a |
+|**Secondary Env**| This is a cubemap used as a secondary reflection environment. See [Eye Shader Secondary Environment](#eyeshadersecondaryenvironment) below. | n/a |
+
+
+
+#### Eye Shader Secondary Environment
+
+The Secondary Environment is a part of the eye shader that provides a faked reflection map to the surface of the eye. You can provide a texture for it using the SecondaryEnv ParamCube in the shader and then controlled with the other SecondaryEnv properties. This can be useful for dramatic closeups of the eye when you want to be able to see specific reflections, such as the immediate environment, another character, etc.
+
+[OBJECT:ComparisonSlider]
+[PARAM:before]
+![Without Secondary Environment (SecondaryEnvBalance at 0.0)](SecondaryEnvOFF.png)(w:700 convert:false)
+[/PARAM]
+[PARAM:after]
+![With Secondary Environment (SecondaryEnvBalance at 0.03)](SecondaryEnvON.png)(w:700 convert:false)
+[/PARAM]
+[/OBJECT]
 
 ### Eye Authoring
 
 Earlier in this section we provided a warning that there are strong interdependencies between the code of the Eye shader, the setup of the material, and the geometry of the eye mesh. The following is a high-level overview of what makes the eye setup so particular, which will be useful if you should try replacing the geometry with your own eye mesh for your characters.
 
-When building your own eyes, we **strongly** recommend that you start using the eye geometry and Material provided in this example. As such, most of the information provided here will be based upon that assumption.
+When building your own eyes, we **strongly** recommend that you start using the actual eye geometry and Material provided in this example. This would involve As such, most of the information provided here will be based upon that assumption. 
+
+[REGION:tip]Located within the Photoreal Character project is a Static Mesh named **PhotorealEyeGeo**. This can be exported from the project into the DCC app of your choice via FBX to be used in your own projects.[/REGION]
 
 #### Shape and UV Layout
 
@@ -246,7 +322,7 @@ Much like the human eye of the real world, the eye mesh is not modeled to be per
 ![](EyeGeometry.png)(w:400 convert:false)
 [REGION:caption]In this image you can see the geometry of the eye viewed from the side in Maya. Note the raised area of the cornea at the front.[/REGION]
 
-The UV layout for the eye is also critical, as it is from this that the proportions of each eye texture are produced. Fortunately, the UV layout that works best with the Eye shader is relatively basic: a planar projection along the front axis of the eye. Inclusion of backfaces is not extremely important (since they'll be behind the eye anyway), which can be seen on the Twinblast example.
+The UV layout for the eye is also critical, as it is from this that the proportions of each eye texture are produced. Fortunately, the UV layout that works best with the Eye shader is relatively basic: a planar projection along the front axis of the eye. Inclusion of backfaces is not extremely important (since they'll be behind the eye anyway), which can be seen on the the example character.
  
  The UV layout will look like this:
 
@@ -260,11 +336,11 @@ Many of the textures you'll create for your eyes will be based on the UV layout
 
 | Texture | Name | Description |
 |---------|------|-------------|
-|![](ScleraMap.png)(w:200 convert:false) |**Sclera Map**|The sclera map controls the colors of the eye whites, as well as the vein and tissue color that will be present at the edges of the eye. If you wanted your character to be able to have bloodshot eyes, for example, this is the texture you would modify. (*TwinBlast_UE4_Demo_Eye_Sclera_Color*)|
+|![](ScleraMap.png)(w:200 convert:false) |**Sclera Map**|The sclera map controls the colors of the eye whites, as well as the vein and tissue color that will be present at the edges of the eye. If you wanted your character to be able to have bloodshot eyes, for example, this is the texture you would modify. (*UE4_Demo_Eye_Sclera_Color*)|
 |![](DispMap.png)(w:200 convert:false) |**Mid Plane Displacement Map**|The mid plane displacement map is used to lock down a plane that cuts through the center of the eye, and is used as a basis for offsetting the depth of the iris. (*EyeMidPlaneDisplacement_Example*)|
-|![](WetNormal.png)(w:200 convert:false) |**Normal Map**|The normal map does the same job as normal maps in all shaders - it supplies tactile texture variation. In the case of the Twinblast eye setup, it provides tiny lumps and bumps along the wet surface of the eye. (*T_Eye_Wet_Normal*)|
-|![](Normal.png)(w:200 convert:false) |**Tangent Map**|The tangent map is used to control the direction the surface is flowing in tangent space. In the case of Twinblast's eye setup, this is used to accentuate the change in direction where the cornea meets the sclera. (*EYE_NORMALS.EYE_NORMALS*)|
-|![](IrisTexture.png)(w:200 convert:false)|**Iris Map**|The iris texture is unique in that it does not match the UV layout. Instead, the texure fills the entirety of the resolution. UVs and alpha masks are employed to control the overall iris size, as well as the size of the pupil. If making your own iris maps, try to keep the pupil size reasonably close to the proportions in the supplied texture. (*TwinBlast_UE4_Demo_Eye_Iris*)|
+|![](WetNormal.png)(w:200 convert:false) |**Normal Map**|The normal map does the same job as normal maps in all shaders - it supplies tactile texture variation. In the case of the the example character eye setup, it produces tiny lumps and bumps along the wet surface of the eye. (*T_Eye_Wet_Normal*)|
+|![](Normal.png)(w:200 convert:false) |**Tangent Map**|The tangent map is used to control the direction the surface is flowing in tangent space. In the case of the example character's eye setup, this is used to accentuate the change in direction where the cornea meets the sclera. (*EYE_NORMALS.EYE_NORMALS*)|
+|![](IrisTexture.png)(w:200 convert:false)|**Iris Map**|The iris texture is unique in that it does not match the UV layout. Instead, the texture fills the entirety of the resolution. UVs and alpha masks are employed to control the overall iris size, as well as the size of the pupil. If making your own iris maps, try to keep the pupil size reasonably close to the proportions in the supplied texture. (*UE4_Demo_Eye_Iris*)|
 
 [REGION:warning]Note that the Tangent Map output is **not** fed directly into the Normal Map input. Rather, it is connected to a **Tangent Output** node, which must be created separately, as seen in the master material **M_HairSheet_Master2**. [/REGION]
 ![](TangentOutput.png)
@@ -273,9 +349,11 @@ Many of the textures you'll create for your eyes will be based on the UV layout
 
 ### Eye Ambient Occlusion
 
-The final part of the eye setup for Twinblast is the use of a translucent helper surface to to simulate ambient occlusion around the eye. In the real world - wherever your eye comes in contact with other tissues such as the eyelid or the tear duct - you will see very soft shadowing. The softness is due to the fact that the eye itself is translucent; light will scatter through it. To help simulate this realistically, we have placed a sheet of geometry that covers the eye and provides approximated ambient occlusion, creating a richer degree of realism across the eye.
+The final part of the eye setup for the example character is the use of a translucent helper surface to to simulate ambient occlusion around the eye. In the real world - wherever your eye comes in contact with other tissues such as the eyelid or the tear duct - you will see very soft shadowing. The softness is due to the fact that the eye itself is translucent; light will scatter through it. To help simulate this realistically, we have placed a sheet of geometry that covers the eye and provides approximated ambient occlusion, creating a richer degree of realism across the eye.
+
+The Material applied to this sheet is named *MI_Eye_Occlusion_Bust* and can be found in Element 3 of the the example character model.
 
-The Material applied to this sheet is named *MI_TwinBlast_Eye_Occlusion_Bust* and can be found in Element 3 of the Twinblast model.
+[REGION:note]If authoring your own content for character eyes, you will very likely need to build this piece of geometry yourself. Note that it should be made to fit the area of the open eyelid and that its UVs should cleanly fill the 0 to 1 UV space in your 3D modeling app. [/REGION]
 
 ![](AOSheetMaya.png)(w:700 convert:false) 
 [REGION:caption]This image shows the ambient occlusion sheet geometry in Maya, without the geometry of the face. The border edge of the sheet will perfectly fill the opening of the eyelid. [/REGION]
@@ -283,13 +361,15 @@ The Material applied to this sheet is named *MI_TwinBlast_Eye_Occlusion_Bust* an
 
 [OBJECT:ComparisonSlider]
 [PARAM:before]
-![Without AO Sheet](EyeAO_off.png)(w:700 convert:false)
+![Without AO Sheet](WithoutAO.png)(w:700 convert:false)
 [/PARAM]
 [PARAM:after]
-![With AO Sheet](EyeAO_on.png)(w:700 convert:false)
+![With AO Sheet](WithAO.png)(w:700 convert:false)
 [/PARAM]
 [/OBJECT]
 
+
+
 ## Special Thanks
 
 Special thanks to [[REDACTED_PERSON], [REDACTED_PERSON] and [REDACTED_PERSON]]([REDACTED_PROFILE_URL]) for releasing their work [*The importance of sampling for physically-based hair fiber models*](http://dl.acm.org/authorize?6967840), upon which some of the work of UE4's hair shader is based.
diff --git a/Engine/Documentation/Source/Support/Builds/Builds.INT.udn b/Engine/Documentation/Source/Support/Builds/Builds.INT.udn
--- a/Engine/Documentation/Source/Support/Builds/Builds.INT.udn
+++ b/Engine/Documentation/Source/Support/Builds/Builds.INT.udn
@@ -5,133 +5,9 @@ Description:Release notes for Unreal Engine builds.
 parent:%ROOT%
 order:9
 type:landing
+topic-icon:%ROOT%/outline_icon.png
+topic-image:%ROOT%/placeholder_topic.png
 
-[VAR:Params]
-[PARAM:image]
-	![%Support/Builds:title%](%ROOT%/placeholder_topic.png)
-[/PARAM]
-[PARAM:icon]
-	![](%ROOT%/outline_icon.png)(convert:false)
-[/PARAM]
-[PARAM:title]
-	%Support/Builds:title%
-[/PARAM]
-[PARAM:description]
-	%Support/Builds:description%
-[/PARAM]
-[PARAM:path]
-	[RELATIVE:Support/Builds]
-[/PARAM]
-[/VAR]
 
-[VAR:Topic]
-[OBJECT:Topic]
-	[PARAM:image]
-		![%Support/Builds:title%](%ROOT%/placeholder_topic.png)
-	[/PARAM]
-	[PARAM:icon]
-		![](%ROOT%/outline_icon.png)(convert:false)
-	[/PARAM]
-	[PARAM:title]
-		%Support/Builds:title%
-	[/PARAM]
-	[PARAM:description]
-		%Support/Builds:description%
-	[/PARAM]
-	[PARAM:path]
-		[RELATIVE:Support/Builds]
-	[/PARAM]
-[/OBJECT]
-[/VAR]
+[DIR(output:"listbutton" order:"descend" parent:"Support/Builds" org:"hierarchy")]
 
-[VAR:TopicCompact]
-[OBJECT:TopicCompact]
-	[PARAM:image]
-		![%Support/Builds:title%](%ROOT%/placeholder_topic.png)
-	[/PARAM]
-	[PARAM:icon]
-		![](%ROOT%/outline_icon.png)(convert:false)
-	[/PARAM]
-	[PARAM:title]
-		%Support/Builds:title%
-	[/PARAM]
-	[PARAM:description]
-		%Support/Builds:description%
-	[/PARAM]
-	[PARAM:path]
-		[RELATIVE:Support/Builds]
-	[/PARAM]
-[/OBJECT]
-[/VAR]
-
-[VAR:Role]
-[OBJECT:Role]
-%params%
-[PARAMLITERAL:id]
-engine
-[/PARAMLITERAL]
-[/OBJECT]
-[/VAR]
-
-[PUBLISH:Rocket]
-[DIR(end:"1" output:"listbutton" order:"descend" parent:"Support/Builds" org:"hierarchy")]
-[/PUBLISH:Rocket]
-
-[PUBLISH:Licensee]
-[OBJECT:TopicButtonList]
-	[PARAM:icon]
-	[/PARAM]
-	[PARAM:title]
-	2015
-	[/PARAM]
-	[PARAM:description]
-	[/PARAM]
-	[PARAM:links]
-	* [](Support/Builds/ReleaseNotes/2016/4_12)
-	* [](Support/Builds/ReleaseNotes/2016/4_11)
-	* [](Support/Builds/ReleaseNotes/2015/4_10)
-	* [](Support/Builds/ReleaseNotes/2015/4_9)
-	* [](Support/Builds/ReleaseNotes/2015/4_8)
-	* [](Support/Builds/ReleaseNotes/2015/4_7)
-	[/PARAM]
-[/OBJECT]
-
-[OBJECT:TopicButtonList]
-	[PARAM:icon]
-	[/PARAM]
-	[PARAM:title]
-	2014
-	[/PARAM]
-	[PARAM:description]
-	[/PARAM]
-	[PARAM:links]
-	* [](Support/Builds/ReleaseNotes/2014/4_6)
-	* [](Support/Builds/ReleaseNotes/2014/4_5)
-	* [](Support/Builds/ReleaseNotes/2014/4_4)
-	* [](Support/Builds/ReleaseNotes/2014/4_3)
-	* [](Support/Builds/ReleaseNotes/2014/4_2)
-	* [](Support/Builds/ReleaseNotes/2014/4_1)
-	* [](Support/Builds/ReleaseNotes/2014/March)
-	* [](Support/Builds/ReleaseNotes/2014/January)
-	[/PARAM]
-[/OBJECT]
-
-[OBJECT:TopicButtonList]
-	[PARAM:icon]
-	[/PARAM]
-	[PARAM:title]
-	2013
-	[/PARAM]
-	[PARAM:description]
-	[/PARAM]
-	[PARAM:links]
-	* [](Support/Builds/ReleaseNotes/2013/Dec)
-	* [](Support/Builds/ReleaseNotes/2013/Oct)
-	* [](Support/Builds/ReleaseNotes/2013/Sept)
-	* [](Support/Builds/ReleaseNotes/2013/August)
-	* [](Support/Builds/ReleaseNotes/2013/July)
-	* [](Support/Builds/ReleaseNotes/2013/June)
-	* [](Support/Builds/ReleaseNotes/2013/May)
-	[/PARAM]
-[/OBJECT]
-[/PUBLISH:Licensee]
diff --git a/Engine/Documentation/Source/Support/Builds/ReleaseNotes/2013/August/ReleaseNotesAugust2013.INT.udn b/Engine/Documentation/Source/Support/Builds/ReleaseNotes/2013/August/ReleaseNotesAugust2013.INT.udn
--- a/Engine/Documentation/Source/Support/Builds/ReleaseNotes/2013/August/ReleaseNotesAugust2013.INT.udn
+++ b/Engine/Documentation/Source/Support/Builds/ReleaseNotes/2013/August/ReleaseNotesAugust2013.INT.udn
@@ -2,7 +2,8 @@ Title: August 2013 Release Notes
 Crumbs:%ROOT%, Support/Builds
 Description:
 Availability:Licensee
-Template:ReleaseNotes.html
+Parent:Support/Builds
+Order:-6
 
 [TOC(start:2 end:2)]
 
diff --git a/Engine/Documentation/Source/Support/Builds/ReleaseNotes/2013/Dec/ReleaseNotesDecember2013.INT.udn b/Engine/Documentation/Source/Support/Builds/ReleaseNotes/2013/Dec/ReleaseNotesDecember2013.INT.udn
--- a/Engine/Documentation/Source/Support/Builds/ReleaseNotes/2013/Dec/ReleaseNotesDecember2013.INT.udn
+++ b/Engine/Documentation/Source/Support/Builds/ReleaseNotes/2013/Dec/ReleaseNotesDecember2013.INT.udn
@@ -2,7 +2,8 @@ Title: December 2013 Release Notes
 Crumbs:%ROOT%, Support/Builds
 Description:
 Availability:Licensee
-Template:ReleaseNotes.html
+Parent:Support/Builds
+Order:-3
 
 [TOC(start:2 end:2)]
 
diff --git a/Engine/Documentation/Source/Support/Builds/ReleaseNotes/2013/July/ReleaseNotesJuly2013.INT.udn b/Engine/Documentation/Source/Support/Builds/ReleaseNotes/2013/July/ReleaseNotesJuly2013.INT.udn
--- a/Engine/Documentation/Source/Support/Builds/ReleaseNotes/2013/July/ReleaseNotesJuly2013.INT.udn
+++ b/Engine/Documentation/Source/Support/Builds/ReleaseNotes/2013/July/ReleaseNotesJuly2013.INT.udn
@@ -2,7 +2,8 @@ Title: July 2013 Release Notes
 Crumbs:%ROOT%, Support/Builds
 Description:
 Availability:Licensee
-Template:ReleaseNotes.html
+Parent:Support/Builds
+Order:-7
 
 [TOC(start:2 end:2)]
 
diff --git a/Engine/Documentation/Source/Support/Builds/ReleaseNotes/2013/June/ReleaseNotesJune2013.INT.udn b/Engine/Documentation/Source/Support/Builds/ReleaseNotes/2013/June/ReleaseNotesJune2013.INT.udn
--- a/Engine/Documentation/Source/Support/Builds/ReleaseNotes/2013/June/ReleaseNotesJune2013.INT.udn
+++ b/Engine/Documentation/Source/Support/Builds/ReleaseNotes/2013/June/ReleaseNotesJune2013.INT.udn
@@ -2,7 +2,8 @@ Title: June 2013 Release Notes
 Crumbs:%ROOT%, Support/Builds
 Description:
 Availability:Licensee
-Template:ReleaseNotes.html
+Parent:Support/Builds
+Order:-8
 
 [TOC(start:2 end:2)]
 
diff --git a/Engine/Documentation/Source/Support/Builds/ReleaseNotes/2013/May/ReleaseNotesMay2013.INT.udn b/Engine/Documentation/Source/Support/Builds/ReleaseNotes/2013/May/ReleaseNotesMay2013.INT.udn
--- a/Engine/Documentation/Source/Support/Builds/ReleaseNotes/2013/May/ReleaseNotesMay2013.INT.udn
+++ b/Engine/Documentation/Source/Support/Builds/ReleaseNotes/2013/May/ReleaseNotesMay2013.INT.udn
@@ -2,7 +2,8 @@ Title: May 2013 Release Notes
 Availability:Licensee
 Crumbs:%ROOT%, Support/Builds
 Description:
-Template:ReleaseNotes.html
+Parent:Support/Builds
+Order:-9
 
 [TOC(start:2 end:2)]
 
diff --git a/Engine/Documentation/Source/Support/Builds/ReleaseNotes/2013/Oct/ReleaseNotesOctober2013.INT.udn b/Engine/Documentation/Source/Support/Builds/ReleaseNotes/2013/Oct/ReleaseNotesOctober2013.INT.udn
--- a/Engine/Documentation/Source/Support/Builds/ReleaseNotes/2013/Oct/ReleaseNotesOctober2013.INT.udn
+++ b/Engine/Documentation/Source/Support/Builds/ReleaseNotes/2013/Oct/ReleaseNotesOctober2013.INT.udn
@@ -2,6 +2,8 @@ Title: October 2013 Release Notes
 Crumbs:%ROOT%, Support/Builds
 Description:
 Availability:Licensee
+Parent:Support/Builds
+Order:-4
 
 [TOC(start:2 end:2)]
 
diff --git a/Engine/Documentation/Source/Support/Builds/ReleaseNotes/2013/Sept/ReleaseNotesSeptember2013.INT.udn b/Engine/Documentation/Source/Support/Builds/ReleaseNotes/2013/Sept/ReleaseNotesSeptember2013.INT.udn
--- a/Engine/Documentation/Source/Support/Builds/ReleaseNotes/2013/Sept/ReleaseNotesSeptember2013.INT.udn
+++ b/Engine/Documentation/Source/Support/Builds/ReleaseNotes/2013/Sept/ReleaseNotesSeptember2013.INT.udn
@@ -2,6 +2,8 @@ Title: September 2013 Release Notes
 Crumbs:%ROOT%, Support/Builds
 Description:
 Availability:Licensee
+Parent:Support/Builds
+Order:-5
 
 [TOC(start:2 end:2)]
 
diff --git a/Engine/Documentation/Source/Support/Builds/ReleaseNotes/2014/4_1/ReleaseNotes_4_1.INT.udn b/Engine/Documentation/Source/Support/Builds/ReleaseNotes/2014/4_1/ReleaseNotes_4_1.INT.udn
--- a/Engine/Documentation/Source/Support/Builds/ReleaseNotes/2014/4_1/ReleaseNotes_4_1.INT.udn
+++ b/Engine/Documentation/Source/Support/Builds/ReleaseNotes/2014/4_1/ReleaseNotes_4_1.INT.udn
@@ -3,6 +3,7 @@ Crumbs:%ROOT%, Support/Builds
 Description:
 Availability:Public
 version: 4.1
+parent:Support/Builds
 order:1
 
 [TOC(start:2 end:2)]
diff --git a/Engine/Documentation/Source/Support/Builds/ReleaseNotes/2014/January/ReleaseNotesJanuary2014.INT.udn b/Engine/Documentation/Source/Support/Builds/ReleaseNotes/2014/January/ReleaseNotesJanuary2014.INT.udn
--- a/Engine/Documentation/Source/Support/Builds/ReleaseNotes/2014/January/ReleaseNotesJanuary2014.INT.udn
+++ b/Engine/Documentation/Source/Support/Builds/ReleaseNotes/2014/January/ReleaseNotesJanuary2014.INT.udn
@@ -2,7 +2,8 @@ Title: January 2014 Release Notes
 Crumbs:%ROOT%, Support/Builds
 Description:
 Availability:Licensee
-Template:ReleaseNotes.html
+Parent:Support/Builds
+Order:-2
 
 [TOC(start:2 end:2)]
 
diff --git a/Engine/Documentation/Source/Support/Builds/ReleaseNotes/2014/March/ReleaseNotesMarch2014.INT.udn b/Engine/Documentation/Source/Support/Builds/ReleaseNotes/2014/March/ReleaseNotesMarch2014.INT.udn
--- a/Engine/Documentation/Source/Support/Builds/ReleaseNotes/2014/March/ReleaseNotesMarch2014.INT.udn
+++ b/Engine/Documentation/Source/Support/Builds/ReleaseNotes/2014/March/ReleaseNotesMarch2014.INT.udn
@@ -2,6 +2,8 @@ Title: March 2014 Release Notes
 Crumbs:%ROOT%, Support/Builds
 Description:
 Availability:Licensee
+Parent:Support/Builds
+Order:-1
 
 [TOC(start:2 end:2)]
 
diff --git a/Engine/Documentation/Source/Support/Builds/ReleaseNotes/2016/4_14/ReleaseNotes_4_14.INT.udn b/Engine/Documentation/Source/Support/Builds/ReleaseNotes/2016/4_14/ReleaseNotes_4_14.INT.udn
--- a/Engine/Documentation/Source/Support/Builds/ReleaseNotes/2016/4_14/ReleaseNotes_4_14.INT.udn
+++ b/Engine/Documentation/Source/Support/Builds/ReleaseNotes/2016/4_14/ReleaseNotes_4_14.INT.udn
@@ -2,6 +2,8 @@ Availability:Public
 Title: Unreal Engine 4.14 Release Notes
 Description: Release Notes for Unreal Engine 4.14
 Crumbs:%ROOT%, Engine
+Parent:Support/Builds
+Order:14
 reference-image:image_6.gif
 reference-image:image_7.gif
 reference-image:image_9.gif
diff --git a/Engine/Documentation/Source/Support/Builds/ReleaseNotes/4_0/Updates/UnrealEngine4_0UpdateNotes.INT.udn b/Engine/Documentation/Source/Support/Builds/ReleaseNotes/4_0/Updates/UnrealEngine4_0UpdateNotes.INT.udn
new file mode 100644
--- /dev/null
+++ b/Engine/Documentation/Source/Support/Builds/ReleaseNotes/4_0/Updates/UnrealEngine4_0UpdateNotes.INT.udn
@@ -0,0 +1,71 @@
+Availability:Public
+Title: Unreal Engine 4.0 Update Notes
+Description:
+Crumbs:
+Parent:Support/Builds
+Order:0
+
+## 4.0.0 Release
+
+Released: 03/19/2014
+
+Announcement:
+
+[https://www.unrealengine.com/blog/welcome-to-unreal-engine-4](https://www.unrealengine.com/blog/welcome-to-unreal-engine-4) 
+
+## 4.0.1 Hotfix
+
+Released: 03/19/2014
+
+Github: [https://github.com/EpicGames/UnrealEngine/tree/4.0.1-release](https://github.com/EpicGames/UnrealEngine/tree/4.0.1-release) 
+
+Fixes:
+
+* Fixed critical launch issues
+
+## 4.0.2 Hotfix
+
+Released: 04/01/2014
+
+Github: [https://github.com/EpicGames/UnrealEngine/tree/4.0.2-release](https://github.com/EpicGames/UnrealEngine/tree/4.0.2-release) 
+
+Fixes:
+
+* [Windows] Fixed DirectX 11 crash when launching on older GPUs
+
+* [Windows] Fixed being unable to compile Android projects when launching from the editor
+
+* [Windows] Fixed crash on startup for some AMD CPUs
+
+* [Windows] Fixed broken shell extensions when running RegisterShellCommands.bat in GitHub builds using "Run as Administrator"
+
+* [Windows] Added support files required for Subversion support (preview release)
+
+* [Windows] Fixed issue where engine could fail to compile with multiple versions of Visual Studio installed
+
+* [Mac] Fixed crash on startup caused by certain USB peripherals plugged in
+
+* [Mac] Fixed logs being cut off after crash
+
+* [Mac] Fixed crash on startup caused by non-standard maximum number of open files
+
+* [Mac] Noticeably improved performance on non-Intel GPUs
+
+* [Mac] Fixed crash on startup running ShooterGame as a standalone game
+
+## QFE for the 4.0.2 Hotfix
+
+If you are using the latest hotfix for this engine version and experience the issue described below, download and run this ‘quick fix’ to resolve the issue.
+
+* Fixes Launcher failing to launch after install on Mac; and unable to read info about mouse. ([report](https://answers.unrealengine.com/questions/30577/unreal-engine-install-launcher-fails-osx.html))
+
+Download 1: [http://cdn.unrealengine.com/qfe/QFE-FromCL-2034640-Shelve-2048959.zip](http://cdn.unrealengine.com/qfe/QFE-FromCL-2034640-Shelve-2048959.zip) 
+
+Download 2: [http://cdn.unrealengine.com/qfe/UnrealEngine_QFE_2039063.dmg](http://cdn.unrealengine.com/qfe/UnrealEngine_QFE_2039063.dmg) 
+
+Notes:  When you get them, install Unreal Engine app from UnrealEngine_QFE_2039063.dmg to your /Applications folder, then open Terminal app and type: 
+
+* open /Applications/Unreal\ Engine.app --args -noselfupdate
+
+-noselfupdate argument is important because without it the app would update to official version that doesn't work for you. After you download the engine using Unreal Engine app, unzip QFE-FromCL-2034640-Shelve-2048959.zip and run the patch installer it contains. After that you should be able to use the editor etc. You can start the editor by launching it directly: [REDACTED_USER_HOME]/UnrealEngine/4.0/Engine/Binaries/MacUE4Editor.app.
+
diff --git a/Engine/Documentation/Source/Support/Builds/ReleaseNotes/4_1/Updates/UnrealEngine4_1UpdateNotes.INT.udn b/Engine/Documentation/Source/Support/Builds/ReleaseNotes/4_1/Updates/UnrealEngine4_1UpdateNotes.INT.udn
new file mode 100644
--- /dev/null
+++ b/Engine/Documentation/Source/Support/Builds/ReleaseNotes/4_1/Updates/UnrealEngine4_1UpdateNotes.INT.udn
@@ -0,0 +1,55 @@
+Availability:Public
+Title: Unreal Engine 4.1 Update Notes
+Description:
+Crumbs:
+Parent:Support/Builds/ReleaseNotes/2014/4_1
+Order:0
+
+## 4.1.0 Release
+
+Released: 04/22/2014
+
+Github: [https://github.com/EpicGames/UnrealEngine/tree/4.1.0-release](https://github.com/EpicGames/UnrealEngine/tree/4.1.0-release)
+
+Release Notes:
+
+[https://docs.unrealengine.com/latest/INT/Support/Builds/ReleaseNotes/2014/4_1/index.html](https://docs.unrealengine.com/latest/INT/Support/Builds/ReleaseNotes/2014/4_1/index.html)
+
+## 4.1.1 Hotfix
+
+Released: 05/12/2014
+
+Github: [https://github.com/EpicGames/UnrealEngine/tree/4.1.1-release](https://github.com/EpicGames/UnrealEngine/tree/4.1.1-release) 
+
+Fixes:
+
+* Fixed shader material instances not being saved to DDC, causing many shaders to be rebuilt when loading the editor.
+
+* Fixed crash on iOS/Android reading blueprint assets with certain unaligned opcodes.
+
+* Fixed handling of volume buttons on Android.
+
+* Fixed sporadic crash caused by race condition in the XboxOne texture compressor.
+
+* Fixed engine compiling more shaders than necessary in the material editor.
+
+* Fixed severe performance regression for some GPU combos on Windows, where the engine would fall back to the software renderer.
+
+* Fixed crash when pressing the news feed refresh button multiple times.
+
+* Updated .gitignore to include SDL2
+
+* Addressed an analytics error
+
+## 4.1.2 Hotfix
+
+Released: 06/25/2014
+
+Github: *N/A*
+
+Fixes:
+
+* Fix crash when loading certain blueprints 
+
+* Fix corrupted map in StrategyGame causing a crash on 4.1.2 launch
+
diff --git a/Engine/Documentation/Source/Support/Builds/ReleaseNotes/4_10/Updates/UnrealEngine4_10UpdateNotes.INT.udn b/Engine/Documentation/Source/Support/Builds/ReleaseNotes/4_10/Updates/UnrealEngine4_10UpdateNotes.INT.udn
new file mode 100644
--- /dev/null
+++ b/Engine/Documentation/Source/Support/Builds/ReleaseNotes/4_10/Updates/UnrealEngine4_10UpdateNotes.INT.udn
@@ -0,0 +1,177 @@
+Availability:Public
+Title: Unreal Engine 4.10 Update Notes
+Description:
+Crumbs:
+Parent:Support/Builds/ReleaseNotes/2015/4_10
+Order:0
+
+## 4.10.0 Release
+
+Released: 11/11/2015
+
+Github: [https://github.com/EpicGames/UnrealEngine/releases/tag/4.10.0-release](https://github.com/EpicGames/UnrealEngine/releases/tag/4.10.0-release) 
+
+Release Notes: [](Support/Builds/ReleaseNotes/2015/4_10)  
+
+## 4.10.1 Hotfix
+
+Released: 12/08/2015
+
+Github: [https://github.com/EpicGames/UnrealEngine/releases/tag/4.10.1-release](https://github.com/EpicGames/UnrealEngine/releases/tag/4.10.1-release) 
+
+Fixes:
+
+*  UE-23364	Make sure that Perforce integration works with SSL connections when compiled with VS2015
+
+*  UE-21181	[CrashReport] Object from PIE level still referenced. (via FWorldContext in UEngine::AddReferencedObjects?)
+
+*  UE-23371	[CrashReport] UE4Editor_GraphEditor!FNodeFactory::CreateNodeWidget() nodefactory.cpp:91
+
+*  UE-23372	[CrashReport] Indexing Blueprints for FiB data causes: KERNELBASE!<Unknown>
+
+*  UE-23230	iOS min supported version > 6.1 = Crashing Leaderboard Writes
+
+*  UE-23366	[CrashReport] UE4Editor_Engine!USceneComponent::OnComponentDestroyed() [scenecomponent.cpp:676]
+
+*  UE-23250	PS4Tracker Crash
+
+*  UE-23394	Intermittent crash in CacheOptimizeIndexBuffer on Editor startup on Linux
+
+*  UE-23302	GitHub 1751 : Fix for USceneComponent::AttachTo() crash
+
+*  UE-20336	Hovering over tooltip for Post Process Settings causes CPU Color Scheme Change & Win10 hard-freeze
+
+*  UE-23877	Lots of people running into an Ensure in PathFollowingComponent
+
+*  UE-23232	Metal + iOS projects = Black Screen on iPhone 6s
+
+*  UE-23288	Bug: No text being displayed on iPhone 4s
+
+*  UE-23289	Bug: Crash and to home screen on iPhone 4s
+
+*  UE-23382	Compile for Size Causes Animations to Play Incompletely on iOS
+
+*  UE-23290	Bug: No Force Feedback
+
+*  UE-23664	Unable to push to SteamVR in 4.10 while Oculus 0.8 service is running
+
+*  UE-23150	Refraction does not render correctly on Mosaic devices
+
+*  UE-23611	Android streaming audio repeats at start if start time is 0.0
+
+*  UE-23694	Rework streaming audio fix for hotfix
+
+*  UE-22968	UnrealFrontend Quick Launch to Android Fails While Deploying
+
+*  UE-23655	GitHub 1776 : Fix compile error for Android x86 and x64
+
+*  UE-21858	Visible dithered quantization to remove banding in New Temporal AA Quality
+
+*  UE-23149	Async linetraces can re-call the results delegate repeatedly
+
+*  UE-23362	Gnm: submitDone () is not called for X seconds
+
+*  UE-23415	USB Keyboard not working in shipping builds
+
+*  UE-23426	Changes to World.h and WorldCollision.h need to be removed
+
+*  UE-23662	HRTF Spatialization Broken
+
+*  UE-23812	Packaged games should not send hardware survey telemetry event
+
+*  UE-23688	Modulated shadow blending incorrect with encoded HDR.
+
+*  UE-23762	Setting Runtime Generation to Static Resets Upon Closing and Re-opening Editor
+
+*  UE-23048	UAIPerceptionStimuliSourceComponent never registers with required senses
+
+*  UE-22866	[4.10 P2] UI is offcenter for dropdown menus and tooltips in editor
+
+*  UE-23064	Target build folder does not clear when repackaging
+
+*  UE-23132	Wildcard Array in Macro Causes Error Message to Appear
+
+*  UE-23233	Upgrading Project to a New Engine Version Causes Errors When Attempting to Open Tutorial Blueprints
+
+*  UE-21717	Self Shadow Only causes rendering artifacts
+
+*  UE-22949	Unable to open UE4Editor on Linux - Bad hlslcc header found
+
+*  UE-23746	4.10 is gathering editor-only BP text for game loc gathers
+
+*  UE-22157 Android In App Purchase can return false even though the transaction is successfully completed.
+
+*  UE-23190 Android/Google Play In-App Purchase Blueprint node only works if you fire the node twice
+
+* UE-23365	GitHub 1763 : Update README.md / VS2013 -> VS2015
+
+* UE-23151	UE4 4.10 needs a Localization Update
+
+## 4.10.2 Hotfix
+
+Released: 01/11/2016
+
+Github: [https://github.com/EpicGames/UnrealEngine/tree/4.10.2-release](https://github.com/EpicGames/UnrealEngine/tree/4.10.2-release) 
+
+Fixes:
+
+*  UE-23845 Crash when using "Set Key Time" on an animation key in UMG
+
+*  UE-24685 Matinee movie recording is broken in 4.10
+
+*  UE-22573 A REINST error occurs in widgets that reference one another
+
+*  UE-24115 Cannot launch the editor in DebugGame Editor configuration from Xcode
+
+*  UE-24563 Editor should launch launcher silently
+
+## 4.10.3 Hotfix
+
+Released: 02/17/2016
+
+Github: [https://github.com/EpicGames/UnrealEngine/tree/4.10.3-release](https://github.com/EpicGames/UnrealEngine/tree/4.10.3-release) 
+
+Fixes:
+
+*  OPP-4947 Installation Error code R-1638 / Unable to upgrade to 4.10.1 through launcher if 4.11 has been installed
+
+*  UE-13030	Crash related to loading sub-level - UEngine::LoadMap()
+
+*  UE-22075	Crash on compile when nodes use same name - FPersistentFrameCollectorArchive::operator<<()
+
+*  UE-22696	Crash on importing specific Skeletal FBX exported with FBX2016
+
+*  UE-22921	Parenting multiple actors under border crashes editor
+
+*  UE-23373	Crash occurs opening some animations - TSharedMapView<FName,FString>::Find()
+
+*  UE-23785	Crash in UE4Editor_Core!FSlowTask::GetCurrentMessage()
+
+*  UE-24169	Crash in PlayerInput.cpp - UPlayerInput::ProcessInputStack()
+
+*  UE-24443	Potential Crash When Migrating Assets - FGenericPlatformString::LogBogusChars<wchar_t,char>()
+
+*  UE-26315	Crash in UE4Editor_Engine!UActorComponent::BeginPlay() 
+
+*  UE-26872	Loc gathering can crash in 4.10
+
+*  UE-23231	Can't Upload iOS build signed on PC
+
+*  UE-24972	VR headtracking locked on game start for some templates
+
+*  UE-26158	bIsActorBeingDestroyed incorrectly Commented out
+
+*  UE-26657	Regression on sprite import improvements from 4.9
+
+*  UE-26903	Loc dashboard doesn't export native translations when exporting a single culture
+
+## 4.10.4 Hotfix
+
+Released: 02/22/2016
+
+Github: [https://github.com/EpicGames/UnrealEngine/releases/tag/4.10.4-release](https://github.com/EpicGames/UnrealEngine/releases/tag/4.10.4-release) 
+
+Fixes:
+
+*  UE-27124	Switching levels via blueprints crashes editor
+
diff --git a/Engine/Documentation/Source/Support/Builds/ReleaseNotes/4_11/Updates/UnrealEngine4_11UpdateNotes.INT.udn b/Engine/Documentation/Source/Support/Builds/ReleaseNotes/4_11/Updates/UnrealEngine4_11UpdateNotes.INT.udn
new file mode 100644
--- /dev/null
+++ b/Engine/Documentation/Source/Support/Builds/ReleaseNotes/4_11/Updates/UnrealEngine4_11UpdateNotes.INT.udn
@@ -0,0 +1,153 @@
+Availability:Public
+Title: Unreal Engine 4.11 Update Notes
+Description:
+Crumbs:
+Parent:Support/Builds/ReleaseNotes/2016/4_11
+Order:0
+
+## 4.11.0 Release
+
+Released: 03/31/2016
+
+Github: [https://github.com/EpicGames/UnrealEngine/tree/4.11.0-release](https://github.com/EpicGames/UnrealEngine/tree/4.11.0-release) 
+
+Release Notes: [](Support/Builds/ReleaseNotes/2016/4_11)  
+
+## 4.11.1 Hotfix
+
+Released: 04/07/2016
+
+Github: [https://github.com/EpicGames/UnrealEngine/tree/4.11.1-release](https://github.com/EpicGames/UnrealEngine/tree/4.11.1-release) 
+
+Fixes:
+
+*  UEBP-140	Oculus Integration 1.3 SDK, automation features, adapter support
+
+*  UE-28961	UnrealLightMass fails to compile with VS2015 Update 2
+
+*  UE-29031	Unable to connect to SSL Perforce Server
+
+*  UE-28933	GarbageCollection crash using IWCE
+
+*  UE-28992	[CrashReport] Crash When Enabling Capsule Direct Shadows on a Skeletal Mesh
+
+*  UE-29005	[CrashReport] UE4Editor_CoreUObject!FGCCollector::HandleObjectReferences() [garbagecollection.cpp:586]
+
+*  UE-29012	Crash attempting to open blueprints with make nodes of missing structures made prior to 4.11
+
+*  UE-28954	PSVR Using Old Tracking Data
+
+*  UE-29080	GearVR - project unexpectedly closes on launch
+
+*  UE-29108	GearVR vrplatlib.jar staging logic is backward
+
+## 4.11.2 Hotfix
+
+Released: 04/18/2016
+
+Github: [https://github.com/EpicGames/UnrealEngine/tree/4.11.2-release](https://github.com/EpicGames/UnrealEngine/tree/4.11.2-release) 
+
+Fixes:
+
+*  UE-28811	Get mouse X/ Get mouse Y returns zero after the user clicks the screen and mouse is shown
+
+*  UE-29148	Using "Show Mouse Cursor" breaks mouse axis movement
+
+*  UE-28378	Mouse is not captured when switching between clients when running PIE
+
+*  UE-29007	Construction script in parent blueprint cannot call function in child blueprint
+
+*  UE-29207	UMG Widgets aren't cleaned up when a player is removed
+
+*  UE-29208	Third Person C++ Template Doesn't use C++ Character
+
+*  UE-29300	CDO Constructor Error when opening C++ Third Person project
+
+*  UE-29136	Ensure starting PIE after editing character blueprint of a code based side scrollers
+
+*  UE-29046	Crash on exit PIE after editing character blueprint of a code based side scrollers
+
+*  UE-29209	Ensure 'Attaching instanced component to template component' in PIE in TP code project after reparenting character BP
+
+*  UE-29100	Disable spawanbles in 4.11 because they can cause scene corruption and won't be compatible with 4.12
+
+*  UE-28791	Crash smoothing landscape with foliage
+
+*  UE-28926	Editor crashes while sculpting landscape heavily covered in foliage
+
+*  UE-28999	UE 4.11 Project Containing vegetation_pack causes Engine to Crash
+
+*  UE-24725	Vehicle templates have jerky movement when turning due to SpringArmComponent
+
+*  UE-28753	SpringArm causes physics to fail for a character standing on the physics object
+
+*  UE-28991	Cannot manipulate morph targets unique to Slave Component through Master Component
+
+*  UE-29003	Manipulating Morph Target on Master Component distorts Slave Components
+
+*  UE-28789	Root motion has significant error when used in character movement
+
+*  UE-25343	Crash in SetRootBodyIndex ( ParentIndex is 0 but SpaceBases.Num() == 0)
+
+*  UE-29116	Animation CopyRecords need to calculate size at runtime
+
+*  UE-29422	Crash in FollowKnots when compiling Anim Blueprint
+
+*  UE-29326	[CrashReport] UE4Editor_Engine!FBlendedCurve::InitFrom() [animcurvetypes.cpp:590]
+
+*  UE-28011	Media Player assets do not play in packaged games or in standalone
+
+*  UE-29091	Sound Cues with Attenuation 'Override' Nodes break when Upgrading project to 4.11
+
+*  UE-29403	.wav files with max concurrency count data outside proper range will crash editor when played
+
+*  UE-29548	Crash in VoiceModuleWindows.cpp
+
+*  UE-29397	Sounds not Playing for 2nd Player Client when running in Local Multiplayer Splitscreen Mode
+
+*  UE-27646	Editor Crashes When Using Set Replicates in the Construction Script
+
+*  UE-25234	Bloom offset in right eye when rendering in stereoscopic
+
+*  UE-29513	Custom lighting channels ignored when object culled in stereoscopic
+
+*  UE-28609	R.ScreenPercentage and HMD SP commands do not work on Morpheus
+
+*  UE-29505	Wrong colors on Nexus 6 Devices
+
+*  UE-28805	HDR and Bloom not working on Mali Galaxy S7 (SM-G930F)
+
+*  UE-29420	[Metal][iOS] Memory spike caused by FMetalSurface in 4.11
+
+*  UE-29421	[Metal][iOS] Crash in FMetalDynamicRHI::RHIReadSurfaceFloatData since 4.11
+
+*  UE-29006	[CrashReport] UE4Editor-MetalRHI.dylib!FMetalCommandEncoder::RestoreRenderCommandEncoding()
+
+*  UE-29194	PS4 leaks memory in Shipping builds submission timestamp array
+
+*  UE-29390	Html 5 Launch Fails on Windows
+
+*  UE-29392	HTML5 broken load/save
+
+*  UE-28874	New Third Party Plugin does not compile on Mac and Linux and prevents project from opening
+
+*  UE-29158	4.11 regression: Stationary skylights no longer get bounce lighting after a lighting build
+
+*  UE-29255	IndirectLightingIntensity on skylights is broken - can cause level lighting to unexpectedly change after getting the new build
+
+*  UE-29297	Building Grass Map toast stays up forever
+
+*  UE-29332	[CrashReport] Crash When Switching from Lit to Light Complexity Mode
+
+*  UE-29322	[CrashReport] UE4Editor_CoreUObject!StaticFindObjectFastInternalThreadSafe() [uobjecthash.cpp:423]
+
+*  UE-29407	Editor process quits during start up in debug with debugger attached
+
+## QFE for the 4.11.2 Hotfix
+
+If you are using the latest hotfix for this engine version and experience the issue described below, download and run this ‘quick fix’ to resolve the issue.
+
+* Fixes compile issue when attempting to use Visual Studio 2015 Update 3
+
+[Windows Download](https://s3.amazonaws.com/unrealengine/qfe/QFE-FromCL-3044413.exe](https://s3.amazonaws.com/unrealengine/qfe/QFE-FromCL-3044413.exe) 
+
diff --git a/Engine/Documentation/Source/Support/Builds/ReleaseNotes/4_12/Updates/UnrealEngine4_12UpdateNotes.INT.udn b/Engine/Documentation/Source/Support/Builds/ReleaseNotes/4_12/Updates/UnrealEngine4_12UpdateNotes.INT.udn
new file mode 100644
--- /dev/null
+++ b/Engine/Documentation/Source/Support/Builds/ReleaseNotes/4_12/Updates/UnrealEngine4_12UpdateNotes.INT.udn
@@ -0,0 +1,215 @@
+Availability:Public
+Title: Unreal Engine 4.12 Update Notes
+Description:
+Crumbs:
+Parent:Support/Builds/ReleaseNotes/2016/4_12
+Order:0
+
+## 4.12.0 Release
+
+Released: 06/01/2016
+
+Github: [https://github.com/EpicGames/UnrealEngine/tree/4.12.0-release](https://github.com/EpicGames/UnrealEngine/tree/4.12.0-release) 
+
+Release Notes: [](Support/Builds/ReleaseNotes/2016/4_12)  
+
+## 4.12.1 Hotfix
+
+Released: 06/07/2016
+
+Github: [https://github.com/EpicGames/UnrealEngine/releases/tag/4.12.1-release](https://github.com/EpicGames/UnrealEngine/releases/tag/4.12.1-release) 
+
+Fixes:
+
+*  UE-31570	Crash in StaticDuplicateObject while compiling or editing variables in blueprint widgets
+
+*  UE-31550	Using Integer Bitmask crashes editor
+
+*  UE-23371	[CrashReport]UE4Editor_GraphEditor!FNodeFactory::CreateNodeWidget() nodefactory.cpp:91
+
+*  UE-26761	Crash When Compiling Child Blueprint Containing a Custom Event with the Same Name as its Parent's Event
+
+*  UE-31443	Crash occurs clicking "Select All" in Foliage Tool in QA-Sequencer map
+
+*  UE-31183	[CrashReport] UE4Editor_Engine!FAnimNode_StateMachine::Evaluate() [animnode_statemachine.cpp:745]
+
+*  UE-31574	RemoveNaNTracks causes anim DDC data to be regenerated every load on old animations
+
+*  UE-31445	[CrashReport] UE4Editor_DetailCustomizations!FBodyInstanceCustomization::GetDefaultCollisionProvider() [bodyinstancecustomization.cpp:35]
+
+*  UE-31255	Shipping Builds packaged for WinXP 32 bit crash when opening on Win7
+
+*  UE-31492	Can no longer generate a dsym on Mac
+
+*  UE-25561	Read IAP information node causes hang if connection is disconnected
+
+*  UE-31539	GPU crash on PS4 due to GS mode not getting disabled when using parallel contexts
+
+*  UE-31247	Emissive for Static Lighting discarded upon Re-opening Level or Restarting Project
+
+*  UE-28817	[CrashReport] UE4Editor_Engine!FStaticLODModel::Serialize() [skeletalmesh.cpp:1261]
+
+*  UE-15903	[CrashReport] Lighting build triggers array out of bounds crash
+
+*  UE-31620	Roughness Calculation is incorrect in 4.12
+
+*  UE-26288	Crash on unchecking Generate Lightmap UVs and Applying with Mesh already assigned and painted as Foliage
+
+*  UE-31182	[CrashReport] UE4Editor_Engine!FMaterial::GetShader() [materialshared.cpp:1564]
+
+*  UE-18906	[CrashReport] Crash in !UE4Editor_D3D11RHI!FD3D11DynamicRHI::Shutdown()
+
+*  UE-30985	Reflective surfaces rendering incorrectly in -featureleveles2 -opengl
+
+*  UE-31431	Crash in Sequencer play rate track when setting negative play rate
+
+*  UE-31647	Crash changing the color of a fade track section
+
+*  UE-31648	Sequencer audio doesn't play after rendering a movie in editor.
+
+## 4.12.2 Hotfix
+
+Released: 06/09/2016
+
+Github: [https://github.com/EpicGames/UnrealEngine/releases/tag/4.12.2-release](https://github.com/EpicGames/UnrealEngine/releases/tag/4.12.2-release) 
+
+Fixes:
+
+*  UE-31740	Crash While Compiling Blueprint in FBlueprintEditorUtils::RefreshExternalBlueprintDependencyNodes()
+
+*  UE-31655	BP EQS generators crash on adding vectors.
+
+*  UE-31723	Attempting to build MinidumpDiagnostics with Visual Studio 2015 Update 2 results in failure
+
+*  UE-31612	Physics enabled Child Actor components that spawn with with a delay, always spawns at world origin
+
+*  UE-31680	[CrashReport] DestroyChildActor->Rename->StaticFindObject()
+
+*  UE-31682	[CrashReport] UE4Editor_CoreUObject!UObjectBaseUtility::GetOutermost() [uobjectbaseutility.cpp:118]
+
+*  UE-31736	SteamController BackRight is not setup correctly
+
+*  UE-29543	Adjusting Light Source Radius/Length creates undesirable effect
+
+*  UE-31589	Capture reflections on translucent materials are broken
+
+*  UE-31715	Static Meshes Disappear with Media Player Asset Applied as Material in Standalone or Packaged Game
+
+*  UE-31687	Vive: Positional Latency with a Possessed Pawn's Camera Component
+
+## 4.12.3 Hotfix
+
+Released: 06/16/2016
+
+Github: [https://github.com/EpicGames/UnrealEngine/releases/tag/4.12.3-release](https://github.com/EpicGames/UnrealEngine/releases/tag/4.12.3-release) 
+
+Fixes:
+
+*  UE-31735	Get Class Defaults node disconnects output pins when project upgraded from 4.11 to 4.12
+
+*  UE-31743	Components derived from scene component reset vector location upon closing and re-opening editor
+
+*  UE-31812	CrashReporter Client has no options to close without sending
+
+*  UE-31789	GitHub 2484 : Git Plugin: fix the Submit To Source Control menu broken by new "migrate" support in 4.12 (and master)
+
+*  UE-31702	Assertion opening blueprint with HISMC created in 4.11
+
+*  UE-30934	GitHub 2406 : Git plugin: Fix for Git diff not working in UE 4.12 (and master)
+
+*  UE-31265	GitHub 2432 : BUGFIX: ProceduralMeshComponent crashes game
+
+*  UE-31578	[CrashReport] UE4Editor_Renderer!InjectTranslucentLightArray() [translucentlighting.cpp:1578]
+
+*  UE-31851	Reflections have an extra vector offset than in 4.11
+
+*  UE-31787	Vive: Positional Latency w/ Rotation
+
+## 4.12.4 Hotfix
+
+Released: 06/28/2016
+
+Github: [https://github.com/EpicGames/UnrealEngine/releases/tag/4.12.4-release](https://github.com/EpicGames/UnrealEngine/releases/tag/4.12.4-release) 
+
+Fixes:
+
+*  UE-32179	Properties of blueprint based behavior tree decorators are reset on game start
+
+*  UE-32178	Vehicles set to "obstacle only" in RVO are moving on their own
+
+*  UE-32356	Editor crashes when opening several marketplace assets - Assertion failed: !HasAnyFlags
+
+*  UE-32414	Using TickInterval for a tick used as a prerequisite will crash
+
+*  UE-31897	PhysicsVolume weak pointer returns null only in packaged project
+
+*  UE-31938	Crash report window does not open on GitHub source builds
+
+*  UE-29086	Crash When Typing a Node Comment and Hovering Over the Alignment Option
+
+*  UE-30963	UCS continuously runs while a blueprint's thumbnail in content browser is focused
+
+*  UE-32381	Crash updating thumbnail after compiling blueprint
+
+*  UE-32283	Animations imported with Skeletal Mesh are not marked as dirty and don't get saved
+
+*  UE-30878	[CrashReport] UE4Editor_Foliage!FFoliageInstanceBaseCache::AddInstanceBaseId() [foliageinstancebase.cpp:101]
+
+*  UE-31956	Instance changes no longer propagating to ChildActor after updating from 4.11.2
+
+*  UE-32219	[CrashReporter] Crash in UEngine::UpdateTimeAndHandleMaxTickRate()
+
+*  UE-32276	AActor::OnRep_AttachmentReplication regression on detach
+
+*  UE-32124	Crash when Deleting Concatenator Node while Playing Sound Cue
+
+*  UE-28604	[CrashReport] UE4Editor_Engine!FSharedPoolPolicyData::GetPoolBucketIndex() [gpuskinvertexfactory.cpp:75]
+
+*  UE-31445	[CrashReport] UE4Editor_DetailCustomizations!FBodyInstanceCustomization::GetDefaultCollisionProvider() [bodyinstancecustomization.cpp:35]
+
+*  UE-32244	Grabbing a Destructible Mesh Child blueprint with a Phsyics Handle crashes the editor
+
+*  UE-31766	[CrashReport] UE4Editor_Engine!AActor::InvalidateLightingCacheDetailed() [actor.cpp:4111]
+
+*  UE-32060	Editing Landscape Collision Mip Level Crashes Editor
+
+*  UE-32132	Safe zone visualization and debugging is broken on Windows
+
+*  UE-31441	A3D doesn't work on PS4
+
+*  UE-31450	[CrashReport] UE4Editor_UnrealEd!FLightmassExporter::WriteSceneSettings() [lightmass.cpp:1818]
+
+*  UE-31765	[CrashReport] UE4Editor_Renderer!FMeshMaterialShader::SetMesh<FRHIVertexShader * __ptr64>() [shaderbaseclasses.cpp:457]
+
+*  UE-31454	Using VR Preview and Shader Complexity viewmode will cause a crash
+
+*  UE-29332	[CrashReport] Crash When Switching from Lit to Light Complexity Mode
+
+*  UE-32036	GitHub 2504 : [Editor Hang new in 4.12 fix in FinalPostProcessSettings] Essential project-killing hang fix ? Rama
+
+*  UE-32415	Editing Landscape Collision Mip Level Crashes Editor
+
+*  UE-29146	CameraVectorWS Input Data as Post Process Blendable not Rendering Correctly
+
+*  UE-31549	r.ScreenPercentage.Editor does not function
+
+*  UE-32507	Third party plugin authors cannot stage PRX files on PS4
+
+*  UE-31914	Remove embedded vertex buffers from GoogleVRHMDViewerPreviews.cpp that are killing compile times and binary sizes
+
+## 4.12.5 Hotfix
+
+Released: 07/07/2016
+
+Github: [https://github.com/EpicGames/UnrealEngine/releases/tag/4.12.5-release](https://github.com/EpicGames/UnrealEngine/releases/tag/4.12.5-release) 
+
+Fixes:
+
+*  UE-32588	VS2015 Update 3 prevents the editor from compiling code
+
+*  UE-32621	[CrashReport] UE4Editor_Renderer!FGlobalShader::SetParameters<FRHIPixelShader * __ptr64,FRHICommandList>() [globalshader.h:134]
+
+*  UE-32715	Crash when selecting "show" > "uncompressed animation" in Persona
+
+*  UE-32772	Hitting ensure in AACtor::BeginPlay in an actor with ChildActorComponents
+
