Copying //UE4/Dev-Documentation to //UE4/Samples-Main (Source: //UE4/Dev-Documentation @ 3372893)

diff --git a/Engine/Documentation/Source/Support/Builds/ReleaseNotes/4_13/Updates/UnrealEngine4_13UpdateNotes.INT.udn b/Engine/Documentation/Source/Support/Builds/ReleaseNotes/4_13/Updates/UnrealEngine4_13UpdateNotes.INT.udn
new file mode 100644
--- /dev/null
+++ b/Engine/Documentation/Source/Support/Builds/ReleaseNotes/4_13/Updates/UnrealEngine4_13UpdateNotes.INT.udn
@@ -0,0 +1,183 @@
+Availability:Public
+Title: Unreal Engine 4.13 Update Notes
+Description:
+Crumbs:
+Parent:Support/Builds/ReleaseNotes/2016/4_13
+Order:0
+
+## 4.13.0 Release
+
+Released: 09/01/2016
+
+Github: [https://github.com/EpicGames/UnrealEngine/releases/tag/4.13.0-release](https://github.com/EpicGames/UnrealEngine/releases/tag/4.13.0-release) 
+
+Release Notes: [](Support/Builds/ReleaseNotes/2016/4_13)  
+
+## 4.13.1 Hotfix
+
+Released: 09/29/2016
+
+Github: [https://github.com/EpicGames/UnrealEngine/releases/tag/4.13.1-release](https://github.com/EpicGames/UnrealEngine/releases/tag/4.13.1-release) 
+
+Fixes:
+
+*  UE-35735	Crash Occurs While Running EQS Using Add Generated Vector Node
+
+*  UE-35852	AI Move To Location zeros velocity each call, causing the AI Pawn to freeze if told to move again before reaching destination
+
+*  UE-35585	Migration from 4.12.5 to 4.13 can result in Blueprint functions losing input/ouput variables
+
+*  UE-35609	Delay nodes used in a macro cause subsequent macro calls during the delay to never complete
+
+*  UE-35728	Some files in the Git repo have inconsistent line endings on Linux
+
+*  UE-36107	Cannot create installed build on Linux
+
+*  UE-35236	Launcher installed build is missing the Documentation/Extras folder
+
+*  UE-36214	Crash in TMap struct deserializer
+
+*  UE-35884	Crash when pressing enter after a shape has closed when using the Geometry editor
+
+*  UE-35822	CRASH: Occurring in UTEditor when attempting to re-order BSP brushes
+
+*  UE-35048	[CrashReport] UE4Editor_PropertyEditor!FDetailCategoryImpl::GetParentLayoutImpl() [detailcategorybuilderimpl.h:168]
+
+*  UE-35482	Selecting a brush surface that is flush with other brush surfaces selects all of them after Building Geometry
+
+*  UE-35627	Lightmap Coordinate Index is not automatically set to 1 on import any longer with lightmap UV
+
+*  UE-35608	FBX & OBJ Files no longer Import with their UVs Depending on Software Used
+
+*  UE-36383	Crash moving a spline point after making it the root component of a blueprint
+
+*  UE-35658	Editor crash when a spline component is added to a blueprint and compiled
+
+*  UE-36096	Inifinite loop in tick task manager with tick when paused and interval
+
+*  UE-35254	Editor crashes when importing alembic asset
+
+*  UE-36121	Crash when importing new mesh and regenerating skeleton for additive animation
+
+*  UE-35546	Pose flickering when changing lods and using Layered Blend per Bone
+
+*  UE-35353	Crash when spawning destructible objects in 2 player PIE
+
+*  UE-35639	HLOD generated proxy doesn't create or assign target lightmap coordinate index any longer
+
+*  UE-35494	Invisible landscape components when using tessellation
+
+*  UE-35873	LODBias causes landscape component to stop rendering if set >0 and Forced LOD set to -1
+
+*  UE-35301	Crash saving hidden landscape streaming level with offset
+
+*  UE-35850	Crash in landscape editor with early Z pass
+
+*  UE-35709	[CrashReport] UE4Editor_Landscape!FLandscapeEditDataInterface::SetAlphaData() [landscapeeditinterface.cpp:2617]
+
+*  UE-36321	Editor Crashes After Deleting Actor with Media Texture
+
+*  UE-36289	WmfMedia: Memory leak while playing videos
+
+*  UE-35477	File Media Source assets created via drag & drop have absolute instead of relative path
+
+*  UE-35504	PS4 Media player has potentially dangerous memory management
+
+*  UE-35696	PS4Media: CPU/GPU may crash due to race condition in destructor
+
+*  UE-35598	Devices reporting OpenGL ES 3.2 such as Galaxy Note 7 detect as ES2-only
+
+*  UE-35179	Need to fall back to GLES if Vulkan API version is too old
+
+*  UE-34782	GPU particles no longer work on iOS or TVOS Metal devices
+
+*  UE-33379	Cannot select Android RHI and feature level from device profile
+
+*  UE-35258	ES2 Fallback Fails on Nougat Devices that do not support Vulkan
+
+*  UE-35261	Android sessions do not appear in Session Frontend
+
+*  UE-32340	HTML5 packaging fails with Github couldn't verify LLVM version
+
+*  UE-32086	Bloom blur produces multi-colored artifacts at the bottom of the screen on Xbox.
+
+*  UE-35492	PS4 deployment does not deploy all the required files.
+
+*  UE-35743	Packaging for Android while using Online Framework and Online Subsystem Steam Plugins cause failure
+
+*  UE-35185	Opening a widget blueprint on Linux causes editor to crash
+
+*  UE-35738	IPv6 implementation is not working on IOS
+
+*  UE-30959	Remove all HITCHHUNTER warnings in next release (4.13)
+
+*  UE-36342	Texture issues with movie player active on Android
+
+*  UE-35363	Huge game window when launching onto Safari 9.1.2
+
+*  UE-35558	Writing to a file on Android doesn't update the tracked file length
+
+*  UE-35907	Slate Remote Enable Remove Server option in Project Settings -> Plugins causes editor to freeze and a memory leak
+
+*  UE-36099	Dots on screen on PS4
+
+*  UE-35794	PS4 can hang on GPU when using lit translucent surface materials
+
+*  UE-26375	Crash when editing draw size for a widget component
+
+*  UE-17227	[CrashReport] UE4Editor_Engine!DoReadTaskResults() [shadercompiler.cpp:147]
+
+*  UE-35240	Protostar crashes on launch on Mac
+
+*  UE-28708	GitHub 2188 : Fixed a crash exiting VR Preview on Windows GL4.
+
+*  UE-35320	Stop>Set Sequence>Play causes spawnables to be duplicated
+
+*  UE-35508	Master sequence dialog button is hidden
+
+*  UE-35588	Incorrect translation from matinee move sub tracks to sequencer
+
+*  UE-35285	Sequencer stop node does not return camera to player
+
+*  UE-35571	Sequence recording in standalone crashes
+
+*  UE-35107	Matinee to Level Sequence converter doesn't bind properties
+
+*  UE-22581	Black artifacts/semi-circles on sides of morpheus when moving left to right
+
+*  UE-36076	PS4Tracker's motion sensor data needs to use a valid sizeOfThis
+
+*  UE-35520	sceKernelWaitEqueue timeout in MorpheusHMD.cpp fix for 4.13.1
+
+*  UE-34986	Need to merge Oculus Online Subsystem Plugin update to expose Android support
+
+*  UE-34786	Editor primitives render incorrectly in stereoscopic with instanced stereo
+
+*  UE-35881	PSVR hmdSetupDialog cancel asserts on app startup
+
+*  UE-36307	PSVR Motion Sensors Need 120hz Update
+
+*  UE-33763	GearVR applications show black screen on Adreno devices
+
+*  UE-36078	Integrate fixes for GoogleVR SDK 1.0
+
+*  UE-35528	PSVR UE4 stuttering issue
+
+*  UE-35344	Crash when playing a HapticFeedbackEffect_Soundwave
+
+*  UE-35964	Crash playing a HapticFeedbackEffect_Soundwave with no sound assigned
+
+## 4.13.2 Hotfix
+
+Released: 10/26/2016
+
+Github: [https://github.com/EpicGames/UnrealEngine/releases/tag/4.13.2-release](https://github.com/EpicGames/UnrealEngine/releases/tag/4.13.2-release) 
+
+Fixes:
+
+* UE-37704	Crash opening packaged project on iOS
+
+* UE-37249	MallocPoisonProxy can result in a memory stomp for aligned allocations in Debug and Development
+
+* UE-36573	Need to update to libpng for Android; current one triggers security alert on Google Play
+
diff --git a/Engine/Documentation/Source/Support/Builds/ReleaseNotes/4_14/Updates/UnrealEngine4_14UpdateNotes.INT.udn b/Engine/Documentation/Source/Support/Builds/ReleaseNotes/4_14/Updates/UnrealEngine4_14UpdateNotes.INT.udn
new file mode 100644
--- /dev/null
+++ b/Engine/Documentation/Source/Support/Builds/ReleaseNotes/4_14/Updates/UnrealEngine4_14UpdateNotes.INT.udn
@@ -0,0 +1,173 @@
+Availability:Public
+Title: Unreal Engine 4.14 Update Notes
+Description:
+Crumbs:
+Parent:Support/Builds/ReleaseNotes/2016/4_14
+Order:0
+
+## 4.14.0 Release
+
+Released: 11/15/2016
+
+Github: [https://github.com/EpicGames/UnrealEngine/releases/tag/4.14.0-release](https://github.com/EpicGames/UnrealEngine/releases/tag/4.14.0-release) 
+
+Release Notes: [](Support/Builds/ReleaseNotes/2016/4_14)
+
+## 4.14.1 Hotfix
+
+Released: 12/08/2016
+
+Github: [https://github.com/EpicGames/UnrealEngine/releases/tag/4.14.1-release](https://github.com/EpicGames/UnrealEngine/releases/tag/4.14.1-release) 
+
+Fixes:
+
+*  UE-35609	Delay nodes used in a macro cause subsequent macro calls during the delay to never complete
+
+*  UE-38690	Regression - Nested scene component subobjects are no longer being registered when the owning scene component is added to an existing Actor instance as an instanced component.
+
+*  UE-38835	Rename "Visual Studio '15'" to "Visual Studio 2017"
+
+*  UE-39298	UE4 and UnrealLightmass fail to build Module.SwarmInterface.cpp due to an error trying to open include file 'metahost.h' in SwarmInterface.cpp
+
+*  UE-39284	Running GenerateProjectFiles.bat -2017 completes with warnings
+
+*  UE-38898	UE4 is not compatible with Visual Studio 2017 RC
+
+*  UE-38604	PS4 Platform Staging copies debug files (not allowed)
+
+*  UE-38992	[CrashReport] UE4Editor_Core!FString::Mid() [unrealstring.h:1154]
+
+*  UE-37249	MallocPoisonProxy can result in a memory stomp for aligned allocations in Debug and Development
+
+*  UE-38991	[CrashReport] UE4Editor_Engine!FUObjectAnnotationSparse<FBoolAnnotation,1>::AddAnnotation() [uobjectannotation.h:76]
+
+*  UE-38752	Cannot Delete Folder Containing Sub Folders from Content Browser
+
+*  UE-38721	Mouse cursor position is locked to center of screen when bShowMouseCursor is false
+
+*  UE-38599	User Defined Struct Editor Hides Long Variable Lists
+
+*  UE-38706	Dropdowns display incorrectly in non 1:1 zoom in BP on Mac and PC
+
+*  UE-38806	Potential fatal assert when using a custom MID with a TextRenderComponent
+
+*  UE-39059	No scrollbar in Structure UI
+
+*  UE-39198	Limited range of motion when controlling the camera view with the mouse
+
+*  UE-11796	Deleting folders in the Content Browser will not delete them permanently from Project Folder
+
+*  UE-38925	Reimport for a single FBX with multiple meshes will add all mesh pieces instead of the single mesh.
+
+*  UE-38935	Crash with FBX Import as Skeletal Mesh
+
+*  UE-38749	Pressing enter does not appear to activate OnClicked events in UMG When a button is focused
+
+*  UE-37698	Widgets can no longer take focus after clicking somewhere other than the widget itself
+
+*  UE-36766	Crash attempting to open an asset from Aim Offset graph in Persona
+
+*  UE-38811	Camera attached to a Spring Arm jitters when put at a +90 or -90 degree angle on the Y axis.
+
+*  UE-38906	Crash when streaming in a level with a child actor
+
+*  UE-38518	Animation Blueprint: Default values cannot be changed after compiling if node is currently selected
+
+*  UE-39197	Delay node of 0.0 causes crash
+
+*  UE-39137	Crash when Dragging Skeletal Mesh with Apex Clothing into Level
+
+*  UE-38998	[CrashReport] UE4Editor_Engine!UMaterialInterface::GetBaseMaterial() [materialinterface.cpp:268]
+
+*  UE-38716	Physics are popping in Vehicle Advanced Template
+
+*  UE-38989	[CrashReport] PhysX3PROFILE_x64!physx::PxcDiscreteNarrowPhase() [pxcnpbatch.cpp:293]
+
+*  UE-38694	Landscape Grasstype warning causes log spam when landscape is streamed in or out
+
+*  UE-39278	Assertion failed: !LandscapeActor || LandscapeActor == Landscape
+
+*  UE-38900	Paper2d Plugin Content Missing from Content Browser
+
+*  UE-39120	ETC2 Android devices are no longer supported in the Google Play Store as of 4.14
+
+*  UE-37101	Xbox One: Issues updating texture subregions
+
+*  UE-35738	IPv6 implementation is not working on IOS
+
+*  UE-39420	QAGame crashes when launched onto Android
+
+*  UE-38709	Using Deferred Decal Material With Dbuffer Decal Enabled Crashes Editor
+
+*  UE-38223	Crash In Material Editor When Previewing A Node
+
+*  UE-39051	Engine crash when Scene Capture Component 2D is added to viewport with forward rendering
+
+*  UE-39123	Mac Editor crashes when opening Realistic Rendering sample - FMetalRHICommandContext::RHISetShaderUniformBuffer
+
+*  UE-38609	DFAO runs out of memory on xb1 (DX11 RHI) due to an RHI bug
+
+*  UE-38699	Material Expression comparison keys incorrect, produces duplicate code
+
+*  UE-38325	Tangent Update Feature Not Working
+
+*  UE-38768	Contact Shadow artifacts with ray lengths despite length
+
+*  UE-34264	Procedural Mesh Component Crashes on XB1
+
+*  UE-39132	Crash when adding actors to the level
+
+*  UE-38353	Crash while deleteing a UMG Animation Track referencing a deleted UMG asset with duplicated BSPs in-scene
+
+*  UE-39067	Oculus Touch controllers jitter when the gamethread runs quickly
+
+*  UE-39113	Post process material needs a full RT clear with hmd mask.
+
+*  UE-39122	Crash when Simulating in Editor after Painting Foliage
+
+*  UEVR-381	Interleaved Compositor: Execute init views before late update
+
+*  UEVR-379	Interleaved Compositor: Implementing Post-Present handoff
+
+*  UE-38670	New IsA() code can be thread-unsafe during saving
+
+## 4.14.2 Hotfix
+
+Released: 01/04/2017
+
+Github: [https://github.com/EpicGames/UnrealEngine/releases/tag/4.14.2-release](https://github.com/EpicGames/UnrealEngine/releases/tag/4.14.2-release) 
+
+Fixes:
+
+* UE-39970	Stat FPS and other text rendering as blocks on android.
+
+* UE-39560	[CrashReport] UE4Editor_Engine!ACullDistanceVolume::CanBeAffectedByVolumes() [culldistancevolume.cpp:62]
+
+* UE-39557	[CrashReport] UE4Editor_Persona!FDelaunayTriangleGenerator::FlipTriangles() [sanimationblendspace.cpp:375]
+
+* UE-39341	Compiling a blueprint with SetCollisionResponseToChannel Causes Crash
+
+* UE-38818	[CrashReport] UE4Editor_D3D11RHI!TerminateOnDeviceRemoved() [d3d11util.cpp:176]
+
+* UE-38772	[CrashReport] UE4Editor-Core.dylib!FThreadStats::FThreadStats()
+
+* UE-38750	Imported Materials Are Not Assigned To Correct Element Index Based Off Mat ID (3dsMax)
+
+* UE-38710	Vehicle physics are resolving erratically and causing random jitter/flips in the Advanced Vehicle Template
+
+## 4.14.3 Hotfix
+
+Released: 01/17/2017
+
+Github: [https://github.com/EpicGames/UnrealEngine/releases/tag/4.14.3-release](https://github.com/EpicGames/UnrealEngine/releases/tag/4.14.3-release) 
+
+Fixes:
+
+* UE-40229	Editor Crashes when importing some FBXs with multiple Material IDs in 4.14.2
+
+* UE-40294	Put fix for UE-39713 into 4.14.3 (Memory leak in GPURealloc)
+
+* UE-40243	Calling FAudioDevice::StopAllSounds from UEngine::CommitMapChange asserts
+
+* UE-40156	PhysX issue causing lag, high ping, and failures
+
