Copying //UE4/Dev-Documentation to //UE4/Samples-Main (Source: //UE4/Dev-Documentation @ 3372893)

diff --git a/Engine/Documentation/Source/Support/Builds/ReleaseNotes/4_15/ReleaseNotes_4_15.INT.udn b/Engine/Documentation/Source/Support/Builds/ReleaseNotes/4_15/ReleaseNotes_4_15.INT.udn
new file mode 100644
--- /dev/null
+++ b/Engine/Documentation/Source/Support/Builds/ReleaseNotes/4_15/ReleaseNotes_4_15.INT.udn
@@ -0,0 +1,2787 @@
+Availability:Public
+Title: Unreal Engine 4.15 Release Notes
+Description: Release Notes for Unreal Engine 4.15
+Crumbs:%ROOT%, Engine
+Parent:Support/Builds
+Order:15
+
+![image alt text](image_0.png)(w:769 h:249 convert:false)
+
+## What's New
+
+Unreal Engine 4.15 includes significant gains in overall stability, enhancements to developer workflows, and improvements to runtime performance resulting in greater efficiency during development and superior end user experiences after release.
+
+Compile times for programmers are drastically reduced - **by as much as 50%**! Reloading content while Unreal Editor is running, Reroute nodes in Materials, a new Blendspace Editor, new mathematics Blueprint nodes, and more contribute to an even more streamlined development process in this release.
+
+For those looking to squeeze out every drop of performance, Cooking Blueprints to C++ native code is no longer an experimental feature, the Texture Streaming system has gotten an overhaul, and Alternate Frame Rendering with NVIDIA SLI gives a boost on high end systems.
+
+The Cinematics and Animation pipelines continue to strengthen with Animation blending now possible in Sequencer, linking Animation Curves to bones for culling in LODs, and modifying curves in Animation Blueprints with the Modify Curve node. Level Sequences can now be embedded in Actor Blueprints, and early support for Level Sequence Components is available for early adopters.
+
+Developing for Nintendo Switch is available as experimental as part of the platform improvements. GPS data is now accessible on iOS and Android using the new Location Services. Also on iOS, streaming audio and remote notifications are fully supported. Monoscopic Far Field Rendering is an option for mobile VR platforms, HDR display output is available in an experimental state, and the ability to use Playstation VR Aim Controllers is also added.
+
+This release includes **79 improvements** submitted by the incredible community of Unreal Engine developers on GitHub! Thanks to each of these contributors to Unreal Engine 4.15:
+
+[REDACTED_CONTRIBUTORS]
+
+## Major Features 
+
+### New: Nintendo Switch Support
+
+We have partnered with Nintendo to help Unreal Engine 4 developers bring their games to Nintendo Switch!
+
+![image alt text](image_1.png)(w:1000 h:385 convert:false)
+
+We have added early support for developing games for the Nintendo Switch console. It is still considered experimental in this release. We plan for it to be in a shippable state in the next release.
+
+### New: Physically Based Default Post-Processing Settings
+
+The default tonemapper now provides a more film-like color and tone response!
+
+![image alt text](image_2.jpg)(w:929 h:522 convert:false)
+
+The filmic tonemapper matching the Academy Color Encoding System (ACES) is now on by default. This tonemapper also enables the engine to target multiple display types including HDR displays. Bloom settings have been updated to provide a more physically correct result by default.
+
+The filmic tonemapper will look different than the tonemapper in previous versions of the engine. You can switch the tonemapper back by adding r.TonemapperFilm=0 to the SystemSettings section of your project's DefaultEngine.ini. We intend to remove the old tonemapper in a future version of the engine.
+
+### New: Improved Texture Streaming
+
+The Texture Streaming system has been optimized to reduce CPU usage, memory usage, and load times while eliminating low resolution artifacts and automatically handling limitations of varying memory budgets of different platforms.
+
+![image alt text](image_3.jpg)(w:929 h:522 convert:false)
+
+* **Texture Memory Usage Improvements** - The texture streaming system now handles per component visibility and computes (smaller) per material bounds for static geometry resulting in as much as a **40% reduction in texture memory usage**.
+
+* **Texture Loading Time Improvements** - **Texture load times are 2 to 3 times faster** due to tracking visible mips, enabling the streamer to focus on visible textures followed by prefetch data and forced load data.
+
+* **CPU Time Reduction**
+
+    * **Game thread update time is reduced by 50%** by moving dynamic component processing and other tasks to the async texture streaming task, enabling them to occur in parallel with other game thread tasks.
+
+    * **Texture processing stalls for streaming levels are reduced up to 98%**s due to amortized processing.
+
+* **Low Resolution Artifacts Reduction** - Mesh UV densities are now computed per material instead of per mesh and take into account LODs. Texture streaming is also now supported on a wider variety of Component types, including Particle Systems and Instanced Static Meshes.
+
+* **Automatic Memory Budgeting** - The Texture Streaming system now automatically adjust its behavior to deal with varying memory budgets without making manual tweaks. The system selects which textures need to be reduced using a variety of heuristics to maintain visual quality. 
+
+### New: Cooking Blueprints to C++
+
+Packaging Blueprints into native source code to reduce the overhead that goes into executing Blueprints is now production ready! 
+
+![image alt text](image_4.png)(w:929 h:457 convert:false)
+
+This feature has been available as experimental, but we've been successfully using it for some of our own projects, like [Robo Recall](https://www.epicgames.com/roborecall/en-US/home) and are removing that label.
+
+You can either exclusively nativize individual Blueprints or inclusively nativize all Blueprint assets. Large projects may run into executable size issues on some platforms when using the Inclusive mode; for many projects the best practice is to use Exclusive mode and target only performance critical Blueprints. For more information, see our [Blueprint Nativization documentation](Engine/Blueprints/TechnicalGuide/NativizingBlueprints).
+
+### New: Faster C++ Compile Times
+
+The Unreal Engine 4 codebase now **compiles up to 50% faster**! Header files only include other headers they need, rather than every source file including large monolithic headers like Engine.h and UnrealEd.h - we call this _Include What You Use (IWYU)_.
+
+![image alt text](image_5.jpg)(w:929 h:557 convert:false)
+
+The conventions for the new model are as follows:
+
+* **Every header now includes everything it needs to compile** - There is a CoreMinimal.h header that contains a set of ubiquitous types from Core (eg. FString, FName, TArray, FVector, etc...), which is now included first by most headers.
+
+* **Every .cpp file includes its matching .h file first** - This helps validate that each header is including everything it needs. UnrealBuildTool will check that source files adhere to this convention.
+
+* **No engine code includes a monolithic header such as Engine.h or UnrealEd.h any more** - You will get a warning if you try to include one of these from the engine. They still exist for compatibility with game projects and do not produce warnings when included there.
+
+* **No engine code explicitly includes a precompiled header any more** - We still use PCHs, but they're force-included on the compiler command line by UnrealBuildTool instead. This lets us tune what they contain without breaking any existing include dependencies.
+
+[REGION:note]
+You can set 'bEnforceIWYU = false' from a *.build.cs file (or through the global BuildConfiguration.xml file) to disable warnings or errors about violations of these rules.
+[/REGION]
+
+### New: UI Blur Rects
+
+You can now blur other UI and the 3D scene in UMG using the new Background Blur widget! Background Blur widgets also have a content slot which can be used to add widgets that will not be blurred.
+
+ ![image alt text](image_6.png)(w:857 h:478 convert:false)
+
+There is a single strength value that defines how blurry the background content is, which goes from 0 (off) to 100 (full strength) and can be animated to modify the amount of blurring over time.
+
+[REGION:imagetable]
+| ![image alt text](image_7.jpg)(w:293 h:293 convert:false) | ![image alt text](image_8.jpg)(w:292 h:292 convert:false) | ![image alt text](image_9.jpg)(w:294 h:293 convert:false) |
+| --- | --- | --- |
+| Strength = 2 | Strength = 10 | Strength = 50 |
+[/REGION]
+
+### New: Material Editor Reroute Nodes
+
+Reroute nodes are now available in the Material Editor to enable better organization of your Material graphs! They function identically between Blueprints and the Material Editor and are purely visual so they should have no impact on the final instruction count of your Materials.
+
+![image alt text](image_10.jpg)(w:929 h:308 convert:false)
+
+To add Reroute nodes in the Material Editor:
+
+* Under Utility in the Material Palette, Drag **Add Reroute Node** to the graph.
+
+* Double-click on a connection spline to insert a new Reroute node inline.
+
+### New: Sequencer Animation Blending
+
+Sequencer now supports weighted blending between multiple animations enabling complex animation sequences to be constructed from existing animations. Any number of animations, including additive (Local Space) animations, can be blended at any given time.
+
+[OBJECT:EmbeddedVideo]
+	[PARAMLITERAL:videoid]
+	bowGfA5VL10
+	[/PARAMLITERAL]
+	[PARAMLITERAL:width]
+	720
+	[/PARAMLITERAL]
+	[PARAMLITERAL:height]
+	678
+	[/PARAMLITERAL]
+	[PARAMLITERAL:type]
+	simple
+	[/PARAMLITERAL]
+[/OBJECT]
+
+You can control the weight of an animation by expanding the track and keying the value in the desired timeline. For full body animations the weights will be normalized so the mesh isn't under-scaled or over-scaled. For additive animations, the specified weight values are used unmodified.
+
+To enable your Skeletal Mesh to use multiple poses for animation blending, a new Blend Multi by Alpha node is available to add to the Anim Graph of the Animation Blueprint for the mesh. You can add new poses from the context menu.
+
+### New: Sequencer Animations in Blueprints (Experimental)
+
+The new Actor Sequence plugin - which adds an Actor Sequence Component - makes it possible to embed Sequences into Actor Blueprints! This enables reuse of Sequences by binding the animations to the Blueprint instance and triggering them using the Event Graph. You can also add an Actor Sequence Component to any Actor in the world to add animations to a single instance of an Actor.
+
+[OBJECT:EmbeddedVideo]
+	[PARAMLITERAL:videoid]
+	auuPgita6GI
+	[/PARAMLITERAL]
+	[PARAMLITERAL:width]
+	720
+	[/PARAMLITERAL]
+	[PARAMLITERAL:height]
+	678
+	[/PARAMLITERAL]
+	[PARAMLITERAL:type]
+	simple
+	[/PARAMLITERAL]
+[/OBJECT]
+
+[OBJECT:EmbeddedVideo]
+	[PARAMLITERAL:videoid]
+	PWlQP3_Mwwk
+	[/PARAMLITERAL]
+	[PARAMLITERAL:width]
+	720
+	[/PARAMLITERAL]
+	[PARAMLITERAL:height]
+	340
+	[/PARAMLITERAL]
+	[PARAMLITERAL:type]
+	simple
+	[/PARAMLITERAL]
+[/OBJECT]
+
+### New: Content Hot-Reloading (Experimental)
+
+We have significantly improved support for reloading packages from disk on the fly with the new Content Hot-Reloading feature.
+
+With this feature, all of the in-editor source control operations are more stable when dealing with loaded assets. You also gain a "Reload" option under "Asset Actions" that forcibly reloads the asset from disk, which is extremely useful if you've made unsaved changes to an asset and want to revert it back to the on-disk version.
+
+To enable this feature, go to your "Editor Preferences", and enable "Content Hot-Reloading" under the "Experimental" section.
+
+[REGION:note]
+Currently, you have to use the in-editor source control to take advantage of hot-reloading content while the editor is open. There are a few technical reasons for this that we're hoping to address in the future to enable you to sync from anything.
+[/REGION]
+
+### New: High Dynamic Range Display Output (Experimental)
+
+Unreal Engine 4 can now output to High Dynamic Range (HDR) displays taking advantage of features such as higher contrast and wider color gamut!
+
+![image alt text](image_11.jpg)(w:929 h:376 convert:false)
+
+[REGION:caption]These are simulated images and are for demonstration purposes only[/REGION]
+
+HDR display output support has been added for Windows Direct3D 11 (NVIDIA-only) and Mac Metal platforms. Other devices are coming in future releases or will be available through GitHub as added. The current implementation handles full processing of the rendered scene through the standard ACES transforms (look modification -> reference rendering -> output display) and currently has paths implemented for 1000-nit or 2000-nit display outputs. We bake the device-specific transforms to a lookup table as part of the tonemapping pass which includes most of the standard white balance and color grading controls. By default the film-mapping curve is disabled during HDR output due to poor compatibility with LDR controls. After tonemapping we apply the relevant range adjustments and encoding if required.
+
+Additionally there is experimental LDR UI composition support. This is handled automatically when enabling HDR output and attempts to match the LDR look as closely as possible, although it's recommended to boost the UI slightly via the new console variables to avoid looking washed out next to the vibrancy of the main scene.
+
+HDR output support is currently detected at launch and flagged only if a supported display is connected. When activating fullscreen mode and enabling HDR output the application will attempt to automatically move to the supported display. HDR output can be toggled either through the new console variables or through GameUserSettings, which will automatically clamp to the closest available match for the current device and set all necessary flags.
+
+### New: Mobile Rendering Preview (OpenGL ES 3.1, Metal, Vulkan)
+
+You can now **preview high-end mobile features in the editor's viewport** using the improved Mobile Preview feature. This mode will emulate the feature set available to iOS Metal, Android OpenGL ES 3.1 and Android Vulkan devices.
+
+![image alt text](image_12.jpg)(w:840 h:505 convert:false)
+
+[REGION:note]
+Only devices enabled via project settings will be shown in the preview platform menu.
+[/REGION]
+
+### New: Map/Set Containers for Blueprints
+
+Collections of items can now be stored in Blueprints using Map and Set variable types in addition to the previously available Array variable type. Maps and Sets guarantee that they contain only unique items, with no repeating entries; no repeating 'keys', where maps are concerned. In particular, they are useful when doing frequent lookups on large datasets. 
+
+![image alt text](image_13.jpg)(w:761 h:356 convert:false)
+
+The Variable Type control is now a drop down, enabling you to select 'Single Variable', 'Array', 'Set', or 'Map'. When 'Map' is selected a second drop down for the 'value' type appears.
+
+![image alt text](image_14.jpg)(w:333 h:205 convert:false)
+
+Set and map variables declared in C++ can now be exposed to Blueprints, though replication of map and set properties is not yet supported in C++ or Blueprints.
+
+### New: Gameplay Tags
+
+We've made several improvements to Gameplay Tags and they are now a fully supported feature! 
+
+A Gameplay Tag is a name of the form "Family.Genus.Species" or just “IndividualThing”. These tags are registered in a central dictionary. The nesting of names using “.” supports searching for matches against both the literal tag and all parent tags. A GameplayTagContainer is a set of tags that can also be used for match operations.
+
+Tags are useful to define a list of names for developers to use instead of typing text in by hand. To use gameplay tags they must be registered in the central tag dictionary, which can be accessed from the new GameplayTags project settings view:
+
+![image alt text](image_15.jpg)(w:580 h:728 convert:false)
+
+If you enable the "Import Tags from Config" option, tags can be added from the editor through the Gameplay Tag List on this page, or from the UI used to select tags. Once you have added Tag properties to your data, you can query them from either Blueprint or C++ and use them to change functionality. The BlueprintGameplayTagLibrary has several useful functions. 
+
+### New: Color Grading Controls 
+
+We added all new controls for the Color Grading properties in Post Process settings that provide a more intuitive experience and more precise control.
+
+![image alt text](image_16.jpg)(w:533 h:218 convert:false)
+
+Highlights of the new controls include:
+
+* Colored lines under each slider denote which color channel the slider is affecting.
+
+* Sliders have improved precision when you edit the value. You can also press the **Shift** key while editing to increase precision even further.
+
+* A new color wheel modifies the X, Y, Z channels all at once. A slider under the color wheel controls the amplitude of the X, Y, Z channels
+
+### New: VR Editor Menu Improvements
+
+The Quick Menu and Radial Menu in the VR Editor have been updated with a new look and feel as well as improved usability and functionality. 
+
+[OBJECT:EmbeddedVideo]
+	[PARAMLITERAL:videoid]
+	tjchjILby4E
+	[/PARAMLITERAL]
+	[PARAMLITERAL:width]
+	409
+	[/PARAMLITERAL]
+	[PARAMLITERAL:height]
+	540
+	[/PARAMLITERAL]
+	[PARAMLITERAL:type]
+	simple
+	[/PARAMLITERAL]
+[/OBJECT]
+
+Click the thumbstick or touchpad to summon/dismiss the radial menu, select options with a direction on the thumbstick/touchpad, and then select with the trigger. You can also select items with the laser, but the thumbstick or touchpad overrides the laser. 
+
+[OBJECT:EmbeddedVideo]
+	[PARAMLITERAL:videoid]
+	IcngDWI5dj4
+	[/PARAMLITERAL]
+	[PARAMLITERAL:width]
+	481
+	[/PARAMLITERAL]
+	[PARAMLITERAL:height]
+	540
+	[/PARAMLITERAL]
+	[PARAMLITERAL:type]
+	simple
+	[/PARAMLITERAL]
+[/OBJECT]
+
+Quick menu summoning and dismissing functionality is mostly the same, but the laser is now hidden while the quick menu is visible to reduce visual clutter.
+
+### New: VR Editor Number Pad
+
+Quickly enter numbers (0-9) as well as negative and decimal values in text fields in the VR Editor using the new Number Pad UI!
+
+[OBJECT:EmbeddedVideo]
+	[PARAMLITERAL:videoid]
+	af3txMXNi5Y
+	[/PARAMLITERAL]
+	[PARAMLITERAL:width]
+	471
+	[/PARAMLITERAL]
+	[PARAMLITERAL:height]
+	540
+	[/PARAMLITERAL]
+	[PARAMLITERAL:type]
+	simple
+	[/PARAMLITERAL]
+[/OBJECT]
+
+You can use the trigger button to select numbers, the grip button to press enter, and the modifier key for backspace. To dismiss the number pad, either click somewhere else within the scene, or click the touchpad or thumbstick, just like the standard radial menu.
+
+### New: Texture Streaming Debugging Visualization Tools
+
+Three new viewmodes are available that make it much easier to ensure the texture streaming system is working optimally for your game!
+
+![image alt text](image_17.png)(w:262 h:200 convert:false) ![image alt text](image_18.png)(w:308 h:198 convert:false) ![image alt text](image_19.png)(w:310 h:198 convert:false)
+
+* **Primitive Distance Accuracy** - This shows the accuracy of the distance computed by the streamer as compared to the distance computed by the GPU. The streamer computes the distance from the viewpoint to the texture instance axis-aligned bounding boxes (AABB). This bounding box is computed in the build process by summing up the component lod-section AABBs that are using the texture.
+
+* **Mesh UV Densities Accuracy** -** **This view mode shows the effect of the UV densities used by the streamer for each mesh. If the UV densities are wrong, then the textures on those meshes risk being under-streamed (blurry) or over-streamed (wasting memory).
+
+* **Material Texture Scales Accuracy** - Some textures used in a material are sampled using a scaled value of one of the mesh UV channels. The texture streaming build tries to compute which texture coordinate is used and what scale is applied to each sampled texture. This could fail for many reasons, in which case the streamer fallbacks on assuming the texture was sampled with texcoord 0 and a scale of 1.
+
+Because the material samples many textures, this view mode shows the worst error (in term of under-streaming and over-streaming) across all textures sampled. Those two extrema values are displayed through the checkerboard look : one shows the worst oversampling while the other shows the worst undersampling. Each texture can also be inspected through the "View Mode Options".
+
+You can also now access live texture streaming stats using the **stat streaming** console command. This stat reports the performance, memory usage and other metrics used by the texture streamer. In the memory part, the word "pool" refers to conceptual memory, and does not relate to actual memory taken. The word “mips” maps to the memory currently taken by textures. The memory metrics are classified in three main pools: Texture, Streaming and Wanted. Those pool size are indicated on the right of each line.
+
+### New: Flight Stick and Steering Wheel Support (Windows)
+
+You can provide support in Windows for steering wheels, flight sticks, and other non-XInput supported devices using the new Raw Input plugin!
+
+[OBJECT:EmbeddedVideo]
+	[PARAMLITERAL:videoid]
+	QMD-OQ2Q-dI
+	[/PARAMLITERAL]
+	[PARAMLITERAL:width]
+	720
+	[/PARAMLITERAL]
+	[PARAMLITERAL:height]
+	408
+	[/PARAMLITERAL]
+	[PARAMLITERAL:type]
+	simple
+	[/PARAMLITERAL]
+[/OBJECT]
+
+Adding new devices is as easy as setting up a configuration for them in the project settings or editing DefaultInput.ini. All of the Vehicle templates and Vehicle Game have been configured to work with the Logitech G920.
+
+### New: Force Feedback Component
+
+The new Force Feedback Component can be added to Actors in the world to create localized haptic effects! These components have attenuation properties to determine the intensity of the playback of the force feedback pattern based on the distance between the player and the effect. The attenuation properties can either be specified directly on the component or you can create a Force Feedback Attenuation asset in the content browser and reuse it for multiple components. Force Feedback Components can also be spawned into the World from blueprints similar to audio, decals, and particle emitters.
+
+### New: PlayStationⓇVR Aim Controller Support
+
+Unreal Engine 4 games on PlayStationⓇVR are even more immersive with support for the PlayStationⓇVR Aim Controller!
+
+![image alt text](image_20.png)(w:429 h:285 convert:false)
+
+To activate, simply change the "Hand" value to “Gun” on your Motion Controller Component. The PlayStationⓇVR Aim Controller's buttons map to the same buttons on the Dual Shock gamepad.
+
+### New: Mathematics Material Nodes
+
+We've added some commonly requested mathematics material nodes that should help avoid the need for material functions or custom nodes for common operations.
+
+![image alt text](image_21.png)(w:546 h:308 convert:false)
+
+The nodes marked with the "Fast" tag will execute approximations instead of the real instructions. These can give a worthwhile performance improvement to more complex materials but have input restrictions and precision tradeoffs.
+
+### New: Previous Frame Switch Material Node
+
+A **Previous Frame Switch** node has been added to support complex vertex animations implemented in materials to generate correct motion vectors so that they work correctly with temporal AA and motion blur.
+
+[OBJECT:EmbeddedVideo]
+	[PARAMLITERAL:videoid]
+	YbjqiNZr4K4
+	[/PARAMLITERAL]
+	[PARAMLITERAL:width]
+	720		
+	[/PARAMLITERAL]
+	[PARAMLITERAL:height]
+	408
+	[/PARAMLITERAL]
+	[PARAMLITERAL:type]
+	simple
+	[/PARAMLITERAL]
+[/OBJECT]
+
+Material vertex animations that are only a function of time already work without modification but cannot account for other variables, such as material parameters, that affect the animation. The previous frame switch enables artists to solve this problem manually when they can track how these parameters change, for instance in a blueprint.
+
+Note that using this feature requires Accurate velocities from Vertex Deformation to be enabled.
+
+### New: Pre-Skinned Local Normal Material Node
+
+You can now use **Pre-Skinned Local Normals **to achieve effects such automatic UV mapping on skeletal meshes using locally-aligned tri-planar materials.
+
+[OBJECT:EmbeddedVideo]
+	[PARAMLITERAL:videoid]
+	LNvy1BiDFOk
+	[/PARAMLITERAL]
+	[PARAMLITERAL:width]
+	440
+	[/PARAMLITERAL]
+	[PARAMLITERAL:height]
+	540
+	[/PARAMLITERAL]
+	[PARAMLITERAL:type]
+	simple
+	[/PARAMLITERAL]
+[/OBJECT]
+
+This feature is exposed as a new **Pre-Skinned Local Normal** material node.  This works in a similar way to the **Pre-Skinned Local Position** node, but returns the local surface normal for Skeletal Meshes and Static Meshes. This opens the door to more local-space, mesh aligned effects or advanced use-cases such as writing dynamic surface data to a mask read-back in another material.
+
+### New: Improved Metal Support on macOS and iOS
+
+Metal support has been extended to use many of the new API and shader language features added by Apple in macOS 10.12 Sierra & iOS 10. 
+
+For Mac, full support for Unreal Engine's tessellation features has been implemented in cooperation with Apple using the new Metal Tessellation pipeline. This enables displacement and World Position Offset features on Mac that have previously been unavailable. The Elemental sample has been updated to use tessellation on Mac Metal by default.
+
+[OBJECT:EmbeddedVideo]
+	[PARAMLITERAL:videoid]
+	4yD_ZRTKf-A
+	[/PARAMLITERAL]
+	[PARAMLITERAL:width]
+	720
+	[/PARAMLITERAL]
+	[PARAMLITERAL:height]
+	695
+	[/PARAMLITERAL]
+	[PARAMLITERAL:type]
+	simple
+	[/PARAMLITERAL]
+[/OBJECT]
+
+Both Mac and iOS now support specifying the Metal Shader Standard version as a project setting, enabling developers to choose whether they wish to adopt new features or maintain full backward compatibility. Enabling the new Metal v1.2 standard enables all Metal platforms to use Unordered Access Views in pixel shaders to generate arbitrary data as on other RHIs.  
+
+There is also experimental support for HDR rendering on Macs with an appropriate display built in, such as the Retina iMac, or attached via Thunderbolt. Enabling the per-project HDR setting will configure Metal for rendering to an HDR display supporting the DCI-P3 color gamut. 
+
+Metal support now also includes an internal validation layer that is enabled by default for development and runs independently from Xcode. This makes it possible to report much more information about errors encountered when running on Metal even without developer tools installed to aid debugging intermittent or hardware specific problems. This layer is automatically compiled out for performance sensitive Test & Shipping builds.
+
+### New: Sequencer Event Track Payloads
+
+Events in Sequencer can now contain payloads of data so that you can trigger events with more contextually-rich information about how to react to the playing sequence. To trigger an event with a payload, create a struct that defines the parameters for your event, and apply that to the event's parameter struct property:
+
+![image alt text](image_22.jpg)(w:557 h:218 convert:false)
+
+Events can call functions and trigger custom events in event contexts, provided the specified parameter struct is bitwise compatible with the destination function's prototype.
+
+### New: Animate Dynamic Objects with Sequencer
+
+Sequencer track bindings can be set using Blueprints so you can apply Sequences to Objects at runtime! This means you can animate a different object than the one the Sequence was authored with. To extract the sequence binding identifiers from a sequence, use the new GetSequenceBindings node.
+
+### New: Sequencer UI Improvements
+
+Sequencer has a new Transform Keys/Sections tool that enables you to reposition and rescale the current selection.
+
+[OBJECT:EmbeddedVideo]
+	[PARAMLITERAL:videoid]
+	Kd3CwzMwTLs
+	[/PARAMLITERAL]
+	[PARAMLITERAL:width]
+	892
+	[/PARAMLITERAL]
+	[PARAMLITERAL:height]
+	196
+	[/PARAMLITERAL]
+	[PARAMLITERAL:type]
+	simple
+	[/PARAMLITERAL]
+[/OBJECT]
+
+Other improvements:
+
+* It's now possible to copy and paste tracks onto compatible objects.
+
+* When viewing a shot that resides within a master sequence, we now evaluate the entire sequence at its relative time. The shot itself can be isolated by setting "Evaluate Sub Sequences in Isolation".
+
+* Component keyframing now adds the capability for all keyframing tools, including autokey, to select and key components individually.
+
+* Attach tracks now actually attach actors to their destination, rather than matching the transform. This makes it possible for attachments and relative transforms to be sequenced together.
+
+### New: Sequencer Audio Volume/Pitch Curves
+
+You can now control the volume and pitch of Audio sections directly in Sequencer, enabling designers to set up fade in, fade out, and other effects.
+
+### New: Blendspace Editor Improvements
+
+The Blendspace Editor has a new look and improved workflows!
+
+[REGION:imagetable]
+| ![image alt text](image_23.jpg)(w:462 h:422 convert:false) | [INCLUDE:#BlendspacePreview] |
+| --- | --- |
+| | |
+[/REGION]
+
+<!--
+[EXCERPT:BlendspacePreview]
+[OBJECT:EmbeddedVideo]
+	[PARAMLITERAL:videoid]
+	aOZvFrZA0Lo
+	[/PARAMLITERAL]
+	[PARAMLITERAL:width]
+	324
+	[/PARAMLITERAL]
+	[PARAMLITERAL:height]
+	422
+	[/PARAMLITERAL]
+	[PARAMLITERAL:type]
+	simple
+	[/PARAMLITERAL]
+[/OBJECT]
+[/EXCERPT:BlendspacePreview]
+-->
+
+* **Multiple ways of changing sample properties** - In addition to dragging and dropping samples, a context menu and spin-boxes have been added to change either the sample value or animation sequence.
+
+* **(Pin-)Previewing the blendspace by holding shift down** - To allow previewing the mesh from different angles, zoom-levels etc, we have added the ability to pin the preview value; by holding down the **Shift** key, you can move the preview value around.
+
+* **Samples get snapped to new grid when it's resized** - When the number of grid divisions or the size of either axis is changed, the existing samples will be snapped to the closest available grid point - those for which no valid position could be found are marked as invalid and will be drawn in red.
+
+* **Drag and drop operations** - Validation and preview of drag and drop operations from the Animation Tools Asset Browser have been improved to show where the sample will be placed and whether or not the animation is valid for the Blendspace.
+
+* **Stretch to fit vs Stretch to Largest Axis** - Since the UI space occupied by the grid has changed we have added the ability to either stretch the grid to fit all the layout space or to stretch it to the largest axis. This means that the largest axis will be the max width or height for the grid.
+
+### New: Pose Snapshotting
+
+We've added the ability to capture a runtime Skeletal Mesh pose in Blueprints for use in the Anim Blueprint!
+
+![image alt text](image_24.jpg)(w:929 h:524 convert:false)
+
+This can be useful for blending from a death pose (after simulating ragdoll) into a 'get up' animation. You can also save a pose snapshot to a variable for later use with any Skeletal Mesh Component.
+
+### New: Mass Properties Debugging
+
+We've added the ability to visualize the center of mass and inertia tensor associated with physics objects. 
+
+![image alt text](image_25.jpg)(w:868 h:338 convert:false)
+
+To visualize the mass properties, enable the new showflag (**Show > Advanced > Mass Properties**). The thickness of each axis indicates the moment of inertia magnitude along that axis.
+
+### New: Linking Curves to Bones
+
+Animation curves can now be linked to a specific bone in the Skeleton so they can be culled out with the bone when changing LODs, which is especially useful when setting up facial animations! This can also be used in conjunction with layer blending to prevent overriding curve values from the incorrect layer.
+
+### New: Modify Curve Animation Node
+
+Animation curves can now be modified using arbitrary logic in the Animation Graph to enable more dynamically-driven animations.
+
+![image alt text](image_26.png)(w:929 h:196 convert:false)
+
+After adding the new Modify Curve node, you can right click on it and choose which curves you want to modify. Each curve will then appear as an input float pin that logic can be connected to.
+
+### New: Font Asset Improvements
+
+Font data can now be shared between Font assets using the new Font Face asset, reducing memory consumption and improving stability!
+
+![image alt text](image_27.png)(w:648 h:173 convert:false)
+
+Font Face is now the asset that stores the font data, and these assets are simply referenced by the Font assets. This means that the same font data can be reused for multiple font assets, or even multiple typefaces within a Font asset. Existing Font assets will automatically upgrade their internal font data into embedded Font Face assets during load. You can use the Font Editor to split these embedded assets out into real Font Face assets that may be edited and shared.
+
+### New: Linux ARM64 Support
+
+We have added support for ARM64 (AArch64) devices running Linux! Note that any such devices must have support for desktop OpenGL 3.2 or later to run the engine.
+
+### New: Alternate Frame Rendering Support for NVIDIA SLI
+
+![image alt text](image_28.png)(w:553 h:356 convert:false)
+
+The engine now supports AFR rendering for Nvidia SLI configurations to provide superior performance when multiple graphics cards are detected. The largest improvement comes from the renderer copying inter-frame dependencies between GPUs as early as possible. Developers will need to work with NVIDIA to test their game and to have it automatically switched over to AFR where appropriate.
+
+### New: Monoscopic Far Field Rendering for Mobile VR (Experimental)
+
+Monoscopic far field rendering can provide a big performance boost by only rendering distant objects once!
+
+[OBJECT:EmbeddedVideo]
+	[PARAMLITERAL:videoid]
+	bAyqVgv-MqA
+	[/PARAMLITERAL]
+	[PARAMLITERAL:width]
+	720
+	[/PARAMLITERAL]
+	[PARAMLITERAL:height]
+	408
+	[/PARAMLITERAL]
+	[PARAMLITERAL:type]
+	simple
+	[/PARAMLITERAL]
+[/OBJECT]
+
+Stereo divergence decreases as objects are rendered farther from the camera. At a certain distance, the stereo rendering of distant objects is indistinguishable from a regular monoscopic rendering. Monoscopic far field rendering takes advantage of this by dividing the scene into two partitions with a clipping plane: near field and far field. Everything on the near field side of the clipping plane is rendered in stereo, everything on the far field side of the plane is only rendered once and then composited into the near field result.
+
+Monoscopic far field rendering is only supported on mobile platforms for now. To enable monoscopic far field rendering, select the checkbox under Project Settings -> Rendering -> VR. We don't currently support both mobile multi-view and monoscopic far field simultaneously and mobile HDR needs to be disabled.
+
+### New: Custom Stencil on Mobile
+
+Mobile devices can now use 'Custom Stencil' in post-processing materials allowing for object outlines, visualizing occluded objects, or other such effects.
+
+![image alt text](image_29.jpg)(w:929 h:524 convert:false)
+
+This feature requires MobileHDR to be enabled. To enable this feature in your project go to Project Settings -> Rendering -> Postprocessing  and set 'Custom Depth-Stencil Pass' to 'Enabled with Stencil'.
+
+### New: Mobile Shader Memory Reduction Options
+
+You can now significantly reduce memory usage due to shader permutations on mobile by disabling lighting setups that your mobile game does not make use of!
+
+The following settings available under Project Settings -> Rendering can be used to reduce memory usage and decrease app package size:
+
+* Disabling *Allow Static Lighting*
+
+* Disabling *Support Combined Static and CSM Shadowing*
+
+* Disabling *Support Distance Field Shadows*
+
+* Disabling *Support Movable Directional Lights*
+
+* Setting *Max Movable Point Lights* to 0
+
+* Enabling *Use Shared Movable Point Light Shaders*
+
+Setting the flags to remove lighting setups not used by the SunTemple sample content reduced its Android package size by 22MB.
+
+### New: Android graphics debugger support integration
+
+Android applications can now be packaged to support a third-party Graphics Debugger. You can choose from either the Mali Graphics Debugger or Adreno Profiler depending on your device's GPU. Graphics Debugger options can be found in Project Settings -> Android. 
+
+You first need to download and install these debuggers from the GPU vendor's website, and they each require some small amount of additional setup to configure your device for debugging. Be sure to to follow the directions that appear after selecting the debugger type.
+
+Also note that if you package your app to support a particular GPU debugger, it may not function correctly when run on a device with a different GPU.
+
+### New: Location Services for GPS data on Android and iOS
+
+Access to GPS data is now available on Android and iOS!  A new On Location Changed delegate is available and Blueprint nodes are provided under Services->Mobile->Location.
+
+### New: Streaming Audio on iOS
+
+Streaming audio is now supported on iOS! Streaming audio can reduce both memory usage and level load times. Assets already marked for streaming will automatically be streamed as on other platforms.
+
+### New: Remote Notifications for iOS
+
+Remote notifications for IOS are now supported! The game can register for callbacks to handle the notification. The editor provides tools to set up notifications for your app.
+
+### New: Vehicle Support moved to a plugin
+
+We have now refactored our PhysX Vehicle support as an optional plugin! This makes it easy for games that are not using vehicles to exclude this feature and save disk space and memory. This work also adds several useful physics extension points to Engine (e.g. OnPhysSceneInit/Term, OnPhysSceneStep) to make it easier for other developers to write their own similar systems.
+
+[REGION:note]
+A manual step is required if you want to migrate old 'TireType' assets (deprecated in Engine) to the new 'TireConfig' assets (defined in the new PhysXVehicles plugin). There is a new console command available in the UE4 Editor that will do this for you: simply summon the console and enter 'ConvertTireTypes'. This will create new TireConfig assets, copy over properties, and assign them to any VehicleWheel Blueprints. The old TireType assets can then be removed.
+[/REGION]
+
+**Release Notes**
+
+**AI**
+
+* New: A Cone EQS generator has been added.
+
+* Bugfix: Fixed AI perception system not properly registering perception listeners spawned as part of "begin play" in streamed-in sublevels.
+
+* Bugfix: Fixed a bug in the Env Query Generator Simple Grid resulting in one extra column and row of points being generated.
+
+* Bugfix: A number of fixes to how and when static mesh collision gets applied to navmesh generation process. This makes static mesh collision changes instantly applied to navmesh on the level.
+
+* Bugfix: Pasting nodes in the Behavior Tree editor has been fixed to make pasting location predictable.
+
+* Bugfix: Fixed a bug in EQS debug drawing resulting in item labels not showing.
+
+* Bugfix: Fixed a bug in EQS debug drawing resulting in displaying labels of item that are no longer being debugged.
+
+* Bugfix: Fixed vertical jitter in replicated Nav Walking movement.
+
+* Improvements for using gameplay debugger in Simulate in Editor mode, unified category toggling with input events between simulate and game/play modes.
+
+**Behavior Tree**
+
+* Bugfix: Fixed the Blackboard component pausing but never being unpaused if we ended up restarting the Behavior Tree instead of continuing.
+
+**Debugging Tools**
+
+* Bugfix: Fixed gameplay debugger's EQS category skipping timestamps, and improved readability of presented data.
+
+**Navigation**
+
+* New: Added new component for creating navigation links: Nav Link Component.
+
+* Bugfix: Get Random Point In Navigable Radius has been fixed to not return invalid locations if there's any navmesh within specified radius.
+
+* Bugfix: Fixed a bug in the Nav Link Proxy Point Links' Area Class application that was resulting in the value getting ignored.
+
+**Animation**
+
+* New: Added experimental support for facial animation.
+
+* New: Anim Blueprint thread safety is now checked in the compiler.
+
+    * Added extra warnings in the Anim Blueprint compiler based around these new keys to help people catch suspect thread usage.
+
+* New: Expanded pose snapshot system.
+
+    * Allows poses to be stored in variables so they can be managed like any other Blueprint data.
+
+    * Added modes to Anim Node Pose Snapshot so that we can either use the named pose or a Pose Snapshot variable pin.
+
+* New: Added editor-only flag Update Animation In Editor and setter function to Skeletal Mesh Components to allow playback of animations on Skeletal Meshes placed in the level.
+
+* Bugfix: Fixed a crash when retargetting animation assets to another skeleton via duplicate and retarget.
+
+* Bugfix: Fixed marker sync crash when animation being played by a sequence node changes while that nodes sync group is not active.
+
+* Bugfix: Fixed 'auto' check box in Morph Target panel.
+
+* Bugfix: Fixed situation in which stale material references could be present as override materials on Static and Skeletal meshes inside of a level.
+
+* Bugfix: Fixed removal of bone tree entries from a Skeleton after performing "Remove Unused Bones" (caused retargeting options to become out of sync).
+
+* Bugfix: Fixed anim blueprint compiler not following reroute nodes when building cached pose fragment list.
+
+* Removed 4.11 deprecated functions from animation systems.
+
+* Removed warning on Least Destructive animation compression (it was incorrectly applied to least destructive due to legacy reasons).
+
+* Added extra information to log warning about using additive animations in Single Node Anim Instance.
+
+* Changed "Ensure All Indices Have Handles" to try and maintain validity of as many of the handles as possible to stop crash when calling “Get Key.”
+
+* Morphtarget applying priority has changed.
+
+    * First, it will apply animation curves to active morphtarget list.
+
+    * Second, it will apply (and override if exists) Set Morph Target curves to active morphtarget list.
+
+    * Set Morph Target can also clear the weight if you decided to not to remove from the list.
+
+**Animation Assets**
+
+* New: Added the ability to specify per-sample animation rate scales in blendspaces.
+
+* Bugfix: Fixed shadowing issue for Alembic Caches which had incorrect (too small) bounds.
+
+* Bugfix: Fixed popping in On Rep Replicated Anim Montage.
+
+* Made Scrubbing the Play Back Speed property on Geometry Cache components less sensitive.
+
+* Changed behaviour for calculating blendspace grid weighting for one, two or colinear triangles.
+
+    * **One**: Fill grid weights to single sample.
+
+    * **Two**: Find closest point on line between the two samples for the grid point, and weight according to the distance on the line.
+
+    * **Colinear**: Find two closest samples and apply behaviour above.
+
+**Animation Blueprint**
+
+* New: Added Blend Option (and optional Custom Curve) to Pose Blend Node.
+
+* Bugfix: Fixed anim sub-instances only allowing one pin to work when any pin required a call out to the VM for evaluation.
+
+* Bugfix: Fixed Joint Target Location not being available in the detail panel.
+
+* Exposed various montage functions to Blueprint.
+
+* Added curve panel to anim BP editor.
+
+* Copy-to-array now works with the anim blueprint 'fast path'.
+
+* 'Blend Two Poses Together Per Bone' will now blend curves using 'combine'. Previously it only took curve data from the second pose input.
+
+**Import/Export**
+
+* Bugfix: Fixed an issue where Matrix Samples would be imported incorrectly for Alembic Caches.
+
+* Changed the way we populate Smoothing Group information for Alembic Caches.
+
+* Refactored the Skeletal LOD workflow. Now we track source files for LODs to save time when reimporting LODs often. It's still possible to pick new files and overwrite the current settings.
+
+**Skeletal Mesh**
+
+* New: Added support for per-component vertex color overrides for Skeletal Meshes.Nintendo
+
+    * Use the Set Vertex Color Override function (also exposed to Blueprints) to specify a new color buffer to apply.
+
+* Bugfix: Fixed edge cases so Skeletal Mesh Components should now properly scale collisions when simulating, or when using physical animation.
+
+* Bugfix: Fixed bounds not being updated on existing skeletal mesh components when changing physics assets.
+
+* Made sure Set Morph Target still take precedence over Animation Curves for Slave Component.
+
+**Tools**
+
+* New: Added access to the viewport client from Persona Viewport.
+
+* New: Added "Show Selected and Parents" to bone display options.
+
+* New: Added display/edit of bone transforms in details panel.
+
+* New: Widened bone reference widget to give a better view of the name and added full name to tooltip.
+
+* New: Auto select current bone when opening the skeleton tree on a bone reference.
+
+* New: Newly added virtual bones are now selected in the skeleton tree.
+
+* New: Changed "Reference Pose" option for additive animation ref pose to "Skeleton Reference Pose" to make it clearer that it is not the meshes reference pose being used.
+
+* Bugfix: Fixed crash in thumbnail rendering when creating a new montage.
+
+* Bugfix: Fixed skeleton tree selection disappearing when filtering changes.
+
+* Bugfix: Fixed notify validation not triggering when using set time/set frame context menu options.
+
+* Bugfix: Fixed montage thumbnail rendering using ref pose instead of the actual asset.
+
+* Bugfix: Fixed "number of curves" label on an animation not updating after adding or removing a curve.
+
+* Bugfix: Fixed newly imported animations having black thumbnails.
+
+* Bugfix: Fixed compile error when blend space on aim offset nodes is exposed as pin.
+
+* Bugfix: Fixed scrolling in Pose Asset editor panel.
+
+* Bugfix: Fixed tooltip when dragging animation assets onto animation player nodes.
+
+* Bugfix: Fixed new virtual bones being created off screen.
+
+* Bugfix: Fixed missing mesh when setting preview skeletal meshes in animation editors.
+
+* Current asset is now initially selected in the asset family shortcut bar dropdown.
+
+* Camera is no longer reset when setting skeletal meshes in the skeletal mesh editor (etc.).
+
+    * Focus shortcut (F key) now works in the skeleton tree too.
+
+* Removed FPersona and supporting classes.
+
+    * Please use individual asset editors instead.
+
+* Rezeroing animation is now explicit about what it does (current vs. specified frame).
+
+* Changed slot node category from Blends to Montage.
+
+**Audio**
+
+* New: Sound Cues can now be played/stopped by clicking a button on the thumbnail.
+
+* New: Added the ability to disable Master Reverb and EQ in Project Settings.
+
+* BugFix: Audio Components no longer override subtitle priority by default.
+
+    * Any existing Audio Component that had changed the priority from the default value will still have bOverrideSubtitlePriority as true, those that had not changed the priority from the default have been set to false.
+
+* BugFix: Fixed a crash that would happen when a Sound Cue Delay node set to 0.0.
+
+* BugFix: Fixed a small memory leak in sound concurrency management. 
+
+**Automation**
+
+* New: Added architecture hash to path for generated includes on Linux Installed Build.
+
+* New: Included .tps files in the installed build.
+
+* New: Changed BuildPlugin command so that it reads from installed platform data to decide which targets to build for more flexible <Copy> task for BuildGraph. Now takes "From" and "To" attributes, which may take individual files as well as wildcards, and permits renaming as part of copies. The Separate "Files" attribute is still supported, but is now optional, and is used to filter the list of source files.
+
+    * BuildGraph: Added a <Move> task, which can move or rename files using the same syntax as the <Copy> task.
+
+* New: Added support in BuildGraph for multi-line properties, declared inside a <Property> tag. Mutli-line properties can be useful for making lists for filtering. Each non-empty line inside the property tag is stripped of leading and trailing whitespace, and appended to the property value separated by a semicolon. 
+
+* New: Added architecture hash to path for generated includes on Linux Installed Build.
+
+**Blueprints**
+
+* New: Making the UI for adding/removing parameters in custom Blueprint functions behave similarly to the struct creation dialog in the Content Browser.
+
+    * The New button is no longer in the bottom right of Inputs and Outputs. Instead, add via the + icon in the Category header.
+
+    * Parameter ordering controls have been moved onto the main parameter row instead of being nested inside the collapse panel.
+
+    * A tooltip will now let you know the full parameter name and type when you hover over the editable name field.
+
+* New: Added RotatorToAxisAndAngle function.
+
+* New: Added aliases for "dot" and "cross" math operations in the Blueprint Math Expression node.
+
+* New: Added support for DataTable members variables in Blueprints. 
+
+* BugFix: Fixed a crash on load for Blueprints casting to a ChildActorComponent's parent.
+
+* BugFix: Fixed a crash when using "Select All Input nodes" option from the Blueprint pin menu. 
+
+* BugFix: Fixed a crash which could happen when using reroute nodes in Blueprints.
+
+* BugFix: Fixed a crash when selecting an enum value from a Blueprint node's enum combo box.
+
+* BugFix: Fixed a crash that could happen after receiving a "Graph is linked to object(s) in external map" error message.
+
+* BugFix: Fixed a crash when debugging Blueprints across multiple worlds.
+
+* BugFix: Fix crash in execArray_Append that could occur when one of the target array's UProperties was missing.
+
+* BugFix: Fixed a crash that would sometimes occur when collapsing nodes to a function or macro. 
+
+* BugFix: Fixed a rare crash in IsTypeCompatiableWithProperty.
+
+* BugFix: Fixed a crash during bootup involving CDO of nativized Blueprints. 
+
+* BugFix: Fixed a stale object crash in HandlePlaceholderScriptRef.
+
+* BugFix: Fixed a crash caused by an object name collision that could occur when loading some older Blueprint assets. 
+
+* BugFix: Fixed a that would occur when loading a Blueprint with a collapsed node with a split struct pin. 
+
+* BugFix: Fixed a crash  with Undo/Redo commands causing the component and actor selection to become unsynchronized.
+
+* BugFix: Fixed a regression in which nested scene component subobjects would no longer be registered after construction of an instance-added component.
+
+* BugFix: Fixed inaccurate ensure message that could trigger when choosing a bad target from the Blueprint debug menu.
+
+* BugFix: The Vector2D default pin values are now properly saved
+
+* BugFix: Fixed the inability to edit an Instanced Static Mesh Component instance array on Actor instances when the Instanced Static Mesh Component is inherited from a Blueprint class.
+
+* BugFix: An explicit warning is now logged on BuildCookRun or a standalone cook when the -nativizeAssets flag is omitted from the command line for a nativization-enabled project.
+
+* BugFix: Fixed various component instance data loss issues for components added to an Actor through one or more Add Component nodes in a Blueprint graph.
+
+* BugFix: Fixed attached root scene component relative transform values after reconstruction of an Actor-based Blueprint instance.
+
+* BugFix: Fixed a bug that would drop pin connections when ctrl dragging them.
+
+* BugFix: Fixed a bug that would cause Volume Actor Reference nodes to use the wrong class icon.
+
+* BugFix: Fix for breakpoints failing in Level Blueprints when multiple Clients are running in PIE.
+
+* BugFix: Fixed an issue with CreateRenderTarget2D called in construction script during cooking. 
+
+* BugFix: Fixed an issue which could cause invalid Blueprint warnings when using a cast node.
+
+* BugFix: Fixed an issue in the struct editor that would prevent you from including spaces in a variable name.
+
+* BugFix: Fixed an issue that could cause a Blueprint compiler warning when using one of the Ease nodes. 
+
+* BugFix: Fixed an issue where Blueprints could improperly warn/error about missing variables. 
+
+* BugFix: Fixed a bug where Blueprints could improperly error about missing Event Dispatchers. 
+
+* BugFix: Fixed a bug that could cause latent Blueprint nodes to lose pin connections when using a split pin.
+
+* BugFix: Fixed an issue where copying variable nodes and pasting them into another graph with a variable of the same name and type would prevent that node from showing up in a reference search. 
+
+* BugFix: Fixed an issue where Reroute nodes for variable would not always mirror data properly from the source pin. 
+
+* BugFix: Fixed an issue while loading of circularly dependent Blueprints when a user defined struct's default value was part of the dependency graph. 
+
+* BugFix: Dragging and dropping a function into a graph can not be properly undone.
+
+* BugFix: Corrected the behavior of nativized UserDefinedEnums. Sometimes enum names did not match the normal version of the data.
+
+* BugFix: Fixed an issue for "Graph Linked to External Private Object" error when deleting folders via the Content Browser when the folders include Actor Blueprints
+
+* BugFix: Fixed cyclic dependencies with Nativization that occur while native class construction happens. 
+
+* BugFix: Fixed a data-loss issue related to non-native Blueprint classes that override inherited component default values from a nativized Blueprint class hierarchy.
+
+* Changed ClearRenderTarget2D default alpha to 1, which is necessary for correct compositing.
+
+* The Graph Editors can now set their grid snap size via the Editor Preferences > Appearance.
+
+* Pre-Populated Blueprint nodes now attach all matching outputs to parent Blueprint calls (not just exec pins).
+
+* The CreateSaveGameObject node now alters its output type based on the SaveGameClass input.
+
+* IntVector can now be used in Blueprints.
+
+* Blueprint Context menu now has improved results for 'sel', 'sele', 'selec', 'bra', and 'foreach'.
+
+* Create Event node now displays the signature of the function that is desired.
+
+* Adding a component in the Blueprint Editor now adds the component as a child of the selected component.
+
+* Blueprint compact nodes no longer have two different looks. Now, all compact nodes have rounded edges and have a matte color.
+
+* GetReflectionVector math function has been exposed to Blueprints. 
+
+* Exposed a new function, GetCommandLine, to Blueprints which will return the command line used to launch the process.
+
+* You can now select Blueprint assets for nativization in the Project Settings > Packaging.
+
+* The Blueprint nativization feature will now be automatically enabled for the project when a Blueprint class is marked for nativization in the Blueprint Class settings view.
+
+* Interface class dependencies that are not enabled for nativization will now raise an error at cook time if dependent classes are marked for nativization.
+
+* We now hide 'EditInstanceOnly' properties when viewing/editing SCS component defaults in the Blueprint editor, as this does not apply to an instance context.
+
+**Core**
+
+* New: Added support for " in FName so they will be converted properly to \" internally.
+
+* New: Removed defunct "USING_SIGNED_CONTENT" define.
+
+* New: All platforms are now compiled with C++14 standard.
+
+* New: Allow enabling sound (-enablesound commandline switch) if the game disabled it by default.
+
+* Bugfix: Fixed a garbage collection crash when copying an empty timeline.
+
+* Bugfix: Module Manager shutdown fixes:
+
+    * Shutdown Module is now called in reverse order to when Startup Module is FINISHED. This allows modules to reference dependencies in their Startup Module to ensure they are loaded, and be sure they will still be around in Shutdown Module.
+
+    * HTTP Module now shuts down in Shutdown Module and not Pre Unload Callback.
+
+    * Added comments to Module headers to indicate this new change in behavior.
+
+    * Removed manual startup of HTTP module in Launch Engine Loop as it's no longer needed. This should save the module from being around if not really used by engine.
+
+* Bugfix: Many fixes to cooking complicated blueprints when s.EventDrivenLoaderEnabled is turned on.
+
+* Bugfix: Fixed a bug on 32-bit targets, some int64s could be aligned to 4 bytes and then passed into platform atomics functions. On apple platforms, those atomics functions can crash when passed in a non-8 byte aligned int64. Some explicit alignment has been added to those variables to 8 bytes.
+
+    * Added some validation of the class index in exportmap entries. Stops the engine crashing when corrupt package data is encountered.
+
+* Bugfix: Fixed Cancel Async Loading crashes when EDL was enabled.
+
+* Bugfix: Fixed memory stomps in SSL certificate initialization.
+
+* Bugfix: Fixed wrong reporting of ever increasing STAT_TextureMemory* and STAT_RenderTargetMemory* numbers in memory profiler.
+
+* Changed the Segment Dist To Segment Safe function in  FMath to handle the case where either (or both) of the input segments create points.
+
+* When CVar override is denied, current value of the CVar will be printed in the warning message for clarity.
+
+* Updated Generate Distill File Sets Commandlet so that it resolves string asset references after loading packages, to ensure that it finds all required files.
+
+* Updated Json Object Converter so that, by default, Struct To Json Attributes now skips transient properties.
+
+* Pak signing has been properly integrated with the new event driven loading system.
+
+* Added optimizations to the signature checking system so it is much faster.
+
+* AES encryption can now be configured for use with pak files, and has been integrated with the new event driven loader system.
+
+* bEncryptIniFiles=True can be added to your projects DefaultGame.ini to encrypt all ini files in the pak files.
+
+* Both AES encryption and pak signing are now configured from an Encryption.ini config hierarchy. See Config\DefaultEncryption.ini and Config\Windows\WindowsEncryption.ini in the ShooterGame sample project for an example of how to configure both encryption and pak signing.
+
+* Minor change to how bulk data is async loaded, allowing it to be freed correctly after first copy when requested.
+
+* Added some extra logging information to try and help catch a corrupt archive chunk crash seen in live builds.
+
+* Prevented reinstancing UEngine-derived objects on hot reload as this class is not reinstanceable.
+
+* Enabled full callstacks in non-monolithic builds by default.
+
+* Made Assemble Reference Token Stream thread safe for blueprints loaded on the async loading thread.
+
+* Pak platform file will now only be used if pak files exist regardless of command line params like -pak, -signedpak and -signed.
+
+* It is no longer necessary to rebuild the cooked exe to enable the Event Driven Loader. EDL and the new async IO are now controlled with ini settings only (through Project Settings). EDL can now be enabled in Blueprint only projects and binary distribution of the engine.
+
+* Event Driven Loader is now enabled by default.
+
+**Editor and Tools**
+
+* New: Added to Sequence Recorder:
+
+    * Optimize record transforms by setting all the keyframes at once.
+
+    * Option to toggle removing redundant keyframes from the recorded tracks. This is on by default.
+
+* New: Improved search for matching Material on FBX mesh import. Added a "Search Location" dropdown in the Material Import Options to control how the search will be performed. When an imported mesh references a Material, the importer will search for a Material with the same name in the assets folders. The search is performed in the following order based on the “Search Location” setting:
+
+    * "Local" option will behave as before with a non recursive search in the destination folder (e.g. in /Game/Content/Assets/MyImportFolder/ only).
+
+    * "Under Parent" option performs a "Local" search first followed by a recursive search from the destination folder and a recursive search from the parent's folder (e.g. in /Game/Content/Assets/MyImportFolder/ and subfolders then in /Game/Content/Assets/ and subfolders).
+
+    * "Under Root" option performs an "Under Parent" search first followed by a recursive search from the root folder (e.g. in /Game/ and subfolders).
+
+    * "All Assets" option performs an "Under Root" search first, followed by a recursive search of every asset folder (e.g. in / and subfolders).
+
+* New: Users are now able to customize the Class Browser/Picker to filter out developer folders as well as hide internal use classes via INI settings.
+
+    * ViewOptions button has been added to enable users to choose whether or not these filters are enabled. By default, internal only classes will be hidden and you will be limited to your own developer folder.
+
+    * Go to Project Settings >Class Viewer >Class Visibility Management to customize the default classes and paths used by the filter.
+
+* New: Added new sorting options for Statistics Window so that it now sort columns based on types and texture dimensions.
+
+* New: You can now export Static Mesh LODs from the editor. This requires that the Static Mesh be in the Level and use the FIle >  Export options.
+
+* New: The Gameplay Tags UI has been improved:
+
+    * In the Project Settings window, Gameplay Tags will show up as a tree instead of as a list.
+
+    * In the treeview, an Add Subtag button has been added to each tag that will add a new subtag in the same INI file.
+
+    * Comments can now be attached to gameplay tags, and will show up in the tooltip for a tag.
+
+    * Each tag has an options drop down menu that allows you to search for references to that tag. In the Project Settings window, the dropdown also has options to delete and rename tags.
+
+    * When creating a new gameplay tag, you can specify what INI file it will be created in.
+
+* New: Added option to import an FBX with front X axis.
+
+* New: The FBX exporter can now export Bone Scale.
+
+* New: It's now possible to shift-drag (Camera follows the object being moved) an object in the Level Viewport and move the 3 axis at the same time (drag the center of the manipulator).
+
+* New: The FBX importer is now able to import LOD groups with the Combine Mesh option enabled.
+
+* New: Added a shortcut for the Merge Actor functionality to right-click Context menu on Actors inside of the Level Viewport.
+
+* New: FBX imports now support using Multi-Sub-Object Materials generated by 3DS Max to prevent importing unused Materials.
+
+* New: FBX now takes into account all scene conversion options that were set at import.
+
+* New: Added support for detail customizations on root structs.
+
+* New: Added the ability to set and retrieve a host tab manager from the details view.
+
+* New: Added the ability to extend default layouts.
+
+* New: Added "Plane" to the Basic Shapes category of the Modes window.
+
+* New: Added the ability to create groups within groups in the Details panel customizations.
+
+* New: Added the ability for struct customizations to add categories to the parent object in the details customization.
+
+* New: Added support in MallocProfiler2 program to use nm.exe to symbolicate callstacks for clang-based platforms that have nm.exe that can look up symbols. It is easy to subclass the base support for each platform.
+
+* BugFix: Fixed a crash when importing Skeletal Mesh FBX file that contain one or more materials or textures with the same name as the FBX filename.
+
+* BugFix: Fixed a crash in the details customizations when the asset picker is used without a property.
+
+* BugFix: Fixed a regression where motion blur was enabled in editor views.
+
+* BugFix: Fixed a crash when importing a FBX with some node containing unused materials.
+
+* BugFix: Fixed a crash on Editor close after VR Foliage editing.
+
+* BugFix: Fixed a crash when using the Editor while Simulating (SIE) with an Actor copied to the clipboard.
+
+    * Affects Mac when navigating the Editor or simply updating the selection to an actor component.
+
+    * Affects Windows when showing a menu containing a Paste option (for example the Edit menu in the main menu).
+
+* BugFix: Fixed a crash in the Curve editors when scaling to fit after resetting the curve to default.
+
+* BugFix: Fixed a crash when importing an FBX that does not contain any mesh.
+
+* BugFix: Fixed a crash when importing a Skeletal Mesh with no Material.
+
+* BugFix: Fixed a crash occurring when trying to merge meshes/actors with their lightmap UV index set to 0.
+
+* BugFix: Fixed a crash in the Texture editor when opening a texture that does not have a valid texture resource.
+
+* BugFix: Fixed the Output Log filter that was not updating when text was committed.
+
+* BugFix: Fixed an issue that prevented skipping of intro movies using the Default Game Movie Player by using a keyboard or gamepad input.
+
+* BugFix: Fixed an issue where the "IDE downloading" would remain open if the Add Code window was closed while the toast was active.
+
+* BugFix: Fixed the functionality to drag a folder and its sub-directories from the OS file system into the Content Browser.
+
+* BugFix: Static Mesh Component's Overridden Lightmap Resolution previously displayed the value it was set to, even when bOverrideLightMapRes is false. The behavior is to use Lightmap Resolution on the Static Mesh member instead when bOverrideLightMapRes is false. The UI was adjusted to reflect the more accurate behavior.
+
+* BugFix: The Build Settings will now be set properly when reimporting with different options.
+
+* BugFix: FBX Reimport of an asset that originally imported from a file containing multiple objects will reimport only the FBX content that match the selected asset in the Content Browser.
+
+* BugFix: The FBX importer now imports local (parented) and global (non-parented) sockets correctly.
+
+* BugFix: Fixed an issue where AssetPtrs/Asset IDs were not getting properly fixed up when renaming assets in the Editor.
+
+* BugFix: Changed the maximum LOD level for merging actors to 7. Previously this was set to 8, which was invalid.
+
+* BugFix: Material settings are now completely disabled in the Merge Actor tab when Material Merging is disabled.
+
+* BugFix: Fixed a situation in advanced preview scenes where hiding the environment would not apply to reflections.
+
+* BugFix: Fixed shadows being faceted/hardened in advanced preview scenes.
+
+* BugFix: Fixed Material and Vertice count issues with Skeletal and Static Mesh conversions.
+
+* BugFix: Fixed Full Scene Importer changing Actor mobility when reimporting.
+
+* BugFix: Fixed an XML Parser assert when parsing a root tag that completes on the same line.
+
+* BugFix: Fixed hidden black callstack text in the Crash Reporter Client.
+
+* BugFix: Fixed an issue with compressed.ddp not containing all the shaders needed to load up the Editor for the first time.
+
+* BugFix: Fixed an issue when collapsing all, hiding everything in the settings editors.
+
+* BugFix: Fixed stretched Landscape Editor Icons.
+
+* BugFix: Fixed an issue when opening new assets in a minimized window not restoring that window.
+
+* BugFix: Fixed a number of Keybindings problems:
+
+    * Removed editor keybindings from Project Settings. It should not have been in there (already in Editor Preferences).
+
+    * Removed duplicate registration of editor keybindings from Editor Preferences.
+
+    * Fixed memory leak regenerating keybinding widgets when ending PIE world.
+
+    * Cleaned up styling a bit to make keybindings widgets clearer.
+
+* BugFix: Fixed the menus in the Editor not appearing in certain preview builds of Windows 10.
+
+* BugFix: Import Commandlet Fixes:
+
+    * Fixed a crash when the source control could not be contacted.
+
+    * Fixed assets not importing correctly if they depended on other assets in a previous import group within the automated import.
+
+* BugFix: Fixed an issue with UGS with Perforce having a branch root with a lowercase drive letter. UBT Intermediate response files include path names and are case sensitive, so it's important that they are always consistent. Windows seems to always include a drive letter in uppercase.
+
+* BugFix: Fixed an issue with UGS where editor .target file being deleted when doing a content-only sync, causing a prompt to rebuilt the editor when trying to launch.
+
+* BugFix: Fixed an issue with UGS to ensure the project config file remains valid even if sync is aborted due to files needing resolve. This prevents user configuration from not being applied onto build steps.
+
+* BugFix: Fixed material assignment error when importing multiple mesh with Combine Mesh option enabled in the Import Options.
+
+* BugFix: Fixed an issue where rebuilding geometry for the current level would actually erroneously build dynamic brushes from all levels.
+
+* BugFix: Fixed an issue where UBT was not correctly passing quoted strings as defines to the code build environment for PS4 and Mac targets.
+
+* Reorganized the Post Process settings so they appear as categories in the parent and so that they have better categories to make it clear what all the settings do.
+
+* Added ability to UGS to use Ctrl + A to select all in the log window.
+
+* Added ability to UGS for Perforce Settings window, allowing users to set optional values for tuning Perforce performance on unreliable connections.
+
+* Changed installed plugins to be disabled by default, but display a notification in the editor the first you load a project with a new one. Installed plugins which are new to this project will be adorned with a  "New!" badge in the plugin browser.
+
+* The cost of deleting hundreds of Actors simultaneously has been reduced by reorganizing how we check for Blueprint dependencies on the Actors.
+
+* Some of the plugin templates define Editor specific plugins. If created and a Standalone build is run, the application will attempt to link in editor libraries in game mode and will run into issues when you hit any key. The fix is to specify an Editor module description for these plugins.
+
+* The scrollbar will now remain at the bottom of the Output Log when a filter is specified.
+
+* Importing a Camera from an FBX will now create a Cinematic Camera with the specified focal length, field of view, and film back settings.
+
+* FBX Reimport now log messages in the message log window.
+
+* The FBX Reimport API now always returns false if there was some error when reimporting an asset.
+
+* FBX import now logs a message when there is too many bone influences on a vertex.
+
+* FBX files with multiple meshes with the same name is now supported.The importer will generate a number at the end of every mesh with a name clash.
+
+* Spline mesh actors in the Level Viewport can no longer be moved by transform gizmos in other viewports.
+
+* Clicking on a row in the Data Table editor before committing changes for the current row will now switch to that row after committing the changes.
+
+* Map properties that contain a struct as the value type are now editable when editing a component's properties in the asset editor.
+
+* The Color Picker will no longer stretch if a very long value is entered into a color component field.
+
+* The FBX Scene Reimport is now marking all reimported asset dirty.
+
+* The FBX importer is no longer using SkinXX workflow to import Static Meshes. It now uses the Material Name Slot workflow.
+
+* Changed how vertex colors are propagated when merging Spline Meshes using the Merge Actor functionality. This used to caused faceted vertex colors.
+
+* Now when painting, you won't be able to activate any other input other than the painting input.
+
+* Changed Panel Widget's Remove Child At function to not release slate resources if the child is a User Widget.
+
+* Near Clip Plane now adjusts to the user's scale to avoid clipping too early at small sizes.
+
+* Updated first include of Plugin Templates as Private PCH files were deleted.
+
+* Disabled "Convert to Static Mesh" option if only unsupported types are selected (e.g. Cable Components).
+
+* Level viewport 'show stats' option is now enabled by default, which avoids confusion with artists thinking lighting is built, when really the message is hidden.
+
+* Object's IsAsset function now returns false if the outermost package is RF_Transient as well. Also, updating the creation of the Transient Package to be RF_Transient. This makes it so transient packages created by sub-editors for things like previewing a streamed in level instance no longer show up in the Content Browser.
+
+* Made rebuilding lighting faster by importing lightmaps in multiple threads.
+
+* Added extra content to crash analytics and crash reports from the Editor.
+
+* Made an optimization to pasting multiple brushes, so that geometry is not constantly rebuilt for every brush that's added. This improves performance by a couple of orders of magnitude!
+
+* Editor creates a string representing user activity from the details of the currently active tab.
+
+* Editor user activity is sent as part of crash-related analytics and crash reports to give context.
+
+**Content Browser**
+
+* New: Added a new View Option to the Content Browser for "View Plugin Content". Originally Engine and Game plugins were only shown when clicking Show Plugin Content, but this made it difficult given the number of Engine Plugins. The new View Plugin Content option makes this easier to locate. In order to see the Engine plugins you'll need both Engine Content and Plugin Content checkboxes enabled.
+
+* New: Added 'play/stop' button to sound wave thumbnail in Content Browser when you hover over the asset.
+
+* New: Added a setting to allow the Sources Panel to be expanded by default.
+
+* BugFix: Fixed text properties exported to asset tags displaying in their NSLOCTEXT form in the asset tooltips.
+
+**Foliage**
+
+* New: New Foliage Tool Features:
+
+    * A new painting mode was added to add a single foliage instance at a time.
+
+    * Add the possibility to only paint in the current level.
+
+* New: Added the ability for Foliage LOD dithering to be disabled on a case-by-case basis by unchecking the "Dithered LOD Transition" checkbox in the foliage materials. Note that all materials used on all LODs need to have the checkbox checked in order for dithered LOD transitions to work.
+
+* BugFix: Fixed an issue where the Hierarchical Instanced Static Mesh Update Instance Transform Blueprint function was not updating bounds for instances that changed location.
+
+* BugFix: Fixed an issue with static lighting on Foliage and Instanced Static Meshes where the shadows on LOD levels other than LOD 0 was incorrect.
+
+    * All sub-LODs must have the same lightmap UV layout as the base LOD.
+
+    * There is currently a map check warning indicating this. This warning can be ignored if your sub-LOD lightmap UVs are set up appropriately.
+
+* LOD for Foliage and Hierarchical Instanced Static Mesh Components is now consistent with regular Static Meshes and uses a new screen size calculation.
+
+**Landscape**
+
+* New: Added an eyedropper to select Flatten Height in the Flatten tool.
+
+* BugFix: Fixed a crash when trying to edit Landscape with r.LightPropagationVolume=1 enabled.
+
+* BugFix: Fixed a crash when changing Landscape LOD settings with the LandscapeGrass "Use Landscape Lightmap" option enabled.
+
+* BugFix: Fixed a crash when opening the Persistent level with a streaming level that was hidden in the Levels window, and that streaming level had previously been duplicated outside the editor and contained a Landscape.
+
+* BugFix: Fixed an issue where levels were unable to be saved after Landscape components with tessellation were moved to another level.
+
+* BugFix: Fixed an issue where negative LODBias on Landscape components was not having any effect.
+
+**Material Editor**
+
+* New : Various Minor improvements:
+
+    * Named each Parameter uniquely either from copy/paste of any creation.
+
+    * Changed the default value for Scalar and Vector Parameter to 1 and (1, 1, 1, 1), respectively.
+
+    * Added a Promote To Parameter when clicking on an Input pin that will generate proper node type based on type pin type.
+
+    * When editing a color property update the Material Expression preview.
+
+**Persona**
+
+* Removing LODs from Skeletal Meshes can now be undone.
+
+**Sequencer**
+
+* New: Added current camera to Level Sequence Player Snapshot to allow any property to be displayed in the burn-in. Adjust Default Burn In to include a few more parameters like Focal Length and Focus Distance.
+
+* New: Added toggle to Movie Scene Capture to override Engine Scalability Settings to use Cinematic Scalability option.
+
+* New: Draw where the In and Out  points of the show section are, just like sub-sequences do. Change to only draw the green starting line if Start Offset is negative.
+
+* New: General Options button menu has been refactored into Play Options and Select Options.
+
+* New: Added Select Sections in Selection Range and Select All in Selection Range.
+
+* New: Added hotkey, Ctrl + L, to toggle Camera Cut track to lock/unlock the camera.
+
+* New: Added the ability to copy/paste tracks in Sequencer.
+
+* New: Sequencer now shows the Camera Name in the Cinematic viewport.
+
+* New: Added the option in Sequencer to lock the playback ranger per movie scene.
+
+* New: Added the option in the Curve Editor to show time in frame numbers.
+
+* New: Added Minimum handle size for time slider scrubber.
+
+* New: Added menu options in Sequencer to reduce keys of all sections in the current Level Sequence.
+
+* New: Exposed Curve Editor settings to Editor Preferences.
+
+* New: Added the ability to trigger Events with parameters.
+
+* New: Added hotkey to Sequencer for the ability to restore all states that have been mutated by Sequencer by using Ctrl +R. See "Restore Animated State" button on the toolbar.
+
+* New: Added the ability to override bindings when playing back a Level Sequence on a Level Sequence Actor.
+
+* New: In Sequencer, added the ability to step to beginning.
+
+* BugFix: Fixed a crash caused when attempting to stop a Level Sequence player while it is being evaluated.
+
+* BugFix: Fixed a crash when there are multiple components to attach. Added a component selector when there are multiple components that have sockets to attach to.
+BugFix: Fixed the offsets that were created when moving multiple sections. The offsets were being created because section bounds were being generated for all sections except for the current section. Instead, they should be computed for all sections except for any that aren't being moved.
+
+* BugFix: Fixed an issue with export from Master Sequence not finding bound objects.
+
+* BugFix: Fixed an issue with Sequence Recorder to not update the current sequence name if it's already set. This fixes a bug where if you pass in a sequence name to record to, it gets reset to the name in the sequence recorder settings.
+
+* BugFix: Fixed an issue with Cinematic Camera to allow manual focus distance when focus smoothing is enabled.
+
+* BugFix: Fixed issues to allow for component keyframing.
+
+* BugFix: When converting Matinee to Level Sequence, fixed matinee 3D scale track conversion to Level Sequence.
+
+* BugFix: Fixed an issue with Cancel when doing a Save As so that it doesn't continue to save.
+
+* BugFix: Fixed some recorded components not being generated.
+
+* BugFix: Updated Sequencer to properly fire events once per loop.
+
+* BugFix: Teleport simulating components when moving them through the transform track. This fixes bugs with recording simulating actors (ie. vehicle game) where recorded actors don't playback with the recorded positions and there are warnings about attempting to move a fully simulated Skeletal Mesh.
+
+* BugFix: Fixed the play rate track interaction between servers and clients.
+
+* In the Curve Editor, Set Tangent to user when flattening or straightening tangents only when the tangent mode is auto and the interp mode is set to cubic.
+
+* Set Snap Time to Dragged Key on by default.
+
+* Set the paste keys time to the current time, rather than the mouse time.
+
+* Remove click to rename shot functionality in the shot thumbnail. Added rename shot to the Shot Context menu.
+
+* When converting Matinee to Level Sequencer you can set the default scale when converting matinee move tracks to sequencer.
+
+* Audio Start Time in Sequencer is now deprecated in favor of Start Offset which is an offset into the audio clip.
+
+* In Sequencer, the Spawnables are no longer destroyed when the cursor leaves the master playback range.
+
+* In Sequencer, Remove click to rename shot functionality in the Shot Thumbnail. Added a Rename Shot to the Shot Context menu.
+
+**Source Control**
+
+* BugFix: Fixed source control to properly notify when files need to be checked out if a Blueprint node is dragged.
+
+* BugFix: Fixed an issue with saving a package that doesn't exist on disk but exists in P4 at a newer revision when the user chooses 'Mark Writable'.
+
+* The P4Changelist command line option is now respected by all Perforce operations that support it.
+
+* Subversion binaries should now support SASL authentication.
+
+* Integrated fix to support named changelists in SVN.
+
+* Git plugin: Multiline initial commit message and other connect screen cleanup.
+
+**Static Mesh Editor**
+
+* New: Static Mesh LODs can not be removed without having a reduction tool (such as Simplygon) configured for use in the editor. Use the "Remove LOD" button under each LOD section or adjust the “Number of LODs” slider and hit Apply.
+
+* New: Added the possibility to copy/paste Material list, Material Item, Section List and Section Item to the Static Mesh Editor and Animation Tools Editors.
+
+* BugFix: Fixed a crash when KDOP collision generation fails and there are existing collision meshes selected. We no longer clear out unrelated collision primitives when KDOP generation fails.
+
+* BugFix: Added a guard against rendering crashes when a mesh with no LOD resources is opened in the editor.
+
+* BugFix: Fixed issues with the Static Mesh sockets on reimport. Socket transforms are now updated correctly and deleted sockets will be removed on reimport.
+
+* BugFix: Undoing a change in the Mesh Preview Scene viewer will now correctly update the changes in the scene viewport.
+
+* LODs other than LOD0 can now have a screen sizes greater than 1.
+
+ **VR-Editor**
+
+* Fixed difficulty with double-clicking with Oculus Touch.
+
+**World Browser**
+
+* Disabled distance culling when rendering preview images for world composition. Otherwise images were missing culled objects.
+
+**World Outliner**
+
+* BugFix: Prevented a crash due to duplicate items in the World Outliner if Actors somehow end up attached to themselves.
+
+**Gameplay Framework**
+
+* New: Force Feedback Effects can now be previewed via the Content Browser.
+
+* New: Significance Manager plugin has been moved out of Experimental.
+
+* New: The Gameplay Abilities system has been moved into a plugin. It can be enabled in the Plugin settings, but it is still experimental.
+
+* New: Added the ability to override Crouch/Uncrouch status in Character Movement Component with Update Character State Before/After Movement.
+
+    * Fixed issue where uncrouching while on top of a physics object would cause penetration and incorrect behavior.
+
+* New: Added Set Auto Activate Blueprint function to Set Auto Activate, which is safe to call from Construction Scripts.
+
+* New: The Gameplay Tags Editor module has been moved to a plugin, while the Gameplay Tags module is now part of the core engine. System internals and performance have been significantly improved.
+
+    * Gameplay Tags can now be added directly from C++ code with AddNativeGameplayTag. Tags added this way do not need to exist in another ini or data table.
+
+    * There is now a Gameplay Tags Project Setting page that allows editing all shared tag settings and managing existing tags.
+
+    * It is now possible to Find Reference, Rename, and Delete tags from the tag management UI. Find Reference will only work on assets that have been resaved in 4.15.
+
+    * Much of the existing C++ API for Gameplay Tags has been deprecated and replaced with a simpler API using Matches/Has or MatchesExact/HasExact. Some broken utility functions were removed.
+
+* New: Added two "virtual buttons", Virtual Accept and Virtual Back that should be used for your menus if you ever need to run on a platform where the Face Button Right is accept, not Face Button Bottom (i.e. Switch). Engine and Epic game code has been updated to use them.
+
+* New: LOD distances have been consolidated between Skeletal, Static and HLOD meshes.
+
+* BugFix: Avoid a crash when calling a Character's Set Replicate Movement when not on the Server.
+
+* BugFix: Fixed a crash when travelling between levels if an Always Loaded sub-level was hidden. 
+
+* BugFix: The correct significance is now chosen when multiple viewpoints are in use.
+
+* BugFix: Tick functions that are disabled in a frame where they were scheduled to execute will now be prevented from executing.
+
+* BugFix: Blueprinted Level Sequence Actors will now correctly receive a tick event in their Blueprints.
+
+* BugFix: When modifying properties of an Actor during Simulate in Editor (SIE), Expose on Spawn properties will no longer be reset to class defaults.
+
+* BugFix: Player Camera is once again positioned correctly at the start of a level even if the game is immediately paused.
+
+* BugFix: Player Controllers are now kept in the same order when going through a seamless travel map transition, which fixes Get Player Controller calls.
+
+* BugFix: Fixed extra white space not being removed in front of console commands.
+
+* BugFix: Actors attached via the World Outliner to a Child Actor will now remain properly attached when making viewport changes.
+
+* BugFix: Fixed Game Mode Base's Get Game Session Class to return Game Session Class when set.
+
+* BugFix: Fixed an issue where Character Movement Component's "On Teleport" was improperly changing movement mode.
+
+* BugFix: Fixed an issue so that Volume Bounds no longer vanish after hot reloading custom trigger volume code.
+
+* BugFix: Fixed Project World To Screen Blueprint node to work for letterboxed viewports.
+
+* BugFix: Fixed an issue where Sound initiated via the Play/Spawn2D sound functions now correctly display subtitles.
+
+* BugFix: Setting the Max Draw Distance for a primitive component in a Blueprint Construction Script will now work as expected.
+
+* BugFix: Fixed an issue with Spline Mesh Components that would cause them to distort over long distances.
+
+* BugFix: Fixed the streaming pause rendering feature starting a movie even if a movie was already playing.
+
+* BugFix: Child Actors now properly get Begin Play called when spawned from the Begin Play of the Parent Actor.
+
+* BugFix: Fixed an issue that could result in duplicate Child Actors in uncooked Standalone games.
+
+* BugFix: Fixed Character Movement hanging on to a bad/penetrating floor check result and not continuing to check for a valid floor.
+
+    * This only occurred if bAlwaysCheckFloor was set to false. This could, in rare situations, cause the character to continue to attempt to depenetrate an object far away from it until another floor check occurred. To prevent this, we now force a floor check after the depenetration.
+
+* Bugfix: Keys can no longer be added to read-only curves in curve tables.
+
+* Bugfix: Actor Component, Scene Component, and Primitive Component functions are flagged as "Unsafe During Actor Construction" where appropriate.
+
+* Bugfix: The correct socket is now always used when calculating the new world location for attached components.
+
+* Bugfix: When an actor/component is selected that has a Child Actor, the components of that Child Actor are now properly highlighted with the selection indicator.
+
+* Bugfix: Subclasses of AInfo will no longer display a warning when running commandlets about having a natively added scene component.
+
+* Begin Play will now be called in a consistent order, regardless of whether they were placed in the Persistent Level, a streamed level, or dynamically spawned.
+
+* Components are now properly removed from their Owner when manipulated through edit inline new properties.
+
+* Clicking on components of a Child Actor in a viewport will now select the related actor/component.
+
+* Incremental unregister component time and batch size can now be edited from the Project Settings > Engine > Streaming.
+
+**Localization**
+
+* New : Line-break iterators will now avoid breaking words in Hangul.
+
+    * The default behavior for wrapping Hangul is to use Western-style wrapping (where words are kept as-is) rather than East Asian-style (where words are broken by syllables).
+
+    * This behavior can be controlled by the using the console variable "Localization.HangulTextWrappingMethod" in cases where the old behavior is needed, but since modern Hangul uses spaces, the per-word wrapping is preferred by native speakers.
+
+* New: Legacy-style translation collapsing is now optionally available at PO export time.
+
+    * This mimics the old behavior of the localization pipeline prior to 4.14 (where identical text within the same namespace would be collapsed into a single entry for translation).
+
+    * This can be enabled by setting "Localized Text Collapse Mode" to "ELocalizedTextCollapseMode::IdenticalNamespaceAndSource" (also available via the Localization Dashboard).
+
+* BugFix: Fixed a crash during IME shutdown.
+
+* BugFix: Fixed a selection bug in TSF based Microsoft IMEs on Windows 8+.
+
+* BugFix: Using FText's Format with an invalid argument no longer strips any associated argument modifier data from the resultant formatted text.
+
+* BugFix: Fixed an issue where IME composition could provide an invalid range if the text was changed while composing.
+
+* BugFix: Fixed localization gather including texts that were instanced or otherwise unchanged.
+
+* Updated User Defined Enum display names to use real FText instances so they can have stable keys.
+
+    * This avoids the issue where the FText display names were cached from an FString, resulting in them having a different identity each time they were re-cached, which lead to localization and deterministic cooking issues.
+
+* The package localization ID is no longer considered significant for PO export/import.
+
+    * This allows texts with identical identities in different packages to be translated as a single entity.
+
+    * If you were reliant on the 4.14 behavior, you can set "Localized Text Collapse Mode" to ELocalizedTextCollapseMode::IdenticalPackageIdTextIdAndSource (also available via the Localization Dashboard).
+
+**Networking**
+
+* New: MAX_ARRAY_SIZE and MAX_ARRAY_MEMORY from RepLayout are now user-configurable.
+
+* BugFix: Fixed an assert with "UPackageMapClient::AddNetFieldExportGroup" while viewing replays.
+
+* BugFix: Fixed edge cases/alternate IPv6 formats in IPAddressBSDIPv6's Set IP.
+
+* BugFix: Fixed OSSNull server/session filtering to better match Session Settings and online OSS. Made MCP, Steam, and Null LAN queries more consistent.
+
+* BugFix: Fixed an issue where movement would be disabled after possession of a pawn.
+
+* BugFix: Fixed a potential DDoS vulnerability.
+
+* BugFix: Fixed an issue that could cause replication conditions for properties to not be re-evaluated before saving those properties to replay checkpoints.
+
+**Other**
+
+* New: Media Assets: Added support for quadraphonic audio sinks.
+
+* New: Media Assets: Added Blueprint functions to query width, height, and aspect ratio of Media Texture instances.
+
+* New: Media Assets: Exposed "Can Play Source" in Blueprints.
+
+* New: Media Assets: Moved media player implementation into reusable class; also moved overlay text handling into separate asset.
+
+* New: Media Assets: Allow pausing while buffering media.
+
+* New: Wmf Media: Added support for Motion JPG (.mjpg).
+
+* New: Wmf Media: Added support for non-RGB32 uncompressed formats.
+
+* New: Wmf Media: Added rudimentary H.265 HEVC support for Windows 10.
+
+* New: Wmf Media: There is now a progress bar showing while the media is being resolved.
+
+* New: Added for video framebuffers with horizontal and vertical padding.
+
+* New: Added "On Get On Screen Messages” delegate to Core Delegates. This delegate can return severity-categorized text messages to be displayed on screen.
+
+* BugFix: Wmf Media: Fixed a crash when attempting to play or seek with an uninitialized Media Player.
+
+* BugFix: Wmf Media: Fixed an issue with RGB32-encoded AVIs rendering upside-down and too bright.
+
+* BugFix: Wmf Media: Reworked fallback stride calculations to fix issues with some video formats.
+
+* BugFix: Media Player Editor: Media player selection is ignored if media specifies player overrides.
+
+* BugFix: Media Player Editor: Fixed successful attempts at opening media reporting failure in some cases.
+
+* BugFix: Media Player Editor: Fixed an issue to show desired player name instead of current player name if no media is loaded.
+
+* BugFix: Media Player Editor: Media Player asset is properly marked as dirty when creating a sound wave or texture for it.
+
+* BugFix: Media: Fixed NV21 conversion shader scaling.
+
+* BugFix: Media: Fixed NV12 and NV21 support.
+
+* BugFix: Fixed platform Media Source so that they work correctly in the editor and in other non-cooked use cases.
+
+* Removed most game thread and render thread dependencies from Media Framework (more to follow in 4.16).
+
+* Optimized media handling of RGB input.
+
+* Moved media related enumeration types into separate header file, so they can be shared.
+
+* Continued effort to implement subtitle tracks.
+
+    * Media caption tracks were split into captions, subtitles, and generic text tracks.
+
+    * Caption Sinks were renamed to Overlay Sinks.
+
+* Prevented the user from nesting a platform Media Source in another platform Media Source. There are also runtime guards against this situation in case it's created through code.
+
+* Prevented loading of media plugins in console apps, such as game servers.
+
+* Added better error message for when no suitable media player plugin was found.
+
+* Media Player Editor: Removed property buttons from Platform Media Source customization.
+
+* Wmf Media: Disabled plugin on Windows 10 because it is currently broken.
+
+**Physics**
+
+* New: Added Hit Result output to Blueprint "Line Trace Component" function.
+
+* New: Added Debug Draw options (Trace Color, Trace Hit Color, and Draw Time) to shape traces for consistency with line traces.
+
+* New: Window properties are now editable at runtime on Wind Components.
+
+* BugFix: Fixed a crash when reimporting APEX destructibles from apb/x files caused by not allowing the renderer to flush destroy resource commands before emptying an array.
+
+* BugFix: Fixed a crash by removing collision between clothing and external Skeletal Mesh components, as clothing simulations could already be in flight and editing collisions while the simulation is running is not supported by APEX.
+
+* BugFix: Fixed a crash when generating clothing skinning data due to coplanar check triggering a check on small triangles.
+
+* BugFix: Fixed menu extenders in PhAT not being exposed correctly.
+
+* BugFix: Fixed edge cases in Body Instance where stale transforms were used when operations were run in Pre Physics.
+
+* BugFix: Fixed collision issues with Procedural Mesh Components that were created from Begin Play in cooked builds.
+
+* BugFix: Fixed non-uniform scaling of collision on meshes with rotated convex elements.
+
+* BugFix: Various thread-safety fixes to wind data. The components are tracked alongside their scene proxies so the game thread can safely access the wind data.
+
+* BugFix: Fixed welded bodies so that they now properly compute mass properties when attached bodies have a different physical material or mass override.
+
+* BugFix: Fixed an issue where collision shapes could be generated with multiple polygons making up a single flat face causing the body to jitter when resting on that face.
+
+* BugFix: Fixed an issue with line traces not working properly in editor worlds when physics substepping is enabled.
+
+* BugFix: Collision traces that start under a heightfields will now register properly with those heightfields.
+
+* BugFix: Fixed an issue where simulated bodies could have their transforms changed during simulations.
+
+    * This is an invalid operation as it means the physics simulation and the user's input will conflict. The old behavior was that we would pick one in an undefined way. This resulted in subtle timing bugs. To prevent this, we now ignore the transform of simulating bodies in this small window during simulation and will ignore the user's input.
+
+* Replicated movement has been disabled for Actors that are welded to another physics object.
+
+**Platforms**
+
+Platform-specific release notes for PS4, XBoxOne, and Switch are available in the forums and UDN for approved console developers.
+
+**General**
+
+* BugFix: Added Desktop Rect and Work Area to Monitor Info for desktop platforms. This is used to fix an issue in Scene Viewport where windowed fullscreen mode would be forced to the primary monitor, and incorrectly positioned on desktops where a taskbar/dock/etc. was on the left.
+
+* BugFix: Fixed support for the Optimize Code setting on Mac, HTML5, Android, Linux, iOS, and Xbox One. This controls optimization level in Debug Game configurations.
+
+* Changed how plugins will upgrade a project to be code-based. Instead of only upgrading some platforms, we now upgrade all to code-based. If only some are upgraded, the .uplugin file won't be staged, so the runtime won't know to disable the plugin that the project was trying to use. You can remove the platform from the whitelisted platforms in your .uproject's plugin list, instead of trying to disable the platform on a per-module basis inside the .uplugin file.
+
+**All Mobile**
+
+* New: Added "Optional Mobile Features Blueprints" module that exposes sound volume, battery level, determining if headphones are plugged in, and on Android the system temperature.
+
+* New: Added a module that provides a common interface for mobile platforms to access GPS/location data in Blueprints.
+
+* BugFix: Fixed an issue where lighting build would not stick during the first editor session for a newly-created Mobile project set up as Scalable 2D/3D (which disabled the Mobile HDR setting).
+
+* BugFix: Fixed an issue where the  Absolute World Position Material node with the "Excluding Shader Offset" flag set would return zero on ES2 devices.
+
+**Android**
+
+* New: Added support for MadCatz C.T.R.L.R. controllers.
+
+* New: Added the ability to set element value with new text parameter for UPL <setElement/>.
+
+* New: Added new Project Setting "Use ExternalFilesDir for UE4Game Files". This will place all files into ExternalFilesDir for automatic cleanup on uninstall for distribution packages, if enabled.
+
+* New: Added a first pass at improved virtual keyboard for Android. This is an experimental feature.
+
+    * Enabled by placing -NewKeyboard on the commandline, or setting IsUsingIntegratedKeyboard to a value greater than zero in AndroidPlatformTextField.cpp.
+
+    * Currently does not deal with Slate editboxes covered by the keyboard; make sure you place your controls in a visible area.
+
+* BugFix: Fixed a crash on 1st generation Kindle tablets if loading Oculus libraries as part of a project.
+
+* BugFix: Fixed a crash that would occur if you changed the orientation of your device during the Splash screen on Android when set to Sensor Orientation.
+
+* BugFix: Fixed an issue to allow Android console command sender to work with ADB version 1.0.36.
+
+* BugFix: Fixed an issue when deploying an Android apk over wifi using Launch On.
+
+* BugFix: Fixed IWYU issues for Location Services plugins.
+
+* BugFix: Fixed x86 and x86_64 libpng libraries for Android.
+
+* BugFix: Corrected memory stats for Android so that it now works like Linux.
+
+* BugFix: Fixed handling of ETC2 filtering by Google Play Store (If only ETC2, require ES3.0+, but doesn't use texture filtering manifest).
+
+* BugFix: Fixed #includes in deferred Android OpenGL.
+
+* BugFix: Fixed a driver bug and using separate context for MediaPlayer14, which corrects some render state problems with the movie playback.
+
+* BugFix: Fixed some issues with cooking when sending "On Target Platform Changed Supported Formats" when format changes in Android Project Settings.
+
+* BugFix: Fixed deploy packagename to use Target Name instead of App Name.
+
+* BugFix: Fixed some warnings from TcpMessaging that appeared in the Editor session's log window when launching on Android.
+
+* BugFix: Fixed issues where most textures appeared black when running in Vulkan mode and packaged with ETC1 textures.
+
+* BugFix: Fixed issues with Mac and Linux install and uninstall scripts if ANDROID_HOME not set.
+
+    * ES 3.0-capable devices which do not have floating point framebuffers and require the Mosaic HDR emulation mode are no longer limited to a resolution of 1024 pixels.
+
+* BugFix: Fixed the handling of numbers in input dialog for number fields, which now allow for sign and decimal.
+
+* Standardized the media track DisplayName in Android media player.
+
+* AndroidManifest now contains ES2.0 instead of 3.0 minimum when Multi and ETC2 and PVRTC selected and enabling all texture formats removes all restrictions.
+
+**HTML5**
+
+* New: We now support Emscripten 1.36.13 with SSE2 fixes.
+
+* New: We now support Cook-On-The-Fly for HTML5 builds.
+
+* Bugfix: 404 now sends the required double newline after headers.
+
+* 32-bit Mobile HDR shaders have been removed from HTML5 because all major browsers support floating point render targets. This change reduces memory usage and content download size.
+
+**iOS**
+
+* New: Added the console variable "rhi.Metal.SupportsIntermediateBackBuffer" to control whether to use an extra render-target so we can support screenshots & movie capture, or render directly to the back-buffers to save memory & GPU performance. Still defaults to ON for Mac & OFF for iOS/tvOS.
+
+* New: iOS Audio Improvements:
+
+    * Audio data is streamed from disk instead of memory when the bStreaming flag is set on the source. (The experimental audio streaming flag must be enabled in the Editor Preferences > Experimental).
+
+    * When not streaming from disk audio data is kept compressed in memory and decompressed on the fly as needed.
+
+* New: Support for iPhone 7 and iPhone 7 Plus has been added.
+
+* New: Exposed the bEnableRemoteNotificationsSupport ini setting in the iOS Project Settings.
+
+* BugFix: Fixed a crash on iOS Metal caused by trying to set texture-descriptor fields that aren't available on older versions of iOS. 
+
+* BugFix: Fixed a crash in FMallocBinned Get Allocation Size on iOS 9 that can happen when making use of Apple's CFNetwork Library.
+
+* BugFix: Fixed a crash in "FMetaldynamicRHI::RHIAsyncReallocationTexture2D" for PVRTC2 textures.
+
+* BugFix: Fixed missing initialization for the Metal Texture Cache in Avf Video Sampler.
+
+* BugFix: Fixed code signing for projects with iCloud enabled.
+
+* BugFix: Fixed "Use Full Precision" Material flag on iOS, which had no effect when using Metal shaders.
+
+* BugFix: Fixed many issues with iOS push notifications (Xcode project fixes, entitlements files, etc).
+
+* BugFix: Fixed an issue where the wrong Launch Image is displayed during startup on iPad.
+
+* Removed 32-bit MobileHDR shaders from iOS, because all supported devices support 64-bit render targets. This change reduces package size and memory usage. These shaders had already been removed from Android, as its shaders are patched on required devices at startup.
+
+**Linux**
+
+* New: It is now possible to launch your project on a remote Linux machine from the editor or UFE.
+
+* New: Updated bundled libc++ to 3.9.
+
+* New: DSOs can specify the mode (RTLD_GLOBAL/LOCAL) with which the engine will open them.
+
+* New: Added support for touch events (e.g. from a touchscreen).
+
+* New: Server will now crash the hung thread with SIGQUIT.
+
+* New: The engine will default to OpenGL 4.x code path on Linux unless -opengl3 is passed on command line.
+
+* New: All Linux targets will use binned malloc by default (previously some would use jemalloc).
+
+* New: Linux architecture (x86_64, ARM, etc) can now be set per project.
+
+* New: Added CEF support on Linux (x86_64 only).
+
+* New: Added prebuilt websockets library for Linux.
+
+* New: Added verbose printout whenever GL context's vblanc sync is being configured.
+
+* New: Linux remote devices added in editor's device manager will now be saved in Project Settings.
+
+* New: Linux touch events are filtered to prevent multiple OnTouchMoved events with the same location.
+
+* New: Added support for AArch64 (64-bit ARM) architecture.
+
+* New: Added build-essential to dependencies installed by Setup.sh so the engine can be built on Windows 10 Ubuntu subsystem.
+
+* BugFix: Fixed a crash when editing Skeletal Meshes on Linux.
+
+* BugFix: Fixed a crash on exit due to unsafe use of virtual functions in FRunnableThreadLinux.
+
+* BugFix: Fixed long-standing GL crash when using vertex attrib binding.
+
+* BugFix: Fixed a crash when opening Project Settings on Linux.
+
+* BugFix: The engine will not refuse to start as root on platforms other than x86_64 (e.g. ARM).
+
+* BugFix: Projects are now properly associated with the engine version or installation.
+
+* BugFix: Fixed GL4 rendering issues with tessellation.
+
+* BugFix: Fixed remote deploying of a packaged build to Linux.
+
+* BugFix: Fixed launch on a remote Linux machine.
+
+* BugFix: Fixed an issue for determining installed Linux cross-toolchain.
+
+* Bugfix: Use response files for ar (fixes cross-compilation failing due to too long command lines).
+
+* 4.15 will compile with clang 3.9.x with fewer suppressed warnings.
+
+* Disabled XGE when cross-compiling on Windows again (does not work well with clang 3.9).
+
+* Updated code which determines number of cores to produce the correct result on ARM boards.
+
+* Test configuration will drop core by default.
+
+* Removed hard-coded paths to mono binary and replaced it with a script.
+
+* "always_inline" is no longer applied to FORCEINLINE functions in debug builds.
+
+* Speed up bForceCompilationAtStartup=True when nothing changed.
+
+* Remove forced -windowed when launching from the editor.
+
+* Bootstrap script is now independent of the working directory and can be started via symlink.
+
+* Large mouse movement workaround will not be applied if the application has only one window.
+
+* SDL audio subsystem will no longer be initialized (it is not used anyway).
+
+**Mac**
+
+* New: Experimental Metal EDR/HDR output support for Macs with HDR displays running 10.12 or later.
+
+* New: Disabled OpenGL support on Mac. -opengl command line param is now ignored, we always use Metal. On old Macs that do not support Metal we show a message saying that the app requires Metal and exit.
+
+* New: Added Mac support for the experimental Alembic importer plugin.
+
+* New: Changed "OS X" to "macOS" in few places where we display it and updated the code that asks users to update to latest version to check for 10.12.2, as this version is strongly recommended.
+
+* New: Implemented Info string in Avf Media for display in Media Player Editor.
+
+* BugFix: Fixed a rare crash in CoreAudio on quit.
+
+* BugFix: Fixed a crash caused by an attempt to use Global Play World Actions in Play World Commands before it is initialized when showing suppressible warning dialogs in Editor Engine's Update Auto Load Project.
+
+* BugFix: Fix a potential crash due to unnecessary reinitialisation of the MetalRenderPipelineDesc mutex.
+
+* BugFix: Fixed an intermittent crash caused by trying to record statistics on an unregistered thread in Mac Metal.
+
+* BugFix: Fixed a crash opening Skeletal Mesh assets on Mac caused by a bug in SSkeletonTreeView.
+
+* BugFix: Setup GPUSkinCache so that it can be enabled on Mac Metal without crashing, though rendering errors persist.
+
+* BugFix: Ensure that only Mac Metal compiles shaders that use Volumetric & Vertex-Shader-Layer rendering features, allowing 3D color look-up textures for the tone mapping to be used on Mac.
+
+* BugFix: Avoid a suspected bug in Nvidia's Metal driver that caused flickering tiles on screen by using the older tiled-reflections algorithm on Metal for now.
+
+* BugFix: Fixed an issue with Metal that meant updates to G8_sRGB textures could corrupt the texture data.
+
+* BugFix: Fixed an error in vertex-declaration handling for Metal that could result in corrupt particle rendering.
+
+* BugFix: Fixed a Metal validation error caused by using a GizmoMaterial in-game which wasn't configuring the compositing parameters correctly.
+
+* BugFix: Fixed incorrect rect initialization in Mac Get Display Metrics.
+
+* BugFix: Fixed an issue with HTTP. If-None-Match response code on Mac and iOS.
+
+* Changed the minimum supported macOS version setting in Info.plists to 10.11.6.
+
+* Changed the Mac-specific high DPI code to use system backing Scale Factor only if "NS High Resolution Capable" in Info.plist is set to true. This solves the problem with macOS Sierra giving us actual backing Scale Factor values even in low DPI modes.
+
+* Decreased the number of max open files for a thread on Apple platforms from 256 to 192 to leave more file descriptors to Cocoa. This will hopefully fix some rare, random crashes in Cocoa.
+
+**Windows**
+
+* New: Windows XP support has been deprecated to allow the engine to use new Windows features that are not available when supporting Windows XP.
+
+    * The option to set the minimum Windows version to Windows XP in the editor has been removed.
+
+    * Code support for Windows XP will be removed in a future version.
+
+    * A deprecation warning will be emitted from Unreal Build Tool if building with Windows XP.
+
+* BugFix: Fixed an issue that caused devices in the Device Manager window to display the wrong platform.
+
+* BugFix: Fixed an issue with Windows Media Player that would initialize the movie player texture on first frame regardless of whether or not the decoder has a sample ready. This prevents a white texture from showing up for a frame.
+
+* Moved the MfMedia plugin out of an XboxOne directory, because it works in Windows using public APIs (i.e. does not need to be hidden in an NDA'd directory).
+
+* Put the AllDesktop cooking target back in to Windows File | Package menu.
+
+* Disabled the GameExplorer startup checks on Windows Server edition because of rare crashes, and it's not a useful check on Server edition anyway.
+
+**Programming**
+
+* New: Ensured that QAGameEditor can successfully be built as a monolithic executable.
+
+* New: Added error handling to symbol store uploading.
+
+* New: Added the command line option DisableUploadSymbols to disable symbol uploading.
+
+* Bugfix: Fixed custom pre-build and post-build steps being executed when generating project files.
+
+* Bugfix: Fixed invalid arguments being used to compile for .c files with Visual C++.
+
+* Removed direct dependency on LightPropagationVolumeRuntime plugin from DetailCustomizations module.
+
+* Removed duplicate dll names when preparing command line for link process.
+
+* GitDependencies no longer tries to clean working directory files that are excluded or ignored.
+
+* Compiler warnings for using undefined identifiers inside #if and #elif directives are now enabled.
+
+**Rendering**
+
+* New: Added support for saving 32-bit EXR images when saving HDR frame captures.
+
+* New: Added "Force 128-bit buffers for rendering pipeline" to the High-res Screenshot tool UI which forces full 128bpp rendering and EXR capture.
+
+    * This can be manually achieved with r.SceneColorFormat and r.PostProcessColorFormat.
+
+* New: Added support for texture streaming on Vulkan.
+
+* New: Vulkan now uses the glslang source from SDK 1.0.33.0.
+
+* New: Added support for global clip plane on Vulkan.
+
+* New: Changed some RHI Clear methods to Draw Clear Quad as some platforms can't support scissor/excluderect.
+
+* New: The View Matrices struct has been made immutable.
+
+    * Direct access to the members has been replaced with getters.
+
+* New: Enable or Disable the Temporal AA jitter during the Custom Depth pass is a project setting now.
+
+* New: Complemented the ENQUEUE_RENDER_COMMAND macro with a debuggable Lambda version: EnqueueUniqueRenderCommand.
+
+* New: First step in making the Engine more Friendly for Pipeline State Objects.
+
+    * We are now passing down 2 structs to collect state. Drawing Policy Render State collects state from the high level renderer down to Set Shared State, and Mesh Drawing Policy Override Settings collects state for Drawing Policy creation.
+
+* New: Improved memory usage and texture load order for texture streaming.
+
+* New: The texture streaming build can now be run through the resave package commandlet.
+
+* New: New way to synchronize with future and pending texture streaming requests.
+
+* New: Added "-FullscreenDisplay=N" commandline where 'N' will be the display index always used for fullscreen applications.
+
+* Bugfix: Fixed a crash when running the "RecompileShaders All" command.
+
+* Bugfix: Fixed a crash when using too many texture coordinate sets in a skeletal mesh material.
+
+* Bugfix: 16-bit images for all formats are no longer incorrectly flagged as sRGB on import. Existing content will be unaffected.
+
+* Bugfix: Fixed comments in Rendering Interface.
+
+* Bugfix: HLOD now respects the FreezeRendering command.
+
+* Bugfix: HLOD now correctly handles visibility across multiple views such as split-screen.
+
+* Bugfix: Fixed GPU skin cache getting disabled unless using recompute tangents. It is now always on.
+
+* Bugfix: Fixed missing velocity vectors due to depth fighting when using Pixel Depth Offset.
+
+* Bugfix: Fixed a case where rainbow-like artifacts would appear in the editor caused by uncleared post-process buffers.
+
+* Bugfix: Fixed incorrect cubemap size displayed in the texture editor.
+
+* Bugfix: Fixed fog opacity on alpha composited material.
+
+* Bugfix: Fixed a corrupt texture bug in the Download Image blueprint node, on some platforms.
+
+    * The new Download Image blueprint node returns a Texture 2D Dynamic, rather than a Texture 2D. It may be necessary to reauthor blueprint/code which uses this node to support the change in return type.
+
+* Exposed API for Material Expression Noise and Material Expression Vector Noise.
+
+* Re-enabled Temporal AA for high-res screenshots as long as the frame delay is active (r.HighResScreenshotDelay).
+
+* Improved error and exit codes when failing to find a good Vulkan driver/version.
+
+* Prevented Drop Detail settings from triggering when using fixed framerate.
+
+* The Film Tonemapper (r.TonemapperFilm=1) and DBuffer Decal (r.DBuffer=1) support have been enabled by default for new projects.
+
+    * These settings can be reverted via project settings and the related cvars.
+
+    * Bloom defaults have also been tweaked to match the new look.
+
+* Scene Capture 2D can specify a global clip plane, which is useful for portals.  This requires the global clip plane project setting to be enabled.
+
+* Removed Mac Graphics Switching plugin. OpenGL is not supported and Metal currently cannot do live GPU selection.
+
+**FX**
+
+* New: Exposed velocity scaling to vertex/surface particles, matching bone/socket particle behavior.
+
+* Now exporting Particle Module, Particle Module Color Base and Particle Module Color to allow users to inherit from them and create their own modules.
+
+**Lighting**
+
+* New: Added Inscattering Texture controls to Exponential Height Fog.
+
+* New: Added Fog Cutoff Distance, so artists can prevent fog on skyboxes (requires fog to be setup matching the fog that was rendered into the sky texture so that distant mountains match).
+
+* Bugfix: Fixed crash when using overly big lightmaps.
+
+* Bugfix: Fixed shadowmap channel assignment for stationary lights which are not in a lighting scenario level, when a lighting scenario level is present.
+
+* Bugfix: Fixed indirect lighting color on post process volumes.
+
+* Bugfix: Fixed contact shadow primitives thickness assumption from the depth buffer causing screen space shadow bleeding artifacts.
+
+* Bugfix: Prevented a large shadow depth bias due to low resolution causing near plane clipping.
+
+* Level packages are now marked dirty when manually forcing all reflection captures to update allowing the data to be saved correctly.
+
+* Dynamic Indirect Shadows for static meshes using distance fields.
+
+    * Individual Static Mesh assets can enable Generate Mesh Distance Field to compute a distance field, without the memory cost of enabling for the whole project.
+
+    * New Static Mesh Component properties Cast Dynamic Indirect Shadow and Dynamic Indirect Shadow Min Visibility.
+
+    * New World Settings property Dynamic Indirect Shadows Self Shadowing Intensity which replaces the cvar.
+
+    * Added Shadow Resolution Scale to Light Component.
+
+* Support for Reflection Capture intensity with simple reflections, which are the default with Forward Shading.
+
+* Volume lighting samples are allowed outside of the importance volume if their influence affects the volume. Fixes black indirect lighting on movable components in maps with small importance volumes. Volume lighting samples placed on surfaces use a radius that covers the layer height spacing, which prevents an uncovered region between layers.
+
+* When Inscattering Color Cubemap is specified, directional light inscattering is disabled.
+
+* Lerps between 1x1 mip at Non Directional Inscattering Color Distance to mip 0 at Fully Directional Inscattering Color Distance.
+
+* Improved deferred shadowing with MSAA by upsampling light attenuation intelligently in the base pass.
+
+* Forward renderer now supports shadowing from movable lights and light functions.
+
+    * Only 4 shadow casting movable or stationary lights can overlap at any point in space, otherwise the movable lights will lose their shadows and an on-screen message will be displayed.
+
+    * Light functions only work on shadow casting lights since they need a shadowmap channel to be assigned.
+
+* Shader complexity with Forward Shading more closely matches Deferred.
+
+**Materials**
+
+* New: Various requested material nodes have been added in this release.
+
+    * Input: Pre-skinned Local Normal.
+
+    * Math: Trigonometry, Fast (approximate) Trigonometry, Saturate, Round, Truncate.
+
+* New: Added Pseudo Volume Texture shader helper function for custom material nodes packing 3D textures or flipbooks into a 2D texture.
+
+* Bugfix: Fixed a rare crash when updating base properties of a material instance in an instance hierarchy.
+
+* Bugfix: Fixed potential crash when changing blend mode on a material instance.
+
+* Bugfix: Fixed incorrect shader complexity output for masked materials when enabling DBuffer decals.
+
+* Bugfix: Fixed low quality textures on grass.
+
+* Bugfix: Fixed POM function so that the world aligned mode works on slanted floors, by building a tangent basis using the vertex normal.
+
+* Bugfix: Fixed Actor Position, Pixel Normal WS material expressions for decals.
+
+* By default the material editor now shows the full instruction count for materials to avoid inconsistencies linked to how a material is updated.
+
+    * In general the numbers shown may be higher but this is not a performance regression, only reporting the actual figures used.
+
+* Added missing material expression tooltips/keywords for new nodes based on 4.15 preview feedback.
+
+**Mobile Rendering**
+
+* New: Added a project option to specify the maximum number of CSM shadow cascades that can be rendered by a movable light.
+
+* New: iOS devices that use Metal API will now use hardware sRGB texture sampling.
+
+* New: Android projects can now be packaged with Mali Graphics Debugger or Adreno Profiler support. The option can be found in Project Settings -> Android. This allows you to debug rendering issues on device using vendor specific tools.
+
+* New: Added mobile renderer support for single sample shadow from stationary lights.
+
+* Bugfix: Fixed rare bug that resulted in crash on some Android devices in the Set Uniform Buffer Parameter function.
+
+* Bugfix: Fixed bug which could cause 2D translucent elements to render as opaque on android devices which do not support fp16 rendering.
+
+* Bugfix: Fixed a bug which could prevent translucent primitives rendering when mobile HDR is disabled and modulated shadows are in use.
+
+* Bugfix: Fixed missing selection outline in Editor Mobile Preview.
+
+* Bugfix: Fixed an issue preventing GPU particles from being correctly rendered on compatible devices in both LDR and HDR modes.
+
+* Bugfix: Fixed inverted shadow culling when rendering CSM shadow depths and mobile HDR is disabled.
+
+* Bugfix: Fixed issue that could cause blank scene captures on devices that use 32 bit HDR encoding.
+
+* Bugfix: Fixed brightness issues with some UMG elements in Editor mobile preview.
+
+* Generating GLSL shaders for ES2 should be a bit shorter now.
+
+* Disabled dynamic buffer discarding on Adreno330 based devices. This was causing ~10ms stalls on UI rendering.
+
+* Splitscreen multiplayer should now correctly render all views in LDR and HDR modes.
+
+* Do not initialize Minimal Dummy Forward Lighting Resources for unneeded feature levels (below SM4).
+
+* Moved high-end mobile preview mode out from experimental, added iOS Metal, ES3.1 and Vulkan quality level preview options.
+
+**Optimizations**
+
+* New: Added texture streaming data in material for each used texture.
+
+* New: Optimized the async texture streaming task.
+
+* New: Iterative level processing to prevent stalls in the texture streaming manager.
+
+* New: New required texture resolution viewmode.
+
+* New: AMD cards on GCN architecture can now take advantage of the Depth Bounds Test optimization.
+
+* Precomputed shadowmaps no longer enable Force 2 To 1 Aspect, which is only needed for lightmaps, and saves some shadowmap memory.
+
+* Laptops with discrete AMD GPUs will now select the discrete video cards instead of the integrated.
+
+**Postprocessing**
+
+* New: Added new options for higher-quality upsampling when using screen percentage reduction. New modes are controlled by changing r.Upsample.Quality.
+
+* New: Better quality Temporal  AA:
+
+    * Fix for correct edge gradients.
+
+    * Image slightly sharper and less aliased.
+
+    * Fixed and enabled antighosting for dynamic objects.
+
+* Bugfix: Fixed output description extent for upscaling post process pass.
+
+* Bugfix: Fixed a case where scene buffer and viewport sizes mismatched on lower screen percentages causing post processing artifacts.
+
+**UI**
+
+* New: Adding support for Material Parameter Collections to Slate UI objects.
+
+* New: Added category info to VR Editor User Widgets to improve organization.
+
+* Bugfix: Fixed a potential race-condition where a Font object could be GC'd while the loading screen was using the font cache.
+
+* UWidget is now Blueprintable to allow users to create data only blueprints to adjust the CDO values of the basic UWidget types.
+
+**Slate**
+
+* New: Improved Canvas Panel render performance by improving batching behavior.
+
+    * Added option to place canvas panel children in same layer using explicit ZOrder setting.
+
+    * Previous behavior placing all children in separate layers is still default.
+
+    * Affects Constraint Canvas in Slate.
+
+    * Affects Canvas Panel in UMG.
+
+* Bugfix: Fixed a crash while updating cursor highlight.
+
+* Bugfix: Fixed a bug where Num Characters In Glyph was set incorrectly for TAB characters.
+
+* Bugfix: Fixed pixel jitter in Slate menus.
+
+    * Mostly seen in Mac version of the Editor.
+
+* SWindow now restores focus back to the widget that last had focus when the window was deactivated.
+
+* Made it so explicitly focusing a slate user that does not yet exist, creates the slate user so that the state is properly maintained in preparation for that user's arrival / input.
+
+* The Ignore Incoming Scale option for the scale box should now behave as expected in more cases.
+
+* Get Relative Layout Scale has been modified to also pass down the prepass scale, otherwise it can't extract out the incoming scale ahead of time before text is measured.
+
+* Added very verbose Slate stats behind a compiler flag to avoid the large overhead of these stats. To enable them you will need to set WITH_VERY_VERBOSE_SLATE_STATS to 1. In order to get indepth Slate stats (with a lot of overhead) you run the console commands:
+
+    * stat group enable slateverbose
+
+    * stat group enable slateveryverbose
+
+    * stat dumpave -root=stat_slate -num=120 -ms=0
+
+* Exposed a cached version of the widget geometry that comes in during Tick, this data has also been made available from blueprints through the Widget interface.  Also made an optimization pass on the class to make some space for the geometry object.
+
+**UMG**
+
+* New: UMG: Expose label browser for UMG.
+
+* New: Combo box widgets now support keyboard focus.
+
+* Bugfix: Fixed the calculation for widget components screen position if the camera aspect is constrained.
+
+* Bugfix: Fixed some HDPI mode problems with widget position calculation when projecting world to viewport / screen, and other absolute space related transforms.
+
+* Bugfix: Fixed a bug in the widget traversing code in sequencer; widgets inside of Named Slots should now be animatable.
+
+* Bugfix: Fixed some issues with HDPI mode in the widget designer.
+
+* Bugfix: Fixed the web browser widget crashing on Android when packaged for Distribution.
+
+* Widget Blueprints that are saved without a root widget will no longer add a Canvas Panel widget when opened in the editor.
+
+* The Widget Hierarchy will now remember which nodes have been collapsed when moving or renaming widgets.
+
+* Widget Blueprints will no longer experience compile issues when moving a widget from one Widget Hierarchy to another.
+
+* Dynamically changing the style for an editable text box or a multiline editable text box from Blueprints will now update it correctly.
+
+* Space and Enter will now fire OnClicked events for focused buttons.
+
+* Deleting widgets from the Widget Blueprint Hierarchy will no longer scroll the hierarchy view back to the top.
+
+* Widget Interaction Components now default to tick while the game is paused. Widget components that do not tick while the game is paused can no longer be interacted with when the game is paused.
+
+* Spin boxes can now be spun correctly when simulating touch events with the mouse.
+
+* When playing an animation, the first frame of the sequence is now immediately evaluated instead of waiting for the next tick.
+
+* Widget Component's Draw At Desired size property now should also work correctly if the widget is in screenspace.
+
+* We no longer also take the platform DPI scale (used for HDPI) into account when applying UMG's UI scale. Since UMG already provides a DPI scaling system, compounding it with the native OSes produces undesirable results, since the DPI scale curve does not take into account some unknown platform scale set by a user.
+
+* Restored the ability of the Widget Component to directly receive hardware input. The Widget Interaction Component is great for just about every interaction use case - the one it is not is when you actually want the 3D widgets to take focus, and to be able to be typed directly into by the user. The kind of situation where you might want to use them as a 3D menu, in a non-VR environment. By default - Widget Components will not behave in this manner, but you can now use the option Receive Hardware Input to enable the ability for Widget Components to function more like a widget in the screenspace of the viewport.
+
+**VR**
+
+* New: Updated the GoogleVR SDK to version 1.01.
+
+* New: Exposed "Is HMD Connected" to blueprint. This means that HMD hardware is ready to use. UE4 may or may not be using the HMD now.
+
+* Bugfix: Fixed rendering of Canvas tiles in stereo mode to be rendered for both eyes. This fixes the rendering of the background behind various stats shown by the STAT console command.
+
+* Bugfix: Fixed threading issue on Oculus Touch controllers, which could cause jittering in some circumstances.
+
+* Bugfix: Fixed an assertion when using Oculus stereo layers with MSAA, due to binding an unresolved depth surface.
+
+* New: Added analytics event to Vive HMD initialization.
+
+* New: Updating Oculus SDKs
+
+    * Oculus PC SDK 1.10.0
+
+    * Mobile SDK 1.0.3
+
+    * Audio SDK 1.0.2
+
+    * Platform SDK 1.10.0
+
+* Added a missing render target clear when using a VR hidden area mask and a post processing material.
+
+* Added in a new CVar vr.SteamVR.UsePostPresentHandoff, which defaults to 0. When set to 0, we do NOT use the SteamVR Post Present Handoff, which costs some performance GPU time. When 1, we use the call, and get some extra GPU performance. However, this call is NOT safe for scenes that have frame-behind GPU work, like Scene Capture components and Widget Components.
+
+    * For users that were using 4.14.1, and seeing a GPU timing improvement, that behavior is now off by default.  You can re-enable it by setting vr.SteamVR.UsePostPresentHandoff on your project to get it back.  Be aware that SceneCaptureComponents or WidgetComponents may cause async reprojection in SteamVR to fail in that case.
+
+**Programming Release Notes**
+
+**AI**
+
+**Navigation**
+
+* Added a sanity check to UNavigationSystem::AddElementToNavOctree to guard against DirtyElement.NavInterface being null.
+
+**Animation**
+
+**Skeletal Mesh**
+
+* New: Stopped SkeletalMeshTypes.h being globally included.
+
+* New: Refactored SkeletalMesh to use same color buffer class as StaticMesh.
+
+* New: Moved SkeletalMesh skin weights into their own vertex stream.
+
+    * If you were accessing skin weights in VertexBufferGPUSkin, you will now need to use SkinWeightVertexBuffer instead.
+
+* New: Moved SkeletalMeshComponent and SkinnedMeshComponent functions out of SkeletalMesh.cpp and into correct cpp files.
+
+* New: Moved SkeletalMeshActor code into its own cpp file.
+
+* Bugfix: Fixed bug with position offset in skeletal mesh vertex buffer when using high precision UVs.
+
+**Tools**
+
+* New: Removed hard-coded asset-to-animnode mapping, added SupportsAssetClass virtual instead.
+
+**Blueprints**
+
+* New: Added UK2Node::ConvertDeprecatedNode which can be used to convert deprecated native blueprint nodes into something else. Fixed several related issues with pin migration.
+
+* Bugfix: Fixed weak references stored in blueprint pins so that they are correctly detected when determining which assets reference other assets.
+
+* Included K2Node_BaseAsyncAction.h in K2Node_AsyncAction.h to fix a compile issue.
+
+* No longer enforcing that game project UFUNCTIONS() have a category.
+
+* Removed support for packaging editor-only Blueprint data in cooked assets (specifically removed all "[EditoronlyBP]" ini settings).
+
+* Fixed a crash that would occur when using the Blueprint Profiler with InputKeyEvents.
+
+**Core**
+
+* New: Added a warning when a thread-specific pak file reader is used on the wrong thread. This is unsafe, and it can lead to difficult-to-reproduce race conditions.
+
+* New: Added a default constructor to TTuple which value-initializes its elements.
+
+* New: Added PUSH_MACRO() and POP_MACRO() macros for preserving and restoring the state of a macro definition.
+
+* New: Added FCString::Strtoi64() for converting a string to a 64-bit signed integer.
+
+* New: Added a THasInserterOperator metafunction for testing if a left-shift operator exists between two types.
+
+* New: Added GET_ENUMERATOR_NAME_CHECKED for returning an enumerator name as a string while statically verifying that the enumerator name is a valid entry in the enum.
+
+* New: Added aliasing constructors to TSharedPtr and TSharedRef, which allow pointers to subobjects, keeping the outer object referenced.
+
+* New: Deprecated some core functions:
+
+    * Deprecated GetUniqueDeviceId. Use GetDeviceId now instead.
+
+    * Deprecated all platform independent usages of GetMacAddress and related functions. It was not reliable the way it was implemented on PC, and generally a bad idea.
+
+    * Deprecated GetMachineId. Use GetLoginId now instead. That's what it really was.
+
+* New: Added debugger visualizer support for delegates.
+
+* New: Added an error report when deserializing a bool which isn't 0 or 1, which means archive corruption.
+
+* New: Added UEnumProperty which supports non-byte enum classes.
+
+    * Enum class properties will now be UEnumProperties instead of UEnums.
+
+* New: You can now hook delegates to be informed when new classes have been added or replaced via hotreload.
+
+* New: TArray now has InsertDefaulted to complement AddDefaulted.
+
+* New: SearchableNameMap has been added to the Linker and Asset Registry, and is used to find Gameplay Tag and Data Table Row references inside assets.
+
+    * If you have your own system of "special names" you can add those names to this map with MarkSearchableName
+
+* New: Added FGCObjectScopeGuard and TStrongObjectPtr as a convenient way to keep a UObject alive without having to add it to the root-set
+
+    * Both use FGCObject internally to reference the object and keep it alive.
+
+    * FGCObjectScopeGuard is designed to be lean and used as a guard for an existing pointer, whereas TStrongObjectPtr is more "full-fat" and designed to be a replacement for a raw-pointer.
+
+    * Of the two, you should prefer FGCObjectScopeGuard where possible. Also note that TStrongObjectPtr isn't supported by UHT/UPROPERTY as you should just use a raw-pointer in that case (it would do the same thing).
+
+* Bugfix: Fixed Algo::IsSorted() to work with C arrays and initializer_lists.
+
+* Bugfix: Fixed MallocStomp to always allocate at least one byte instead of returning null.
+
+* Bugfix: Fixed TSparseArray::Reserve so that new elements get allocated earlier in memory instead of later, giving the more expected iteration order.
+
+* Bugfix: Fixed some bad memory allocation semantics (e.g. mixing FMemory::Malloc and delete).
+
+* Bugfix: Fixed the TIsZeroConstructType specialization for TScriptDelegate.
+
+* Bugfix: Fixed UObject::IsA crash when reinstancing blueprints.
+
+* Bugfix: Fixed typo in FMallocLeakDetection.
+
+* Bugfix: Fixed FArchive::bCustomVersionsAreReset not being copied when the archive is copied.
+
+* Bugfix: Fixed race condition crash when serializing UClasses.
+
+* Bugfix: Fixed it so calling CreateDefaultSubobject with the bTransient=true flag correctly sets the Transient Object Flag on the created object.
+
+* Bugfix: Fixed several cooking issues with editor-only components and sub objects.
+
+* Bugfix: Fixed TSparseArray debugger visualization of types smaller than 8 bytes.
+
+* Bugfix: Fixed parsing of standalone macros like PRAGMA_DISABLE_DEPRECATION_WARNINGS in UnrealHeaderTool.
+
+* Bugfix: Fixed a memory leak in the task graph.
+
+* Bugfix: Fixed a regression causing an assert to fire in the editor related to UGCObjectReferencer.
+
+* Bugfix: Fixed a very old bug with lock free lists, probably not relevant to anything we currently use.
+
+* Bugfix: Fixed return value of Ar.Tell() being assigned to int32 instead of int64.
+
+* Bugfix: Added a thread safety fix for UBlueprintGeneratedClass::PostLoadDefaultObject while UBlueprintGeneratedClass::SerializeDefaultObject runs on the async loading thread.
+
+* Clarified the error message emitted when binding a delegate to a pending kill object.
+
+* Deprecated TAutoPointer.
+
+* Deprecated TScopedPointer.
+
+* Deprecated TAutoPtr.
+
+* Made TFunction::Unset protected. This is an implementation-only function which doesn't properly destroy the binding (can leak memory).
+
+* Deprecated the use of TEnumAsByte with enum classes.
+
+* Improved comments and added useful utility functions to FAssetPtr, TAssetSubclassOf, and TSubclassOf.
+
+* Shortened LexicalConversion namespace to Lex.
+
+    * LexicalConversion is now an alias for namespace Lex, so cannot have overloads added to it.
+
+* Added some extra safety to FEvent by proxying them avoid bugs with stale events.
+
+* Made FindShaderResourceById and FindShaderById return a raw pointer instead of TRefCountPtr to prevent creating temporary TRefCountPtrs. It is no longer necessary to use TRefCountPtrs as shader serialization has changed and shaders are no longer registered on the async loading thread.
+
+* Added multiple thread safety checks to async loading code.
+
+* Thread Heartbeat will no longer report the same hang multiple times.
+
+* Async Loading code will now detach the linker when resetting async package loader to avoid situations when loading the same asset multiple times results in the following load request finding the old linker after the package has been loading but the async package hasn't been deleted yet (async package for the old request in limbo state but linker exists).
+
+* ModuleManager will now use a critical section instead of FMultiReaderSingleWriterGT to make it thread safe.
+
+    * Removed FMultiReaderSingleWriterGT because it was not fully thread safe.
+
+* Async Loading stability improvements:
+
+    * Made sure all package dependencies are loaded before post loading its objects
+
+    * Added more Linker detach safety checks
+
+**Editor and Tools**
+
+* New: Added extension point for level editor viewport Show menu.
+
+* New: Added support for "Show Only Modified Properties" filter to DetailWidgetRow.
+
+* New: FPropertyRestriction has been restructured to differentiate between hidden and disabled properties.
+
+    * AddValue and IsValueRestricted no longer exist, use either AddHiddenValue and IsValueHidden, or AddDisabledValue and IsValueDisabled as appropriate.
+
+* Bugfix: Fixed several issues with using the SettingsEditor to modify new module-specific ini files.
+
+**Gameplay Framework**
+
+* New: AActor::GetComponents can now optionally choose to include components in Actors created via its ChildActorComponents.
+
+* New: Systems can now register console autocompletions without modifying an ini by registering a delegate with UConsole::RegisterConsoleAutoCompleteEntries.
+
+* New: UActorComponent subclasses can now override IsReadyForOwnerToAutoDestroy() to indicate to the owning Actor whether they can begin auto destroying if desired.
+
+* Bugfix: Fixed a crash when player is destroyed and the server checks to see if it needs to force a network update.
+
+* Bugfix: Fixed a crash that could occur when loading streaming levels via the PrepareMapChange and CommitMapChange functions.
+
+* AActor::BeginPlay virtual function is now protected and has a dispatching function instead.
+
+* Adding a new volume class via hotreload now allows that volume to be properly added to the world. Volume classes that are modified by hotreload will remain visualized.
+
+* Added IsPaused to GameModeBase now that the pausers array is protected.
+
+* AActor's GetAttachParentActor, GetAttachParentSocketName, and GetAttachedActors functions are no longer inappropriately marked virtual.
+
+    * Any class overriding GetAttachParentActor, GetAttachParentSocketName, or GetAttachedActors will need to be refactored to accomplish the intended functionality.
+
+* FAttenuationSettings has been renamed FSoundAttenuationSettings.
+
+* Optimized and cleaned up some Actor methods related to location and rotation.
+
+    * Inlined GetActorForwardVector(), GetActorUpVector(), GetActorRightVector(). Wrapped them to simply call the methods on USceneComponent rather than using a different approach to computing these vectors.
+
+    * Inlined blueprint versions: K2_GetActorLocation(), K2_GetActorRotation(), K2_GetRootComponent().
+
+    * Cleaned up template methods that are used to delay compilation of USceneComponent calls to make them private and prefix "Template" to their names so they don't show up in autocomplete for calls to the public methods.
+
+**Networking**
+
+* Bugfix: Fixed overflow issues in SerializeProperties_DynamicArray_r. Also, fixed a crash from not ensuring properties were serialized successfully.
+
+* Added an experimental feature for game clients that are recording network replays to do the per-frame recording work on a task thread in parallel with the Slate tick. If Slate is a significant portion of frame time, and client replays are being recorded, this can save a lot of game thread time each frame.
+
+    * Use the console variables tick.DoAsyncEndOfFrameTasks and demo.ClientRecordAsyncEndOfFrame to enable this feature.
+
+**Other**
+
+* New: CrashReportProcess crash report backend update:
+
+    * InvalidCrashReports now saved to S3 instead of local folder.
+
+    * Removed option to sync MinidumpDiagnostics from Perforce.
+
+    * Moved MinidumpDiagnostics from old Perforce synched location to its own folder.
+
+    * Added improved logging to Slack with option to monitor MDD performance.
+
+    * Added hourly log folder bins to MDD logs to avoid bloated folders and associated performance issues.
+
+    * Added support for types of crashes we don't want to symbolicate (using it to skip callstack gen for hang detected ensures).
+
+    * CRP now has improved handling of lost comms with website/DB.
+
+    * Crash reports will not be lost.
+
+    * Added compression support for crash files output from CRP to S3.
+
+* Improved MinidumpDiagnostics performance for use on the CR server by CrashReportProcess.
+
+**Physics**
+
+* Bugfix: Fixed FVehicleAnimInstanceProxy::PreUpdate not calling FAnimInstanceProxy::PreUpdate.
+
+* Added sanitization code around PhysX flags and refactored the ways flags are managed through a single code path.
+
+**Platforms**
+
+* New: Both RTTI and Exceptions can be enabled. The "Use RTTI" and “Enable Exceptions” flags, respectively, must be set to true in the Module.Build.cs file and in the build configuration file as well.
+
+* Made a generalized "growable heap" allocator that had previously been copied and pasted between RHIs. This is a heap allocator that allocates in large chunks, and then suballocates from within (often with large alignment). Mainly meant for GPU resource memory allocations, but can be used in other circumstances. To use it, you simply implement four easy functions (for creating the memory and other reasons). Allows for more than just a raw pointer to be the result of an allocation.
+
+**HTML5**
+
+* New: Added HTML5 thirdparty libraries build scripts in python.
+
+* New: Updated HTML5 thirdparty libraries (bash) build scripts to make use of new toolchain version.
+
+* Bugfix: Fixed PhysX for HTML5 crash. Removed 32-bit sized unions when browser floats are 64-bits.
+
+**iOS**
+
+* New: Added an iOS-specific implementation of the LocationServices plugin.
+
+**Linux**
+
+* New: Added helper build scripts.
+
+    * Pull source and compile: zlib, openssl, libcurl, and libwebsockets.
+
+    * Dockerfile to setup "CentOS 7" environment (GLibC 2.17 and GCC 4.8.5)
+
+**Mac**
+
+* New: Extended the internal Metal debugging options for MetalRHI:
+
+    * Added rhi.Metal.BufferScribble, which will fill freed buffer regions with 0xCD to help identify any areas where we are writing to a buffer while it is still being processed on the GPU.
+
+    * Added rhi.Metal.BufferZeroFill, which will zero-fill newly allocated buffer regions before any other data is read/written. Useful for catching cases where we might be reading uninitialised memory.
+
+    * Added rhi.Metal.ResourcePurgeOnDelete, which will purge the backing store of resources prior to releasing them back to the system or the appropriate pool.
+
+    * Added rhi.Metal.ResourceDeferDeleteNumFrames to defer releasing resources to the system or the resource pool by the specified number of frames (in addition to the current policy of waiting for the current end of frame & command-buffer completion). Useful for tracking down resource lifetime errors.
+
+    * Fixed the flickering geometry when enabling rhi.Metal.RuntimeDebugLevel=4 which breaks render passes into separate command-buffers - the occlusion query was waiting on the wrong command buffer in this case.
+
+* New: Extended the Metal state recording when rhi.Metal.RuntimeDebugLevel is set to 3 or greater to include the objects and states used.
+
+    * The object labels, types, and descriptions will be printed on failure. If the object is deleted prior to this, we have a lifetime error and it will crash at this point and can be debugged further using our -metalretainrefs command-line option or Xcode's zombie-objects.
+
+* Extended FShaderResource's OptionalData field to allow arbitrary amounts of data so that Metal shader debug information can be stored in it.
+
+* Redundant state changes in MetalRHI are optimized out.
+
+* Changed FAppleHttpResponse::GetContentLength to return expected content size instead of payload size, so it is consistent with other implementations.
+
+* Do not try to use installed Mono 4.6 on Mac as it is known to have issues on macOS 10.12 (for example building the editor with UBT often fails with Mono running out of file descriptors).
+
+* Added -Wno-undefined-var-template on Mac to work around an issue with compiling UHT in Xcode 8.3.
+
+**Programming**
+
+* New: Added list of Dependant modules to EULA check so that it's clear what can't be included.
+
+* Removed non-cyclic actions from messages about cyclic actions being present.
+
+* Removed CopyVisualizers.bat now that UE4.natvis is included in the VS2015 solution and we're no longer supporting VS2013.
+
+* Moved Sample game projects into a different solution folder, similar to Templates.
+
+**Rendering**
+
+* New: Added r.Vulkan.SubmitOnDispatch to help debug compute shader hangs on Vulkan.
+
+* New: Added GDynamicRHI::GetName() to get the name of the current RHI.
+
+* New: Locked buffers are now aligned to 16 bytes to allow SIMD operations on Vulkan.
+
+* New: Reworked Vulkan render pass system.
+
+* Bugfix: Fixed a bug with default parameters on shaders for non-HLSL platforms.
+
+* Bugfix: Fixed a memory leak when using VisualizeTexture on some RHIs.
+
+* Bugfix: Fixed a typo in check relating to stalling the RHI thread.
+
+* Bugfix: Fixed a rare crash with particle systems destroying render state.
+
+**Lighting**
+
+* Bugfix: Fixed a rare bug with shadow groups rendering things that have not been setup to render this frame.
+
+**Optimizations**
+
+* New: Glsl source for Vulkan shaders is now stripped out to save memory. Use r.Vulkan.StripGlsl=0 to not remove it at runtime.
+
+* Projected shadows are no longer created when r.ShadowQuality is 0
+
+    * Fixes crash in the forward path trying to render shadows
+
+    * In the deferred path, the shadowmap was still being rendered and only the projection skipped, now all cost will be skipped
+
+**VR**
+
+* The relative transform of the UCameraComponent will no longer reset if the HMDDevice returns a new position. This avoids changing the entire hierarchy twice during every update.
+
+**Mass debug rendering**
+
+* Switched fudged auto-LOD calculations to use a power decay instead of a reciprocal.
+
+**Upgrade Notes**
+
+**Animation**
+
+**Pose snapshot tool**
+
+* CVar forcing multithreaded update is now defaulted to off. Projects now enable it by default, but can easily opt-out.
+
+* The setting to enable and disable multithreaded update has moved from the class defaults to the class settings. Any of these flags that were previously set in the class defaults will need to be re-set in the class settings.
+
+* This will change the behavior of morphtarget if SetMorphTarget and animation curves were being used together. 
+
+**Skeletal Mesh**
+
+* If you were using types defined in SkeletalMeshTypes.h, you may now need to include this file explicitly.
+
+**Blueprints**
+
+* If your project is missing "[EditoronlyBP]" ini settings, or if you've set any of those to false, then you need to upgrade code where you use Blueprint editor-only data in your packaged project.
+
+    * Places where you use UBlueprint objects at runtime will need to be updated. Most commonly, ConstructorHelpers::FObjectFinder<UBlueprint> will need to be updated to ConstructorHelpers::FClassFinder<UObject>.
+
+**Core**
+
+* Deprecated some core functions:
+
+    * GetUniqueDeviceId is deprecated; replace with GetDeviceId.
+
+    * GetMachineId is deprecated; replace with GetLoginId.
+
+    * All platform-independent usages of GetMacAddress and related functions have been deprecated. These functions were not considered reliable on all platforms.
+
+* Any types which inherit TAutoPointer should replace it by implementing their own conversion operators. TAutoWeakObjectPtr implicit conversions should be replaced with explicit Get calls.
+
+* TScopedPointer and TAutoPtr usage should be replaced with TUniquePtr.
+
+* TFunction::Unset usage should be replaced with an assignment to nullptr.
+
+* Any code which finds an enum class property as a UByteProperty should be changed to look for a UEnumProperty instead. Any code which depends on UByteProperty::Enum or UByteProperty::GetIntPropertyEnum should support UEnumProperty as well.
+
+* Usage of TEnumAsByte with an enum class type should be changed to just use the enum class type directly.
+
+* Use of the LexicalConversion namespace should be changed to Lex.
+
+**Gameplay Framework**
+
+* GameplayAbilities are no longer enabled with the "GameplayAbilitiesEditorEnabled" ini setting. There is now a GameplayAbilities plugin that can be activated for this purpose.
+
+* Activate and SetActive are no longer valid to call on components from construction scripts. Call SetAutoActivate instead.
+
+* For LOD distance consolidation, LOD distances should be converted automatically on load. If you are using update rate optimizations with custom thresholds for skeletal meshes you will need to update these manually.
+
+* The Blueprint API for GameplayTags has changed, but existing nodes should update themselves correctly. There is now an ExactMatch parameter to replace the older enums, so you may need to update that bool in some cases.
+
+    * If you have existing GameplayTag data in Data Tables or .ini files, you will have to migrate them.
+
+    * Searching for GameplayTag usage will only work properly once assets have been resaved in version 4.15 or later.
+
+* Uses of FAttenuationSettings need to be changed to FSoundAttenuationSettings.
+
+* Delete DDC or otherwise change Spline Mesh Component parameters to ensure they get regenerated.
+
+**Physics**
+
+* Many deprecated collision functions have now been removed from KismetSystemLibrary, and the _NEW postfix has been removed.
+
+* Cooked PhysX convex element data is now in "body space" instead of “element space”. If you are directly accessing PhysX convex mesh data, you may need to take that into account.
+
+* Any simulated body with attached kinematic bodies that have different physical materials, or mass overrides, will now have a correct mass distribution.
+
+    * The old behavior was to use the physical material of the simulated body. This may lead to re-tuning some content. The most likely case is any vehicle with attached kinematics that had special mass properties or mass overrides. As a result of this change, collision boxes with zero mass attached to vehicles will behave correctly.
+
+**Platforms**
+
+**Windows**
+
+* It is currently possible to set the minimum Windows version to Windows XP directly in the INI files.
+
+**Rendering**
+
+* There is now an error when using too many texture coordinates on a skeletal material that displays the following output: "Too many texture coordinate sets defined on GPUSkin vertex input". Previously this case would silently fail and crash when the material was used, now it's correctly blocked. The fix for this error is to uncheck "Used with Skeletal Meshes" on the affected material as this is unsupported.
+
+* Some of the variables have been moved from the SceneView and renamed to be more explicit. Getters that do a lot of computation have been renamed to Compute*
+
+* Prefer EnqueueUniqueRenderCommand over the macro as we want to switch to those on the future
+
+* Remember to take a copy of FDrawingPolicyRenderState before you hand it off to another thread.
+
+**Slate**
+
+* Improved Canvas Panel render performance by improving batching behavior.
+
+    * Default behavior has not changed.
+
+    * To take advantage of better render performance enable project setting "Engine->Slate Settings->Constraint Canvas->Explicit Canvas Child ZOrder".
+
+    * Recommended for projects using Canvas Panel on mobile platforms.
+
+    * Projects enabling "Explicit Canvas Child ZOrder" must set ZOrder on all Canvas Panel (SConstraintCanvas) child slots to avoid undesirable clipping.
+
+    * Projects enabling "Explicit Canvas Child ZOrder" can place children that will never overlap in the same layer by setting the same ZOrder value to improve batching.
+
+**Google Daydream**
+
+* Daydream apps require android-24.  This is not installed by default, but can be manually installed with NVPACK\Chooser.exe.
+
