Copying //UE4/Dev-Documentation to //UE4/Samples-Main (Source: //UE4/Dev-Documentation @ 3372893)

diff --git a/Engine/Documentation/Source/Support/Builds/ReleaseNotes/4_15/Updates/UnrealEngine4_15UpdateNotes.INT.udn b/Engine/Documentation/Source/Support/Builds/ReleaseNotes/4_15/Updates/UnrealEngine4_15UpdateNotes.INT.udn
new file mode 100644
--- /dev/null
+++ b/Engine/Documentation/Source/Support/Builds/ReleaseNotes/4_15/Updates/UnrealEngine4_15UpdateNotes.INT.udn
@@ -0,0 +1,117 @@
+Availability:Public
+Title: Unreal Engine 4.15 Update Notes
+Description:
+Crumbs:
+Parent:Support/Builds/ReleaseNotes/4_15
+Order:0
+
+## 4.15.0 Release
+
+Released: 02/15/2017
+
+Github: [https://github.com/EpicGames/UnrealEngine/releases/tag/4.15.0-release](https://github.com/EpicGames/UnrealEngine/releases/tag/4.15.0-release) 
+
+Release Notes: [](Support/Builds/ReleaseNotes/4_15)  
+
+## 4.15.1 Hotfix
+
+Released: 03/20/2017
+
+Github: [https://github.com/EpicGames/UnrealEngine/releases/tag/4.15.1-release](https://github.com/EpicGames/UnrealEngine/releases/tag/4.15.1-release) 
+
+Fixes:
+
+* UE-42385 Crash adding an AnimComposite to itself
+
+* UE-41929 Crash on reimport animation when there is an additive animation track and fbx frame number is different
+
+* UE-41921 Crash Re-Importing animation with Additive Layer Tracks after adding a bone to the hierarchy
+
+* UE-41817 Crash when manipulating Bone Space IK before setting EffectorSpaceBoneName
+
+* UE-40025 Inconsistant Line Trace behavior
+
+* UE-42756 UnrealVS not installing correctly for VS2017
+
+* UE-42321 Unable to launch project from Visual Studio using installed engine build
+
+* UE-42411 Resource files are not recompiled when Version.h changes
+
+* UE-42205 Hot reload reported as failed in Visual Studio for projects opened in the Editor after having hot reloads performed successfully previously
+
+* UE-42507 check(PIEInstanceID != -1) fails when executing SeamlessTravel
+
+* UE-42480 SetVectorParameterValueOnMaterials In Construction Script Not Working in Packaged Game
+
+* UE-42459 BP nativization doesn't work with localized text
+
+* UE-42166 Crash when pressing Ctrl-Z after creating a blueprint function
+
+* UE-41893 Child blueprints cannot override functions if their parent was nativized
+
+* UE-41168 Crash opening blueprint when a node class has been removed (e.g., anim blueprint with Ragdoll node with Immediate Mode plugin disabled)
+
+* UE-42209 Incorrect error message when child actor mobility doesn't match
+
+* UE-42253 Enum Element Names Reset on Project Upgrade
+
+* UE-42375 No input received from Raw Input device when packaged for 32bit
+
+* UE-39884 Foliage LODs do not appear to use the base LODs lightmap any longer
+
+* UE-42390 Ad banner is displayed incorrectly on Android 7.0 devices
+
+* UE-42361 Slate UI asset colors washed out on iOS
+
+* UE-42191 Scene Captures not working on Android
+
+* UE-42036 Exposure Is More Extreme In High-End Mobile Preview Modes
+
+* UE-42368 Find Sessions Consistently Returning Results Length of 0
+
+* UE-42452 No textures listed when selecting material
+
+* UE-42789 Rebuilding lighting not saving rebuilt light data
+
+* UE-42196 Static meshes with auto-generated LODs can no longer share lightmaps between LODs
+
+* UE-42179 Crash in FStreamingTextureLevelContext::GetBuildDataIndexRef
+
+* UE-42165 [CrashReport] UE4Editor_Renderer!FReflectionEnvironmentSceneData::ResizeCubemapArrayGPU() [reflectionenvironment.cpp:215]
+
+* UE-42102 Cooked build size increase due to MIC problem
+
+* UE-41950 GitHub 3252 : Added MobileMaterialInterface to UsedMaterials
+
+* UE-29491 Eye Adaption is incorrect at any size other than Fullscreen in PIE
+
+* UE-42554 Regression: Visibility property is not seen in the animation track in UMG
+
+* UE-42210 Crash when using Ctrl+Z to undo a single digit change in the size of a widget
+
+* UE-41829 Crash scrubbing through sequencer when playing montage from old assets
+
+* UE-42665 Update default super search engine to Bing
+
+* UE-42428 Enabling "Show Only Modified Properties" causes details panel menu to disappear
+
+* UE-42312 Crash when saving a Data Table that is referencing itself in a Data Table Row Handle
+
+* UE-35606 Crash when Importing .OBJ using the File > Import into Level Option
+
+* UE-42669 "Package localization ID mismatch during cook!" log spam
+
+* UE-41797 Update to Korean localization
+
+* UE-42106 Edge case where cursor can become hidden in gameplay
+
+* UE-42631 Stereo off while prox sensor uncovered causes low framerate
+
+* UE-42576 Oculus crash on "stereo on" when the proximity sensor is not triggered.
+
+* UE-42101 [daydream] Black screen with FirstPersonTemplate and 4.15P4
+
+* UE-41604 Integrate fix for improper failure case for allocation of RenderTargets in Oculus
+
+* UE-41568 Need to enable GearVR on Win64 platforms
+
diff --git a/Engine/Documentation/Source/Support/Builds/ReleaseNotes/4_2/Updates/UnrealEngine4_2UpdateNotes.INT.udn b/Engine/Documentation/Source/Support/Builds/ReleaseNotes/4_2/Updates/UnrealEngine4_2UpdateNotes.INT.udn
new file mode 100644
--- /dev/null
+++ b/Engine/Documentation/Source/Support/Builds/ReleaseNotes/4_2/Updates/UnrealEngine4_2UpdateNotes.INT.udn
@@ -0,0 +1,49 @@
+Availability:Public
+Title: Unreal Engine 4.2 Update Notes
+Description:
+Crumbs:
+Parent:Support/Builds/ReleaseNotes/2014/4_2
+Order:0
+
+## 4.2.0 Release
+
+Released: 06/04/2014
+
+Github: [https://github.com/EpicGames/UnrealEngine/releases/tag/4.2.0-release](https://github.com/EpicGames/UnrealEngine/releases/tag/4.2.0-release)
+
+Release Notes: [](Support/Builds/ReleaseNotes/2014/4_2) 
+
+## 4.2.1 Hotfix
+
+Released: 06/17/2014
+
+Github: [https://github.com/EpicGames/UnrealEngine/tree/4.2.1-release](https://github.com/EpicGames/UnrealEngine/tree/4.2.1-release) 
+
+Fixes:
+
+* Fixed texture streaming for static meshes with multiple LODs when exceeding the default 1gb texture pool. 
+
+* Fixed Paper2D flipbook components and tile map actors not rendering, sprite color not taking effect. 
+
+* [XB1] Fixed particles rendering as black when bounds were zero. 
+
+* Fixed OnSuccess and OnFail outputs not triggering for AI Move To blueprint node.
+
+* Prevent triggers blocking projectiles in first-person template.
+
+* [iOS] Some projects would fail to launch on device.
+
+* [PS4] Packaging in shipping would not create executable. 
+
+* [Mac] Added Subversion binaries to default install.
+
+## QFE for the 4.2.1 Hotfix
+
+If you are using the latest hotfix for this engine version and experience the issue described below, download and run this ‘quick fix’ to resolve the issue.
+
+* Fixes a crash that occurs with Blueprints that have soft and hard dependencies. ([report](https://answers.unrealengine.com/questions/70620/42-urgent-help-needed-unrecoverable.html))
+
+[Windows Download](http://cdn.unrealengine.com/qfe/QFE-FromCL-2121920-Shelve-2235825.exe](http://cdn.unrealengine.com/qfe/QFE-FromCL-2121920-Shelve-2235825.exe) 
+
+[Mac Download](http://cdn.unrealengine.com/qfe/QFE-FromCL-2121920-Shelve-2235825_Mac.zip](http://cdn.unrealengine.com/qfe/QFE-FromCL-2121920-Shelve-2235825_Mac.zip) 
+
diff --git a/Engine/Documentation/Source/Support/Builds/ReleaseNotes/4_3/Updates/UnrealEngine4_3UpdateNotes.INT.udn b/Engine/Documentation/Source/Support/Builds/ReleaseNotes/4_3/Updates/UnrealEngine4_3UpdateNotes.INT.udn
new file mode 100644
--- /dev/null
+++ b/Engine/Documentation/Source/Support/Builds/ReleaseNotes/4_3/Updates/UnrealEngine4_3UpdateNotes.INT.udn
@@ -0,0 +1,47 @@
+Availability:Public
+Title: Unreal Engine 4.3 Update Notes
+Description:
+Crumbs:
+Parent:Support/Builds/ReleaseNotes/2014/4_3
+Order:0
+
+## 4.3.0 Release
+
+Released: 07/16/2014
+
+Github: [https://github.com/EpicGames/UnrealEngine/releases/tag/4.3.0-release](https://github.com/EpicGames/UnrealEngine/releases/tag/4.3.0-release)
+
+Release Notes: [](Support/Builds/ReleaseNotes/2014/4_3) 
+
+## 4.3.1 Hotfix
+
+Released: 07/29/2014
+
+Github: [https://github.com/EpicGames/UnrealEngine/tree/4.3.1-release](https://github.com/EpicGames/UnrealEngine/tree/4.3.1-release) 
+
+Fixes:
+
+* Fixed: Japanese Localization
+
+* Fixed: PIE with Box BSP affected by Radial Force Actor WorldStatic
+
+* Fixed: PC to iOS workflow doesn't properly sign IPAs for distribution
+
+* Fixed: menu windows do not open when 'Enable Window Animations' is enabled
+
+* Fixed: Landscape Spline: Content created with 4.2 Crashes with 4.3
+
+* Fixed: only first UV channel is usable other channels are broken and textures will not be mapped
+
+* Fixed: Morpheus crashes on init
+
+## QFE for the 4.3.1 Hotfix
+
+If you are using the latest hotfix for this engine version and experience the issue described below, download and run this ‘quick fix’ to resolve the issue.
+
+* Fixes a crash that occurs with Blueprints that have soft and hard dependencies. ([report](https://answers.unrealengine.com/questions/70620/42-urgent-help-needed-unrecoverable.html))
+
+[Windows Download](http://cdn.unrealengine.com/qfe/QFE-FromCL-2226259-Shelve-2236052.exe](http://cdn.unrealengine.com/qfe/QFE-FromCL-2226259-Shelve-2236052.exe) 
+
+[Mac Download](http://cdn.unrealengine.com/qfe/QFE-FromCL-2226259-Shelve-2236052_Mac.zip](http://cdn.unrealengine.com/qfe/QFE-FromCL-2226259-Shelve-2236052_Mac.zip) 
+
diff --git a/Engine/Documentation/Source/Support/Builds/ReleaseNotes/4_4/Updates/UnrealEngine4_4UpdateNotes.INT.udn b/Engine/Documentation/Source/Support/Builds/ReleaseNotes/4_4/Updates/UnrealEngine4_4UpdateNotes.INT.udn
new file mode 100644
--- /dev/null
+++ b/Engine/Documentation/Source/Support/Builds/ReleaseNotes/4_4/Updates/UnrealEngine4_4UpdateNotes.INT.udn
@@ -0,0 +1,69 @@
+Availability:Public
+Title: Unreal Engine 4.4 Update Notes
+Description:
+Crumbs:
+Parent:Support/Builds/ReleaseNotes/2014/4_4
+Order:0
+
+## 4.4.0 Release
+
+Released: 08/14/2014
+
+Github: [https://github.com/EpicGames/UnrealEngine/releases/tag/4.4.0-release](https://github.com/EpicGames/UnrealEngine/releases/tag/4.4.0-release)
+
+Release Notes: [](Support/Builds/ReleaseNotes/2014/4_4) 
+
+## 4.4.1 Hotfix
+
+Released: 08/26/2014
+
+Github: [https://github.com/EpicGames/UnrealEngine/tree/4.4.1-release](https://github.com/EpicGames/UnrealEngine/tree/4.4.1-release) 
+
+Fixes:
+
+* Unreal Front End: Fixed a crash that occurred when using cooking-on-the-fly with MacNoEditor
+
+* Oculus Rift support has been updated to use the Oculus 0.4.1 Beta SDK
+
+* Fixed a crash when opening animations from the Military Weapons Asset Content Pack on the Marketplace
+
+* Fixed a bug where mouse handling was broken on Mac
+
+* Blueprints: Added a redirect to the base engine ini which allows existing content to point to the correct location for an enum, fixing a compatibility break
+
+* Fixed info.plist search order in UAT
+
+* Fixed a bug with audio in Couch Knights
+
+## 4.4.2 Hotfix
+
+Released: 09/09/2014
+
+Github: [https://github.com/EpicGames/UnrealEngine/tree/4.4.2-release](https://github.com/EpicGames/UnrealEngine/tree/4.4.2-release)  
+
+Fixes:
+
+* Fixed an error that could occur with Notify tracks in animation Montages.
+
+* Fixed an error that could cause the Mac Editor to become unresponsive when Slate Text Views were used
+
+* Fixed game plugin source being compiled with optimizations when Debug or DebugGame configuration is set.
+
+* Updated Vehicle Template with Oculus Rift support!
+
+## 4.4.3 Hotfix
+
+Released: 09/16/2014
+
+Github: [https://github.com/EpicGames/UnrealEngine/tree/4.4.3-release](https://github.com/EpicGames/UnrealEngine/tree/4.4.3-release) 
+
+Fixes:
+
+* Fixed a crash after loading a Blueprint with particular linked nodes set as a dependency of another Blueprint.
+
+* Fixed a .NET Framework SDK test dependency added in 4.4.1.
+
+* Added Xcode 6 support; please review the known issues.
+
+* Improved Oculus compatibility with the Vehicle template – fixed the camera setup and replaced the vehicle mesh
+
diff --git a/Engine/Documentation/Source/Support/Builds/ReleaseNotes/4_5/Updates/UnrealEngine4_5UpdateNotes.INT.udn b/Engine/Documentation/Source/Support/Builds/ReleaseNotes/4_5/Updates/UnrealEngine4_5UpdateNotes.INT.udn
new file mode 100644
--- /dev/null
+++ b/Engine/Documentation/Source/Support/Builds/ReleaseNotes/4_5/Updates/UnrealEngine4_5UpdateNotes.INT.udn
@@ -0,0 +1,51 @@
+Availability:Public
+Title: Unreal Engine 4.5 Update Notes
+Description:
+Crumbs:
+Parent:Support/Builds/ReleaseNotes/2014/4_5
+Order:0
+
+## 4.5.0 Release
+
+Released: 10/14/2014
+
+Github: [https://github.com/EpicGames/UnrealEngine/releases/tag/4.5.0-release](https://github.com/EpicGames/UnrealEngine/releases/tag/4.5.0-release) 
+
+Release Notes: [](Support/Builds/ReleaseNotes/2014/4_5) 
+
+## 4.5.1 Hotfix
+
+Released: 10/23/2014
+
+Github: [https://github.com/EpicGames/UnrealEngine/tree/4.5.1-release](https://github.com/EpicGames/UnrealEngine/tree/4.5.1-release) 
+
+Fixes:
+
+* The default landscape material now correctly renders on VehicleGame.
+
+* Fixed an error when using Hot-Reload with XCode.
+
+* Fixed an error when packaging projects on Linux OS.
+
+* Corrected bugs that occurred on Nvidia devices in 10.10 Yosemite. 
+
+* Fixed an error while attempting to join a match of Shootergame while on Linux OS.
+
+* Fixed the UI in Shootergame to be properly masked on Linux OS.
+
+* Material Parameter Collections now work properly when called through Blueprints.
+
+* Mouse input has been adjusted to a more consistent and less sensitive state.
+
+* Resolved some issues discovered with retargeting animations with Epic compatible skeletons.
+
+* Fixed an issue with iOS certificate validation while packaging projects.
+
+## QFE for the 4.5.1 Hotfix
+
+If you are using the latest hotfix for this engine version and experience the issue described below, download and run this ‘quick fix’ to resolve the issue.
+
+* Fixes issue where editor incorrectly detects that iTunes isn’t installed and will not allow any iOS Launch On or Package.
+
+[Windows Download](https://s3.amazonaws.com/unrealengine/qfe/QFE-FromCL-2370532.exe](https://s3.amazonaws.com/unrealengine/qfe/QFE-FromCL-2370532.exe) 
+
diff --git a/Engine/Documentation/Source/Support/Builds/ReleaseNotes/4_6/Updates/UnrealEngine4_6UpdateNotes.INT.udn b/Engine/Documentation/Source/Support/Builds/ReleaseNotes/4_6/Updates/UnrealEngine4_6UpdateNotes.INT.udn
new file mode 100644
--- /dev/null
+++ b/Engine/Documentation/Source/Support/Builds/ReleaseNotes/4_6/Updates/UnrealEngine4_6UpdateNotes.INT.udn
@@ -0,0 +1,61 @@
+Availability:Public
+Title: Unreal Engine 4.6 Update Notes
+Description:
+Crumbs:
+Parent:Support/Builds/ReleaseNotes/2014/4_6
+Order:0
+
+## 4.6.0 Release
+
+Released: 12/03/2014
+
+Github: [https://github.com/EpicGames/UnrealEngine/releases/tag/4.6.0-release](https://github.com/EpicGames/UnrealEngine/releases/tag/4.6.0-release) 
+
+Release Notes: [](Support/Builds/ReleaseNotes/2014/4_6) 
+
+## 4.6.1 Hotfix
+
+Released: 12/16/2014
+
+Github: [https://github.com/EpicGames/UnrealEngine/tree/4.6.1-release](https://github.com/EpicGames/UnrealEngine/tree/4.6.1-release) 
+
+Fixes:
+
+* ContentExamples Media map movie no longer renders black.
+
+* ContentExamples Media map no longer shows the UDK splash screen.
+
+* ContentExamples Media map is no longer missing source video.
+
+* XBox One now boots from Visual Studio correctly.
+
+* Fixed an error that caused material If nodes to break.
+
+* Fixed a bug that was causing new and converted non-ascii text projects to crash.
+
+* Corrected a bug that prevented composite decorators from being opened in Behavior Trees.
+
+* Event Nodes can now be called for Event Dispatchers.
+
+* Fixed an error that prevented the creation of new events in the AnimBP Eventgraph for new Anim Notifies.
+
+* Corrected Static Lighting on InstancedStaticMeshes.
+
+* Corrected the adb path and device serial number for Android packaged install batch files
+
+* Fixed an error that caused crashes on Android when Touch events were called too early.
+
+* Fixed a bug in pathfinding when default querying extent was not large enough on the Z axis.
+
+* InstancedStaticMesh now renders correctly on mobile devices.
+
+* Corrected a bug with animation retargeting adversely affecting looping animations.
+
+## QFE for the 4.6.1 Hotfix
+
+If you are using the latest hotfix for this engine version and experience the issue described below, download and run this ‘quick fix’ to resolve the issue.
+
+* Fixes issue where editor incorrectly detects that iTunes isn’t installed and will not allow any iOS Launch On or Package.
+
+[Windows Download](https://s3.amazonaws.com/unrealengine/qfe/QFE-FromCL-2405993.exe](https://s3.amazonaws.com/unrealengine/qfe/QFE-FromCL-2405993.exe) 
+
diff --git a/Engine/Documentation/Source/Support/Builds/ReleaseNotes/4_7/Updates/UnrealEngine4_7UpdateNotes.INT.udn b/Engine/Documentation/Source/Support/Builds/ReleaseNotes/4_7/Updates/UnrealEngine4_7UpdateNotes.INT.udn
new file mode 100644
--- /dev/null
+++ b/Engine/Documentation/Source/Support/Builds/ReleaseNotes/4_7/Updates/UnrealEngine4_7UpdateNotes.INT.udn
@@ -0,0 +1,361 @@
+Availability:Public
+Title: Unreal Engine 4.7 Update Notes
+Description:
+Crumbs:
+Parent:Support/Builds/ReleaseNotes/2015/4_7
+Order:0
+
+## 4.7.0 Release
+
+Released: 02/24/2015
+
+Github: [https://github.com/EpicGames/UnrealEngine/releases/tag/4.7.0-release](https://github.com/EpicGames/UnrealEngine/releases/tag/4.7.0-release) 
+
+Release Notes: [](Support/Builds/ReleaseNotes/2015/4_7)  
+
+## 4.7.1 Hotfix
+
+Released: 02/27/2015
+
+Github: [https://github.com/EpicGames/UnrealEngine/releases/tag/4.7.1-release](https://github.com/EpicGames/UnrealEngine/releases/tag/4.7.1-release) 
+
+Fixes:
+
+* UE-6706 Text does not scale when using the "Draw Text" node in a HUD blueprint
+
+* UE-7139 Crash caused by "/Game/Maps/UEDPIE_2_Menu.Menu not cleaned up by garbage collection!"
+
+* UE-7497 No Graphs are displayed in Visual Logger after enabling them
+
+* UE-9745 Color picker does not work.
+
+* UE-10016 Using the Merge Blueprints feature crashes the editor
+
+* UE-10114 LevelScript should not get a Overlap event added by default
+
+* UE-10137 Redirected graph nodes cause a crash in IsDataOnlyBlueprint
+
+* UE-10141 Vehicle code project has default texture in packaged game
+
+* UE-10143 High reflectivity on all materials in SunTemple on Android
+
+* UE-10167 Undoing a pasted custom event node with an input pin results in a crash
+
+* UE-10169 When installing VS through the editor, invoke installer with a command-line parameter that enables Microsoft's Telemetry
+
+* UE-10257 Mac to Android first time device is plugged in, after installing tadp user still has to place sdk paths
+
+* UE-10260 Pre-existing Foliage stays at highest LOD
+
+* UE-10277 Android fails to show as valid platform for Mac in UFE until editor is opened
+
+* UE-10283 TargetArray nodes are reported as invalid if played in Standalone mode but not in PIE
+
+* UE-10298 Creating an exposed variable does not always update correctly to the bp's Details panel
+
+* UE-10305 Set Timer node does not stop the timer if it is set to zero by the Custom Event
+
+* UE-10313 Crash playing animation montage in reverse
+
+* UE-10319 Crash attempting to edit a Struct Transform's rotation value as a slider
+
+* UE-10406 All shadow casting on Paper2D sprites is broken (no sprites cast shadows)
+
+* UE-10410 Reparenting an Actor blueprint crashes the editor
+
+* UE-10510 Crash occurs applying instance changes to blueprint when instances have different changes
+
+* UE-10524 Icon reports failure to add if creating a project with source control enabled
+
+* UE-10529 TappyChicken: HTML5: Firefox: The Fullscreen button now displays the screen offcentre (large black border on the right)
+
+* UE-10537 Editor crash when child component has physics enabled and is moved
+
+* UE-10562 Empty if blocks are generated in shaders and can cause driver compiler crashes on Android
+
+* UE-10646 Crash occurs copying a non-scene component in component window
+
+* UE-10650 Possible crash in USceneComponent::UpdatePhysicsVolume
+
+* UE-10658 Excessive memory consumption in some circumstances due to off-by-1 count error
+
+* UE-10732 Actor's attached to an Actor with a blueprint created root component don't stay in correct relative position when moved
+
+* UE-10793 GitHub 851 : Update .gitignore - ignore FeaturePacks
+
+* UE-10812 Comparing BB entries doesn't work
+
+* UE-10822 Crash clearing the undo buffer during blueprint merge
+
+* UE-10826 Crash compiling Level BP after deleting event
+
+## 4.7.2 Hotfix
+
+Released: 03/04/2015
+
+Github: [https://github.com/EpicGames/UnrealEngine/releases/tag/4.7.2-release](https://github.com/EpicGames/UnrealEngine/releases/tag/4.7.2-release) 
+
+Fixes:
+
+* UE-11207 ios launch fail after upgrading to 4.7.1, provision not found, signing key not found
+
+* UE-11202 Horizontal navigation in Notify timelines is super-frustrating
+
+* UE-11130 Cannot build code projects for ios
+
+* UE-11102 Mac editor sometimes crashes in FMacWindow::IsMaximized when creating a new window
+
+* UE-11046 Crash in FLandscapeComponentSceneProxy::GetStaticBatchElementVisibility with tiny component size
+
+* UE-11023 Opening the engine crashes: Assert failed: GConfig->GetInt
+
+* UE-11018 Synced BlackBoard keys don't work
+
+* UE-8589 MatineeFightScene has no sound on mac
+
+* UE-8550 Audio cuts out during scene transitions in the Particle Effects cinematic
+
+## 4.7.3 Hotfix
+
+Released: 03/17/2015
+
+Github: [https://github.com/EpicGames/UnrealEngine/releases/tag/4.7.3-release](https://github.com/EpicGames/UnrealEngine/releases/tag/4.7.3-release) 
+
+Fixes:
+
+* UE-11029 Crash after converting to 4.7.0 from 4.6.1
+
+* UE-11404 Second player is not affecting by any gamepad input.
+
+* UE-11402 Editor will not build if you have PS4 code, and ONLY SDK 2.0
+
+* UE-11485 Compiling the engine with Xcode 6.3 beta generates errors
+
+* UE-10128 Preview Static mesh does not translate with socket
+
+* UE-11315 Bugg 38187: SIGSEGV: invalid attempt to access memory at address 0x1010102
+
+* UE-10569 Crash when opening maps and blueprints in user's project
+
+* UE-10708 Build failing for users project with all platforms
+
+* UE-11312 Mac Bugg 38298: SIGSEGV: invalid attempt to access memory at address 0x3
+
+* UE-11369 Crash when clicking on a "Two Bone IK" node
+
+* UE-11185 A function in a Blueprint Function Library that references another function in that library cannot be saved after closing and reopening the project
+
+* UE-11322 Crash when refreshing/moving/deleting BP nodes
+
+* UE-11743 4.7 does not support Morpheus under SDK 2.0
+
+* UE-11302 Cannot deploy to iPad
+
+* UE-11692 Crash when rendering hit proxies
+
+* UE-11304 Can't compile CrossCompilerTools due to missing file
+
+* UE-11484 Fix Xcode 6.3 beta clang issues.
+
+* UE-11393 Tiled landscape import' creates overlapped landscape components
+
+* UE-11017 GitHub 857 : Fixed native enum support in blackboard component.
+
+* UE-10697 Crash when opening Post Process Material with Live Data Preview running in Profiler
+
+* UE-11225 Closing Chunk Parameters tabs with chunk selected and reopening Destructible will cause crash
+
+* UE-10149 Project crashes on load after modifying Parent Blueprints
+
+* UE-11244 Crash when renaming assets from the Persona asset browser
+
+* UE-11272 Blueprint compile error when upgrading project
+
+* UE-11319 Crash on project open (Assertion Failed)
+
+* UE-11401 Need more logging in the Deployment Server
+
+* UE-11292 Editor Crashes when clicking on Transform(Modify) Bone node in Content Examples
+
+* UE-10900 Crash in FBatchingSPDI::DrawMesh in stand alone game
+
+* UE-11505 Crash launching Blueprint Office standalone
+
+* UE-10856 Crash on dragging static meshes into the scene
+
+* UE-11298 Crash with Distance Fields while using World Composition
+
+* UE-11422 GitHub 894 : wrong parsing of the return of adb cmd: device
+
+* UE-11447 Player Character does not collide with landscape
+
+* UE-11638 GetActorsPerception useless in 4.7
+
+* UE-11589 GitHub 909 : Update .gitignore
+
+* UE-11605 Overridden components are broken in static Actors in cooked build
+
+* UE-11666 UberGraphFunction should be preloaded in Class::Link
+
+* UE-11711 Severe log spam when particles collide with static objects
+
+* UE-11640 AI Move To issues in 4.7
+
+* UE-11360 Placing a ' - ' in theproject name will cause the Android build to fail
+
+* UE-11647 Denver Tegra K1 (Nexus 9) disables rendering thread causing bad performance
+
+* UE-11383 Screen orientation is ignored on Android devices.
+
+* UE-11460 A level will not save if a Blueprint that uses a Blueprint Component with an Event Dispatcher is placed in it
+
+* UE-11600 Performing a hot reload on a project that contains an FStringAssetReference will crash the Editor.
+
+* UE-11707 Implement workaround for Localization being broken for a number of editor strings
+
+* UE-11592 Cook on the fly Android projects crash after launching on the device.
+
+* UE-11739 In some circumstances CrowdManager doesn't get created.
+
+* UE-11321 Crash when creating landscape
+
+* UE-11777 BP-implemented BT decorators end up thinking they're not on active branch
+
+* UE-11826 MediaPlayer14 class lookup crash with distribution builds
+
+* UE-11840 Launch on Nexus 9 fails due to unknown arm64-v8a architecture
+
+* UE-11893 Editor crashes when dragging a docked window
+
+## 4.7.4 Hotfix
+
+Released: 04/02/2015
+
+Github: [https://github.com/EpicGames/UnrealEngine/releases/tag/4.7.4-release](https://github.com/EpicGames/UnrealEngine/releases/tag/4.7.4-release) 
+
+Fixes:
+
+* UE-11891 MoveToActor doesn't update AI's path on goal actor's move
+
+* UE-11997 Navigation crashing on PIE in absence of NavOctree
+
+* UE-12108 [CrashReport] Ensure in UPrimitiveComponent::SetCanEverAffectNavigation()
+
+* UE-12106 [CrashReport] Saving after blueprint change causes crash.
+
+* UE-12443 [CrashReport] Crash in blueprints during autosave
+
+* UE-9417 It is easy to make a comment box unselectable when it overlaps with another comment block
+
+* UE-10891 Force deleting a blueprint that is parented to another blueprint causes crash
+
+* UE-10865 4.6.1 to 4.7 conversion crash "Assertion failed: Pin != nullptr ...K2Node_DynamicCast.cpp]"
+
+* UE-11523 DateTime variables will cause a crash on compile
+
+* UE-11011 Crash (assert) when accessing a property in a blueprint
+
+* UE-12073 [CrashReport] Disconnecting/changing blueprint nodes causes crash.
+
+* UE-12206 GitHub 944 : UE-12128: Properly save structs when struct field default value is different
+
+* UE-12962 Editor crashes after clicking on "Apply Instance Change to Blueprint" button
+
+* UE-10574 Component properties of template not propagated when SpawnActor called
+
+* UE-12109 [CrashReport] Crash while attempting to load a project
+
+* UE-11368 uv scale of materials on hollow BSP does not persist after building geometry
+
+* UE-4743 Crash when moving a sprite imported as part of a spritesheet to another folder if the sprite names had invalid character
+
+* UE-11372 Access Violation in UE4Editor_SlateCore!TWeakPtr when changing view options
+
+* UE-9824 Widget class is replaced with a placeholder class when reloaded
+
+* UE-10040 Copy/Paste a button inside of another button then attempting to move the pasted button crashes editor
+
+* UE-11204 Crash when returning to viewport from immersive mode with an editor window tabbed beneath
+
+* UE-12915 Crash when using static lighting in foliage
+
+* UE-7213 Crash occurred trying to cycle gizmo in Persona Anim BP Viewport
+
+* UE-9982 Sprites with non-center pivots can be invisible in thumbnail renders
+
+* UE-12196 Automatic collision generation is incorrectly rotated for sprites that were imported as rotated from a sprite sheet
+
+* UE-12041 Sprite sheets with modified textures don't reimport correctly
+
+* UE-7918 Sprites don't initialize PixelsPerUnrealUnit in cooked builds, causing invalid values for GetSocketLocation, etc... for assets that have Pixels/uu of 2.56
+
+* UE-11903 Crash closing ElementalDemo game in opengl
+
+* UE-12150 Crash launching PIE when sublevel contains landscape in hidden level
+
+* UE-12111 [CrashReport] Crash after undoing/deleting landscape.
+
+* UE-12433 Editor crashes when toggling off visibility for a sublevel that contains a landscape
+
+* UE-11815 [CrashReport] Crash while attempting to connect iOS device.
+
+* UE-12165 Morpheus crashes occasionally while running
+
+* UE-12209 Canvas items are not drawn correctly in Morpheus
+
+* UE-12211 Morpheus tears visibly even when VSync is on.
+
+* UE-12212 Morpheus runs HMD distortion shader an extra time.
+
+* UE-12324 Morpheus rendering broken when console rendering is on the screen.
+
+* UE-11302 Cannot deploy to iPad
+
+* UE-12063 Cannot package project on Linux
+
+* UE-12210 Morpheus IPD is applied incorrectly.
+
+* UE-13048 Xbox projects (code and BP) launch to black screen, logs show success
+
+* UE-12213 Landscape material reverts to default when launched onto an Android device
+
+* UE-12542 MAC: Very slow texture compression times in Editor
+
+* UE-12235 Large performance regression in shadow rendering.
+
+* UE-6190 UVs skewed in Mobile Preview mode for some skeletal meshes
+
+* UE-11865 Open a level that contains instanced meshes (with built lighting) will crash the editor
+
+* UE-10664 Crash in OpenGL3 and OpenGL4 occurs when opening any map after having a map loaded with foliage.
+
+* UE-11810 [CrashReport] Crash in ULevel::BuildStreamingData casting Texture
+
+* UE-11434 New user crashes on Editor load
+
+* UE-12228 Invalid error messages in crashes
+
+## 4.7.5 Hotfix
+
+Released: 04/08/2015
+
+Github: [https://github.com/EpicGames/UnrealEngine/releases/tag/4.7.5-release](https://github.com/EpicGames/UnrealEngine/releases/tag/4.7.5-release) 
+
+Fixes:
+
+* UE-13250 Multiplayer Shootout project sample crashes on start
+
+* UE-13274 Projects crashing on launch after converting from 4.7.3 to 4.7.4
+
+* UE-11828 HydroErosion on landscape crashes on Linux
+
+## 4.7.6 Hotfix
+
+Released: 04/20/2015
+
+Github: [https://github.com/EpicGames/UnrealEngine/releases/tag/4.7.6-release](https://github.com/EpicGames/UnrealEngine/releases/tag/4.7.6-release) 
+
+Fixes:
+
+* UE-13664 MAC: Crash on startup
+
diff --git a/Engine/Documentation/Source/Support/Builds/ReleaseNotes/4_8/Updates/UnrealEngine4_8UpdateNotes.INT.udn b/Engine/Documentation/Source/Support/Builds/ReleaseNotes/4_8/Updates/UnrealEngine4_8UpdateNotes.INT.udn
new file mode 100644
--- /dev/null
+++ b/Engine/Documentation/Source/Support/Builds/ReleaseNotes/4_8/Updates/UnrealEngine4_8UpdateNotes.INT.udn
@@ -0,0 +1,181 @@
+Availability:Public
+Title: Unreal Engine 4.8 Update Notes
+Description:
+Crumbs:
+Parent:Support/Builds/ReleaseNotes/2015/4_8
+Order:0
+
+## 4.8.0 Release
+
+Released: 06/08/2015
+
+Github: [https://github.com/EpicGames/UnrealEngine/releases/tag/4.8.0-release](https://github.com/EpicGames/UnrealEngine/releases/tag/4.8.0-release)
+
+Release Notes: [](Support/Builds/ReleaseNotes/2015/4_8)  
+
+## 4.8.1 Hotfix
+
+Released: 06/24/2015
+
+Github: [https://github.com/EpicGames/UnrealEngine/releases/tag/4.8.1-release](https://github.com/EpicGames/UnrealEngine/releases/tag/4.8.1-release) 
+
+Fixes:
+
+* UE-11816	[CrashReport] Garbage collection crashes not caused directly by the editor
+
+* UE-12079	Crash: UGameInstance::StartGameInstance "Failed to enter xx"
+
+* UE-14313	Tappy chicken size has increased dramatically from 4.6 to 4.7
+
+* UE-15505	Elemental Demo when played does not start matinee
+
+* UE-15651	GitHub 1148 : Fix divide by zero when importing 0-length animations
+
+* UE-16471	UMapProperty crashes
+
+* UE-16513	[CrashReport] Adding a Return node to a function that contains a Local array variable then compiling will crash the engine
+
+* UE-16521	[CrashReport] UE4Editor_Core!FRunnableThread::SetTls() [threadingbase.cpp:267]
+
+* UE-16711	Xbox fails to open BP projects for shipping config
+
+* UE-16724	StaticFindObject can return objects that have not been fully streamed in
+
+* UE-16752	[CrashReport] UE4Editor_D3D11RHI!VerifyD3D11CreateTextureResult() [d3d11util.cpp:223]
+
+* UE-16754	GatherTextFromSource commandlet always fails when using the deprecated SourceFileSearchFilters setting.
+
+* UE-16789	'Target' pin name in certain BP nodes jitters in localized mode
+
+* UE-16806	[CrashReport] UE4Editor_HotReload!GetCachedCDODuplicate() [hotreload.cpp:636]
+
+* UE-16844	[CrashReport] Building lighting with a CustomMeshComponent in level causes crash
+
+* UE-16894	Interrupting the KiteDemo cinematic leaves you with no drone controls
+
+* UE-16900	Crash when Applying Hole Material to Imported Tiled Landscape with World Composition
+
+* UE-16907	Project using ScriptPluginComponents can crash when upgrading to 4.8
+
+* UE-16918	Sprite sheets cannot be reimported in 4.8 in a project that was converted from an older version
+
+* UE-16971	Packaging fails with plugin and full rebuild checked
+
+* UE-16972	Self reference disconnecting on editor restart
+
+* UE-16995	Crash undoing TMap UProperties
+
+* UE-17008	Editor crashes every time you save if there is no editor startup map set
+
+* UE-17067	Engine builds need to include engine localization data
+
+* UE-17075	The localization dashboard generates inaccurate paths when referencing files on a different drive than the binary.
+
+* UE-17082	Project debug symbols are no longer included in crash reporter callstacks.
+
+* UE-17087	Assertion failed: !Result || Result->LinkerLoad == Linker crash
+
+* UE-17111	Indirect Lighting Cache is broken on PS4
+
+* UE-17129	Crash when terminating a Matinee camera anim
+
+* UE-17153	USkinnedMeshComponent::QuerySupportedSockets behaves differently in !WITH_EDITOR
+
+* UE-17194	LoadPackageInternal may crash due to corrupted linkers
+
+* UE-17221	Hot reloading does not work on Linux
+
+* UE-17222	CPU usage is not decreased when the editor isn't in foreground on Linux
+
+* UE-17228	Cook commandlet dies on Linux with error if run on a headless machine.
+
+* UE-17233	Media streaming from the internet is not working on Windows
+
+## 4.8.2 Hotfix
+
+Released: 07/14/2015
+
+Github: [https://github.com/EpicGames/UnrealEngine/releases/tag/4.8.2-release](https://github.com/EpicGames/UnrealEngine/releases/tag/4.8.2-release) 
+
+Fixes:
+
+* UE-16177	When launching the editor and having oculus_runtime_sdk_0.6.0.0_win_partner installed, the editor crashes.
+
+* UE-16608	Convert Actor option is grayed out and not selectable. (ConversionRoot metadata not found)
+
+* UE-16822	[CrashReport] Saving in animation blueprint can cause a crash
+
+* UE-16896	Level streaming crash in KiteDemo after retriggering cinematic
+
+* UE-16905	User Play Space Rot does not change anything in Client Play Camera Shake
+
+* UE-16966	Default settings cannot be used to import character animations
+
+* UE-16970	StrategyGame crash on android devices bad name index
+
+* UE-17103	Curve editor breaks randomly in Cascade
+
+* UE-17118	Custom Output node added to a Material Editor crashes Editor
+
+* UE-17163	[CrashReport] dxgi! and RHI related crash on project open
+
+* UE-17167	[CrashReport] Character interaction with asset causes crash
+
+* UE-17177	Localization manifest writing fails if source control is hooked up and the manifest file does not exist already.
+
+* UE-17211	[CrashReport] Clicking connection wires between blueprints nodes causes crash
+
+* UE-17212	Stack overflow due to heavy recursion in indirect lighting cache
+
+* UE-17213	LightSourceAngle is incorrectly greyed out
+
+* UE-17216	[CrashReport] Crash in FD3D11ComputeShader::~FD3D11ComputeShader()
+
+* UE-17226	OnPaint function causes artifacts to show up in a New Editor window
+
+* UE-17384	Import options broken when importing bones-only fbx animation
+
+* UE-17385	Android game packaged with obb outside apk does not create SaveGame directory on device
+
+* UE-17400	Regression for UPROPERTY specifier causing crash ("Template Mismatch during attachment") during PIE
+
+* UE-17419	Crash compiling BPAIController and RunnerAI via session frontend test
+
+* UE-17443	DynamicModifiersOnly navmesh generation mode doesn't work
+
+* UE-17523	Paper2D Set Sprite function in Construction requires mobility to be set to Movable
+
+* UE-17562	[CrashReport] Crash in TBaseMulticastDelegate<void,TSharedPtr<FSearchEntry,0> >::AddDelegateInstance()
+
+* UE-17569	GitHub 1309 : [Core][Windows] Don't crash when GetRawInputData() fails.
+
+* UE-17673	Clicking the GameMode Override dropdown in the World Settings tab crashes the editor
+
+* UE-17792	Analog Input does not work for network clients (AddMovementInput scale values other than 0 and 1 do not replicate for clients)
+
+* UE-17829	Second client fails to load when using multiplayer PIE in separate processes
+
+* UE-18020	Crash Report Client - updates to data collection
+
+* UE-18041	[CrashReport] Fatal error in PostInitProperties
+
+* UE-18080	[CrashReport] Crash in UProjectileMovementComponent::TickComponent()
+
+* UE-18173	Crash report client submits user names for github builds
+
+## 4.8.3 Hotfix
+
+Released: 07/29/2015
+
+Github: [https://github.com/EpicGames/UnrealEngine/releases/tag/4.8.3-release](https://github.com/EpicGames/UnrealEngine/releases/tag/4.8.3-release) 
+
+Fixes:
+
+* UE-18668 Crash when closing projects
+
+* UE-17342 Crash when exiting PIE after server travel
+
+* UE-18530 Assert in ConformComponentsUtils::FindeNativeArchetype
+
+* UE-19151 Parallel Rendering is disabled in Shipping builds
+
diff --git a/Engine/Documentation/Source/Support/Builds/ReleaseNotes/4_9/Updates/UnrealEngine4_9UpdateNotes.INT.udn b/Engine/Documentation/Source/Support/Builds/ReleaseNotes/4_9/Updates/UnrealEngine4_9UpdateNotes.INT.udn
new file mode 100644
--- /dev/null
+++ b/Engine/Documentation/Source/Support/Builds/ReleaseNotes/4_9/Updates/UnrealEngine4_9UpdateNotes.INT.udn
@@ -0,0 +1,139 @@
+Availability:Public
+Title: Unreal Engine 4.9 Update Notes
+Description:
+Crumbs:
+Parent:Support/Builds/ReleaseNotes/2015/4_9
+Order:0
+
+## 4.9.0 Release
+
+Released: 08/31/2015
+
+Github: [https://github.com/EpicGames/UnrealEngine/releases/tag/4.9.0-release](https://github.com/EpicGames/UnrealEngine/releases/tag/4.9.0-release) 
+
+Release Notes: [](Support/Builds/ReleaseNotes/2015/4_9)  
+
+## 4.9.1 Hotfix
+
+Released: 09/15/2015
+
+Github: [https://github.com/EpicGames/UnrealEngine/releases/tag/4.9.1-release](https://github.com/EpicGames/UnrealEngine/releases/tag/4.9.1-release) 
+
+Fixes:
+
+*  UE-20265 ForEach 'Enum' node returns incorrect enum values
+
+*  UE-20643 Enum to Byte conversion returns 2 for the first enum then 4 then 6 and so on rather than 0, 1, 2, etc
+
+*  UE-20878 Child Blueprint getting certain property values cleared in cooked games
+
+*  UE-20009 [CrashReport] UE4Editor_UnrealEd!operator<<() [editortransaction.cpp:157]
+
+*  UE-20867 Issues converting enum to byte in BP
+
+*  UE-20733 Disabling collision of an actor after an overlap event will cause the event to fire twice
+
+*  UE-20466 [CrashReport] UE4Editor_Engine!UAnimSequence::GetBonePose() [animsequence.cpp:994]
+
+*  UE-20413 Montage with negative value of curves doesn't work
+
+*  UE-20496 Overlap checks during swept movement can generate repeated begin/end overlaps, or wrong end overlap event
+
+*  UE-20146 Can't increase lightmap resolution on landscape, it crashes
+
+*  UE-20202 Dynamic point lights generate shader compile errors on HTML5
+
+*  UE-20492 Renamed/Copied code project folders fail to launch on to iOS
+
+*  UE-20728 Editable Text crashes on Android
+
+*  UE-20380 HTML5: Mac: RunMacHTML5LaunchHelper.command has incorrect file type in P4, causes error
+
+*  UE-20419 Mono error encountered with packaged HTML5 sample
+
+*  UE-20872 Enabling MSAA crashes when using metal
+
+*  UE-20968 Offline Shader Compile has been disabled on Mac
+
+* Note: When performing a launch on from the Mac to iOS the first time there will be an increase in the cook time as the shaders will be rebuilt.
+
+*  UE-20928 Launch on for Android may crash or use wrong version of aapt
+
+*  UE-20864 Launch on to Wi-Fi Android devices will fail
+
+*  UE-20008 SunTemple exits to home screen on the iPad 4
+
+*  UE-20590 Can't build UE4Game for IOS from PC
+
+*  UE-12215 Sound effects continue to play when game is suspended on Android
+
+*  UE-20692 Epic Games Launcher unhides while the editor is loading on Mac
+
+*  UE-20991 Launch On for iOS does not automatically start
+
+* Note: The Launch On feature for iOS on Mac has been temporarily disabled due to an incompatibility with Xcode 7.0. We are working on the issue for 4.10.
+
+*  UE-20640 SSAO 4.9 perf regression fix
+
+*  UE-20335 Blendable Weight set higher than 1.0 crashes
+
+*  UE-20639 Crash on Mobile Previewing or Playing in a New Editor Viewport a Mobile/Tablet - Mobile/HTML Project
+
+*  UE-20297 REGRESSION: HierarchicalInstancedStaticMeshComponent LOD transitions are broken in 4.9 and Main
+
+*  UE-20172 [CrashReport] UE4Editor_Engine!FBatchedElements::FSimpleElementBSSContainer::GetBSS() [batchedelements.h:308]
+
+*  UE-20699 "Build Solution" does not work for GitHub users
+
+*  UE-20291 Rev Oculus SDK to 0.7 / 0.6.18
+
+## 4.9.2 Hotfix
+
+Released: 09/30/2015
+
+Github: [https://github.com/EpicGames/UnrealEngine/releases/tag/4.9.2-release](https://github.com/EpicGames/UnrealEngine/releases/tag/4.9.2-release) 
+
+Fixes:
+
+*  UE-20792	DrawMaterial() Function not Rendering Translucent Material to HUD
+
+*  UE-21233	Calling RequestStimuliListenerUpdate on actor with one element in AIPerception Sensing Config causes crash
+
+*  UE-21234	Crash on play when AIPerception Senses Config element is set to "sense config" options
+
+*  UE-20913	Undo after deletion in BP editor causes crash
+
+*  UE-17698	[CrashReport] UE4Editor_CoreUObject!UStruct::IsChildOf() [class.h:400]
+
+*  UE-20644	[CrashReport] Crash while undoing in Blueprints - FObjectRecord::FReader::operator
+
+*  UE-20258	Child blueprint classes can no longer change default values for inherited variables
+
+*  UE-21081	Cannot open assets created in 4.9.1 in 4.9.0 version of editor
+
+*  UE-20210	[CrashReport] Crash on Windows Chinese version when changing editor language
+
+*  UE-21337	Movies are not skipped when Wait for Movies to Complete is false
+
+*  UE-21292	Parallax Occlusion Mapping broken on Macs
+
+*  UE-21127	Crash in FSlateImageRun destructor
+
+*  UE-20550	Sliders and check boxs not working on touch screen device
+
+*  UE-21076	Possible divide by zero in animation transitions with zero crossfade durations
+
+*  UE-20961	[CrashReport] UE4Editor_Engine!FKShapeElem::GetShapeCheck<FKConvexElem>() [shapeelem.h:55]
+
+*  UE-20952	Launching a Project which uses Depth of Field crashes to Home Screen on iOS8 device
+
+*  UE-21123	PowerVR Galaxy S4 (GT-I9500) shows black scene when MobileHDR enabled
+
+*  UE-21367	Submitting to App store for IOS 9 leads to validation errors
+
+*  UE-20940	Effects.umap causes Standalone and Mobile Preview PIE to Crash
+
+*  UE-21126	Foliage culling broken on XboxOne
+
+*  UE-21185	Update buttons on CrashReporter window to encourage submission
+
