Documentation for build tools

diff --git a/Engine/Documentation/Source/Programming/UnrealBuildSystem/ModuleFiles/ModuleFilesProperties/ModuleFilesProperties.INT.udn b/Engine/Documentation/Source/Programming/UnrealBuildSystem/ModuleFiles/ModuleFilesProperties/ModuleFilesProperties.INT.udn
--- a/Engine/Documentation/Source/Programming/UnrealBuildSystem/ModuleFiles/ModuleFilesProperties/ModuleFilesProperties.INT.udn
+++ b/Engine/Documentation/Source/Programming/UnrealBuildSystem/ModuleFiles/ModuleFilesProperties/ModuleFilesProperties.INT.udn
@@ -89,4 +89,6 @@ $ AdditionalPropertiesForReceipt (List&lt;ReceiptProperty&gt;): List of addition
 
 $ PrecompileForTargets (PrecompileTargetsType): Which targets this module should be precompiled for
 
+$ ExternalDependencies (List&lt;String&gt;): External files which invalidate the makefile if modified. Relative paths are resolved relative to the .build.cs file.
+
 
diff --git a/Engine/Documentation/Source/Programming/UnrealBuildSystem/TargetFiles/TargetFilesProperties/TargetFilesProperties.INT.udn b/Engine/Documentation/Source/Programming/UnrealBuildSystem/TargetFiles/TargetFilesProperties/TargetFilesProperties.INT.udn
--- a/Engine/Documentation/Source/Programming/UnrealBuildSystem/TargetFiles/TargetFilesProperties/TargetFilesProperties.INT.udn
+++ b/Engine/Documentation/Source/Programming/UnrealBuildSystem/TargetFiles/TargetFilesProperties/TargetFilesProperties.INT.udn
@@ -42,6 +42,8 @@ $ PakSigningKeysFile (String): Path to the set of pak signing keys to embed in t
 
 $ SolutionDirectory (String): Allows a Program Target to specify it's own solution folder path.
 
+$ bBuildInSolutionByDefault (Nullable&lt;Boolean&gt;): Whether the target should be included in the default solution build configuration
+
 $ bOutputToEngineBinaries (Boolean): Output the executable to the engine binaries folder.
 
 $ bShouldCompileAsDLL (Boolean): Whether this target should be compiled as a DLL.  Requires LinkType to be set to TargetLinkType.Monolithic.
