Documentation for build tools

diff --git a/Engine/Documentation/Source/Programming/Development/BuildGraph/ScriptAnatomy/Tasks/PredefinedTasks/BuildGraphScriptPredefinedTasks.INT.udn b/Engine/Documentation/Source/Programming/Development/BuildGraph/ScriptAnatomy/Tasks/PredefinedTasks/BuildGraphScriptPredefinedTasks.INT.udn
--- a/Engine/Documentation/Source/Programming/Development/BuildGraph/ScriptAnatomy/Tasks/PredefinedTasks/BuildGraphScriptPredefinedTasks.INT.udn
+++ b/Engine/Documentation/Source/Programming/Development/BuildGraph/ScriptAnatomy/Tasks/PredefinedTasks/BuildGraphScriptPredefinedTasks.INT.udn
@@ -44,14 +44,15 @@ Spawns the editor to run a commandlet.
 
 Compiles a target with UnrealBuildTool.
 
-|               |                           |          |                                                                                   |
-| ------------- | ------------------------- | -------- | --------------------------------------------------------------------------------- |
-| Target        | String                    | Required | The target to compile                                                             |
-| Configuration | UnrealTargetConfiguration | Required | The configuration to compile                                                      |
-| Platform      | UnrealTargetPlatform      | Required | The platform to compile for                                                       |
-| Arguments     | String                    | Optional | Additional arguments for UnrealBuildTool                                          |
-| Tag           | Tag List                  | Optional | Tag to be applied to build products of this task                                  |
-| If            | Condition                 | Optional | Whether to execute this task. It is ignored if this condition evaluates to false. |
+|                       |                           |          |                                                                                   |
+| --------------------- | ------------------------- | -------- | --------------------------------------------------------------------------------- |
+| Target                | String                    | Required | The target to compile                                                             |
+| Configuration         | UnrealTargetConfiguration | Required | The configuration to compile                                                      |
+| Platform              | UnrealTargetPlatform      | Required | The platform to compile for                                                       |
+| Arguments             | String                    | Optional | Additional arguments for UnrealBuildTool                                          |
+| AllowParallelExecutor | Boolean                   | Optional | Whether to allow using the parallel executor for this compile                     |
+| Tag                   | Tag List                  | Optional | Tag to be applied to build products of this task                                  |
+| If                    | Condition                 | Optional | Whether to execute this task. It is ignored if this condition evaluates to false. |
 
 ### Cook
 
diff --git a/Engine/Documentation/Source/Programming/UnrealBuildSystem/Configuration/BuildConfigProperties/BuildConfigProperties.INT.udn b/Engine/Documentation/Source/Programming/UnrealBuildSystem/Configuration/BuildConfigProperties/BuildConfigProperties.INT.udn
--- a/Engine/Documentation/Source/Programming/UnrealBuildSystem/Configuration/BuildConfigProperties/BuildConfigProperties.INT.udn
+++ b/Engine/Documentation/Source/Programming/UnrealBuildSystem/Configuration/BuildConfigProperties/BuildConfigProperties.INT.udn
@@ -76,8 +76,6 @@ $ bUseFastPDBLinking : Whether to use the :FASTLINK option when building with /D
 
 $ bCreateMapFile : Outputs a map file as part of the build.
 
-$ bEnableCodeAnalysis : Enables code analysis mode.  Currently, this has specific requirements.  It only works on Windows platform with the MSVC compiler.  Also, it requires a version of the compiler that supports the /analyze option, such as Visual Studio 2013.
-
 $ bCheckSystemHeadersForModification : Whether headers in system paths should be checked for modification when determining outdated actions.
 
 $ bFlushBuildDirOnRemoteMac : Whether to clean Builds directory on a remote Mac before building.
@@ -184,6 +182,8 @@ $ ProcessorCountMultiplier : Processor count multiplier for local execution. Can
 
 $ MaxProcessorCount : Maximum processor count for local execution.
 
+### ProjectFileGenerator
+
 ### VCProjectFileGenerator
 
 $ Version : The version of Visual Studio to generate project files for.
diff --git a/Engine/Documentation/Source/Programming/UnrealBuildSystem/ModuleFiles/ModuleFilesProperties/ModuleFilesProperties.INT.udn b/Engine/Documentation/Source/Programming/UnrealBuildSystem/ModuleFiles/ModuleFilesProperties/ModuleFilesProperties.INT.udn
--- a/Engine/Documentation/Source/Programming/UnrealBuildSystem/ModuleFiles/ModuleFilesProperties/ModuleFilesProperties.INT.udn
+++ b/Engine/Documentation/Source/Programming/UnrealBuildSystem/ModuleFiles/ModuleFilesProperties/ModuleFilesProperties.INT.udn
@@ -45,6 +45,8 @@ $ WhitelistRestrictedFolders (List&lt;String&gt;): List of folders which are whi
 
 $ bEnforceIWYU (Boolean): Enforce "include what you use" rules when PCHUsage is set to ExplicitOrSharedPCH; warns when monolithic headers (Engine.h, UnrealEd.h, etc...) are used, and checks that source files include their matching header first.
 
+$ bAddDefaultIncludePaths (Boolean): Whether to add all the default include paths to the module (eg. the Source/Classes folder, subfolders under Source/Public).
+
 $ PublicIncludePathModuleNames (List&lt;String&gt;): List of modules names (no path needed) with header files that our module's public headers needs access to, but we don't need to "import" or link against.
 
 $ PublicDependencyModuleNames (List&lt;String&gt;): List of public dependency module names (no path needed) (automatically does the private/public include). These are modules that are required by our public source files.
diff --git a/Engine/Documentation/Source/Programming/UnrealBuildSystem/TargetFiles/TargetFilesProperties/TargetFilesProperties.INT.udn b/Engine/Documentation/Source/Programming/UnrealBuildSystem/TargetFiles/TargetFilesProperties/TargetFilesProperties.INT.udn
--- a/Engine/Documentation/Source/Programming/UnrealBuildSystem/TargetFiles/TargetFilesProperties/TargetFilesProperties.INT.udn
+++ b/Engine/Documentation/Source/Programming/UnrealBuildSystem/TargetFiles/TargetFilesProperties/TargetFilesProperties.INT.udn
@@ -206,8 +206,6 @@ $ bUseFastPDBLinking (Boolean): Whether to use the :FASTLINK option when buildin
 
 $ bCreateMapFile (Boolean): Outputs a map file as part of the build.
 
-$ bEnableCodeAnalysis (Boolean): Enables code analysis mode.  Currently, this has specific requirements.  It only works on Windows platform with the MSVC compiler.  Also, it requires a version of the compiler that supports the /analyze option, such as Visual Studio 2013.
-
 $ BundleVersion (String): Bundle version for Mac apps.
 
 $ bDeployAfterCompile (Boolean): Whether to deploy the executable after compilation on platforms that require deployment.
