Documentation for build tools

diff --git a/Engine/Documentation/Source/Programming/UnrealBuildSystem/ModuleFiles/ModuleFilesProperties/ModuleFilesProperties.INT.udn b/Engine/Documentation/Source/Programming/UnrealBuildSystem/ModuleFiles/ModuleFilesProperties/ModuleFilesProperties.INT.udn
--- a/Engine/Documentation/Source/Programming/UnrealBuildSystem/ModuleFiles/ModuleFilesProperties/ModuleFilesProperties.INT.udn
+++ b/Engine/Documentation/Source/Programming/UnrealBuildSystem/ModuleFiles/ModuleFilesProperties/ModuleFilesProperties.INT.udn
@@ -15,6 +15,8 @@ $ PrivatePCHHeaderFile (String): Explicit private PCH for this module. Implies t
 
 $ SharedPCHHeaderFile (String): Header file name for a shared PCH provided by this module.  Must be a valid relative path to a public C++ header file. This should only be set for header files that are included by a significant number of other C++ modules.
 
+$ ShortName (String): Specifies an alternate name for intermediate directories and files for intermediates of this module. Useful when hitting path length limitations.
+
 $ PCHUsage (PCHUsageMode): Precompiled header usage for this module
 
 $ bUseRTTI (Boolean): Use run time type information
