Documentation for build tools

diff --git a/Engine/Documentation/Source/Programming/UnrealBuildSystem/TargetFiles/TargetFilesProperties/TargetFilesProperties.INT.udn b/Engine/Documentation/Source/Programming/UnrealBuildSystem/TargetFiles/TargetFilesProperties/TargetFilesProperties.INT.udn
--- a/Engine/Documentation/Source/Programming/UnrealBuildSystem/TargetFiles/TargetFilesProperties/TargetFilesProperties.INT.udn
+++ b/Engine/Documentation/Source/Programming/UnrealBuildSystem/TargetFiles/TargetFilesProperties/TargetFilesProperties.INT.udn
@@ -58,8 +58,6 @@ $ bCompileAPEX (Boolean): Whether to include PhysX APEX support.
 
 $ bCompileNvCloth (Boolean): Whether to include NvCloth.
 
-$ bRuntimePhysicsCooking (Boolean): Whether to allow runtime cooking of physics.
-
 $ bCompileBox2D (Boolean): Whether to include Box2D support.
 
 $ bCompileICU (Boolean): Whether to include ICU unicode/i18n support in Core.
