Documentation for build tools

diff --git a/Engine/Documentation/Source/Programming/Development/BuildGraph/ScriptAnatomy/Tasks/PredefinedTasks/BuildGraphScriptPredefinedTasks.INT.udn b/Engine/Documentation/Source/Programming/Development/BuildGraph/ScriptAnatomy/Tasks/PredefinedTasks/BuildGraphScriptPredefinedTasks.INT.udn
--- a/Engine/Documentation/Source/Programming/Development/BuildGraph/ScriptAnatomy/Tasks/PredefinedTasks/BuildGraphScriptPredefinedTasks.INT.udn
+++ b/Engine/Documentation/Source/Programming/Development/BuildGraph/ScriptAnatomy/Tasks/PredefinedTasks/BuildGraphScriptPredefinedTasks.INT.udn
@@ -50,6 +50,7 @@ Compiles a target with UnrealBuildTool.
 | Configuration         | UnrealTargetConfiguration | Required | The configuration to compile                                                      |
 | Platform              | UnrealTargetPlatform      | Required | The platform to compile for                                                       |
 | Arguments             | String                    | Optional | Additional arguments for UnrealBuildTool                                          |
+| AllowXGE              | Boolean                   | Optional | Whether to allow using XGE for compilation                                        |
 | AllowParallelExecutor | Boolean                   | Optional | Whether to allow using the parallel executor for this compile                     |
 | Tag                   | Tag List                  | Optional | Tag to be applied to build products of this task                                  |
 | If                    | Condition                 | Optional | Whether to execute this task. It is ignored if this condition evaluates to false. |
