Documentation for build tools

diff --git a/Engine/Documentation/Source/Programming/UnrealBuildSystem/TargetFiles/TargetFilesProperties/TargetFilesProperties.INT.udn b/Engine/Documentation/Source/Programming/UnrealBuildSystem/TargetFiles/TargetFilesProperties/TargetFilesProperties.INT.udn
--- a/Engine/Documentation/Source/Programming/UnrealBuildSystem/TargetFiles/TargetFilesProperties/TargetFilesProperties.INT.udn
+++ b/Engine/Documentation/Source/Programming/UnrealBuildSystem/TargetFiles/TargetFilesProperties/TargetFilesProperties.INT.udn
@@ -82,6 +82,8 @@ $ bCompileSimplygonSSF (Boolean): Whether we should compile in support for Simpl
 
 $ bCompileLeanAndMeanUE (Boolean): Whether to compile lean and mean version of UE.
 
+$ bUseCacheFreedOSAllocs (Boolean): Whether to utilize cache freed OS allocs with MallocBinned
+
 $ bCompileAgainstEngine (Boolean): Enabled for all builds that include the engine project.  Disabled only when building standalone apps that only link with Core.
 
 $ bCompileAgainstCoreUObject (Boolean): Enabled for all builds that include the CoreUObject project.  Disabled only when building standalone apps that only link with Core.
