Documentation for build tools

diff --git a/Engine/Documentation/Source/Programming/Development/BuildGraph/ScriptAnatomy/Tasks/PredefinedTasks/BuildGraphScriptPredefinedTasks.INT.udn b/Engine/Documentation/Source/Programming/Development/BuildGraph/ScriptAnatomy/Tasks/PredefinedTasks/BuildGraphScriptPredefinedTasks.INT.udn
--- a/Engine/Documentation/Source/Programming/Development/BuildGraph/ScriptAnatomy/Tasks/PredefinedTasks/BuildGraphScriptPredefinedTasks.INT.udn
+++ b/Engine/Documentation/Source/Programming/Development/BuildGraph/ScriptAnatomy/Tasks/PredefinedTasks/BuildGraphScriptPredefinedTasks.INT.udn
@@ -162,17 +162,18 @@ Renames a file, or group of files.
 
 Updates the local version files (Engine/Source/Runtime/Launch/Resources/Version.h, Engine/Build/Build.version, and Engine/Source/Programs/DotNETCommon/Metadata.cs) with the given version information.
 
-|                  |           |          |                                                                                                                           |
-| ---------------- | --------- | -------- | ------------------------------------------------------------------------------------------------------------------------- |
-| Change           | Int32     | Required | The changelist to set in the version files                                                                                |
-| CompatibleChange | Int32     | Optional | The engine compatible changelist to set in the version files                                                              |
-| Branch           | String    | Required | The branch string                                                                                                         |
-| Build            | String    | Optional | The build version string                                                                                                  |
-| Licensee         | Boolean   | Optional | Whether to set the IS_LICENSEE_VERSION flag to true                                                                       |
-| Promoted         | Boolean   | Optional | Whether to set the ENGINE_IS_PROMOTED_BUILD flag to true                                                                  |
-| SkipWrite        | Boolean   | Optional | If set, don't actually write to the files - just return the version files that would be updated. Useful for local builds. |
-| Tag              | Tag List  | Optional | Tag to be applied to build products of this task                                                                          |
-| If               | Condition | Optional | Whether to execute this task. It is ignored if this condition evaluates to false.                                         |
+|                  |           |          |                                                                                                                                                                  |
+| ---------------- | --------- | -------- | ---------------------------------------------------------------------------------------------------------------------------------------------------------------- |
+| Change           | Int32     | Required | The changelist to set in the version files                                                                                                                       |
+| CompatibleChange | Int32     | Optional | The engine compatible changelist to set in the version files                                                                                                     |
+| Branch           | String    | Required | The branch string                                                                                                                                                |
+| Build            | String    | Optional | The build version string                                                                                                                                         |
+| Licensee         | Boolean   | Optional | Whether to set the IS_LICENSEE_VERSION flag to true                                                                                                              |
+| Promoted         | Boolean   | Optional | Whether to set the ENGINE_IS_PROMOTED_BUILD flag to true                                                                                                         |
+| SkipHeader       | Boolean   | Optional | If set, don't update Version.h; just update data files (eg. Build.version, Metadata.cs). The editor can read the build changelist at runtime from build.version. |
+| SkipWrite        | Boolean   | Optional | If set, don't actually write to the files - just return the version files that would be updated. Useful for local builds.                                        |
+| Tag              | Tag List  | Optional | Tag to be applied to build products of this task                                                                                                                 |
+| If               | Condition | Optional | Whether to execute this task. It is ignored if this condition evaluates to false.                                                                                |
 
 ### Sign
 
diff --git a/Engine/Documentation/Source/Programming/UnrealBuildSystem/Configuration/BuildConfigProperties/BuildConfigProperties.INT.udn b/Engine/Documentation/Source/Programming/UnrealBuildSystem/Configuration/BuildConfigProperties/BuildConfigProperties.INT.udn
--- a/Engine/Documentation/Source/Programming/UnrealBuildSystem/Configuration/BuildConfigProperties/BuildConfigProperties.INT.udn
+++ b/Engine/Documentation/Source/Programming/UnrealBuildSystem/Configuration/BuildConfigProperties/BuildConfigProperties.INT.udn
@@ -120,8 +120,6 @@ $ bUseUBTMakefiles : Enables support for very fast iterative builds by caching t
     * You can force regeneration of the 'UBT Makefile' by passing the '-gather' argument, or simply regenerating project files 
     * This also enables the fast include file dependency scanning and caching system that allows Unreal Build Tool to detect out of date dependencies very quickly.  When enabled, a deep C++ include graph does not have to be generated, and instead we only scan and cache indirect includes for after a dependent build product was already found to be out of date.  During the next build, we'll load those cached indirect includes and check for outdatedness.
 
-$ bUseUHTMakefiles : Enables support for fast UHT parsing by caching results of previous UHT runs. If a module or *.uhtmanifest gets changed, all modules up to first changed one get loaded from makefile and the rest is parsed regularly. This feature is new and can have issues that weren't discovered in isolated testing so disabling by default.
-
 $ bAllowDistcc : Whether DMUCS/Distcc may be used. Distcc requires some setup - so by default disable it so we don't break local or remote building
 
 $ bIgnoreJunk : Whether to skip checking for files identified by the junk manifest
diff --git a/Engine/Documentation/Source/Programming/UnrealBuildSystem/TargetFiles/TargetFilesProperties/TargetFilesProperties.INT.udn b/Engine/Documentation/Source/Programming/UnrealBuildSystem/TargetFiles/TargetFilesProperties/TargetFilesProperties.INT.udn
--- a/Engine/Documentation/Source/Programming/UnrealBuildSystem/TargetFiles/TargetFilesProperties/TargetFilesProperties.INT.udn
+++ b/Engine/Documentation/Source/Programming/UnrealBuildSystem/TargetFiles/TargetFilesProperties/TargetFilesProperties.INT.udn
@@ -104,6 +104,8 @@ $ bCompileWithStatsWithoutEngine (Boolean): Whether to include stats support eve
 
 $ bCompileWithPluginSupport (Boolean): Whether to include plugin support.
 
+$ bIncludePluginsForTargetPlatforms (Boolean): Whether to allow plugins which support all target platforms.
+
 $ bWithPerfCounters (Boolean): Whether to include PerfCounters support.
 
 $ bUseLoggingInShipping (Boolean): Whether to turn on logging for test/shipping builds.
@@ -234,6 +236,10 @@ $ ExtraModuleNames (List&lt;String&gt;): List of additional modules to be compil
 
 $ BuildEnvironment (TargetBuildEnvironment): Specifies the build environment for this target. See TargetBuildEnvironment for more infomation on the available options.
 
+$ PreBuildSteps (List&lt;String&gt;): Specifies a list of steps which should be executed before this target is built, in the context of the host platform's shell. The following variables will be expanded before execution: $(EngineDir), $(ProjectDir), $(TargetName), $(TargetPlatform), $(TargetConfiguration), $(TargetType), $(ProjectFile).
+
+$ PostBuildSteps (List&lt;String&gt;): Specifies a list of steps which should be executed after this target is built, in the context of the host platform's shell. The following variables will be expanded before execution: $(EngineDir), $(ProjectDir), $(TargetName), $(TargetPlatform), $(TargetConfiguration), $(TargetType), $(ProjectFile).
+
 $ OverrideExecutableFileExtension (String): Can be set to override the file extension of the executable file (normally .exe or .dll on Windows, for example)
 
 
