Copying //UE4/Dev-Framework to //UE4/Dev-Main (Source: //UE4/Dev-Framework @ 3582324)
#rnx

============================
  MAJOR FEATURES & CHANGES
============================

Change 3431439 by [REDACTED_PERSON]

	Editor only subobjects shouldn't exist in PIE world

Change 3457323 by [REDACTED_PERSON]

	Undo CL# 3431439 and once again allow (incorrectly) for editor only objects to exist in a PIE world

Change 3499927 by [REDACTED_PERSON]

	UField::Serialize no longer serialize's its next ptr,  UStruct::Serialize serializes all Children properties instead. This resolves a hard circular dependency between function libraries that EDL detected. It was resolved in an ad hoc way by the old linker

Change 3502939 by [REDACTED_PERSON]

	Back out changelist 3499927

Change 3522783 by [REDACTED_PERSON]

	#ue4 - Imported new simple collision for Engine/Content/BasicShaps/Cylinder.uasset which is a single convex shape (rather than being 4 shapes as before).

Change 3544641 by [REDACTED_PERSON]

	Remove conditional that is no longer needed, replace with ensure. It is unsafe to change CDO names

Change 3544645 by [REDACTED_PERSON]

	In addition to marking nodes as not transient, FBlueprintEditor::ExpandNode needs to mark them as transactional

Change 3545023 by [REDACTED_PERSON]

	Properly encapsulate FPinDeletionQueue
	Fix ensure during deletion of split pins when not clearing links
	Fix split pins able to end up in delete queue twice during undo/redo

Change 3545025 by [REDACTED_PERSON]

	Properly allow changing the pin type from a struct that is split on the node

Change 3545455 by [REDACTED_PERSON]

	Fix issue where combined streamable handles could be freed before their complete callback is called if nothing external referenced them
	Copy of CL#3544474

Change 3545456 by [REDACTED_PERSON]

	Allow PrimaryAssets to update their AssetData based on in-memory changes when launching 'Standalone Game' and 'Mobile Preview' from the editor. As it was, changes could be detected and propagated through UPrimaryDataAsset::PostLoad, but the changes would always reapply whatever already exists in the AssetRegistry. The primary use-case for this: making AssetBundle tag changes and allowing them to propagate without resaving/recooking all affected assets.
	Copy of CL #3544374

Change 3545547 by [REDACTED_PERSON]

	CIS Fix

Change 3545568 by [REDACTED_PERSON]

	PR #3758: Fixing a comment typo from Transistion to Transition (Contributed by [REDACTED_PERSON])

Change 3545582 by [REDACTED_PERSON]

	Blueprints: Prevent duplicate messages from being added to the compiler results log (fixes a crash when resizing the results log while a math expression node has an error)
	Blueprints: Fixed the tooltip of math expression nodes not showing up, and error messages getting cleared on subsequent compiles
	[Duplicating fixes for UE-47491 and UE-47512 from 4.17 to Dev-Framework]

Change 3546528 by [REDACTED_PERSON]

	Fix crash when a map's key type has changed but value has not, it was passing the UStruct defaults in when serialize was expecting the default instance, so pass null instead as we don't have the instance

Change 3546544 by [REDACTED_PERSON]

	Fix split pin restoration logic to deal with wildcards and variations in const/refness

Change 3546551 by [REDACTED_PERSON]

	Don't crash if the struct type is missing for whatever reason

Change 3547152 by [REDACTED_PERSON]

	Fix array exporting so you don't end up getting none instead of defaults

Change 3547438 by [REDACTED_PERSON]

	Fix split pins on class defaults
	Don't cause a structural change when reapplying a split pin as part of node reconstruction

Change 3547501 by [REDACTED_PERSON]

	Fix ensure, it's valid to pass a null path for a dynamic asset

Change 3551185 by [REDACTED_PERSON]

Change 3551723 by [REDACTED_PERSON]

	Improve UI for displaying actor soft references. Add an error/warning icon if the cross level reference is broken or unloaded, and fix various display and copy/paste behaviors

Change 3553216 by [REDACTED_PERSON]

	- Nativized UDS now support external asset dependencies and will construct their own linker import tables on load.

	Change summary:
	- Modified FCompactBlueprintDependencyData and FFakeImportTableHelper to be more relevant to UStruct and not just UClass-derivative types.
	- Modified FBlueprintDependencyData to accept a single FCompactBlueprintDependencyData struct rather than its individual fields.
	- Modified FBlueprintCompilerCppBackendBase::GenerateCodeFromStruct() to emit dependency assignment and static type registration functions for nativized UStruct types.
	- Modified FBlueprintNativeCodeGenModule::FStatePerPlatform to include an array for tracking UDS assets that need to be converted during the nativization pass at cook time.
	- Modified FBlueprintNativeCodeGenModule::GenerateFullyConvertedClasses() to generate nativized code for UDS assets. This ensures that UDS conversion falls under the same scope as BPGC conversion, so that they both feed into the same NativizationSummary context for asset dependency tracking (i.e. since we only have a single global dependency table in the codegen). Also modified InitializeForRerunDebugOnly() to do the same.
	- Modified FBlueprintNativeCodeGenModule::Convert() to defer UDS conversion so that it's done at the same time as BPGC conversion (see note above).
	- Modified FEmitDefaultValueHelper::AddStaticFunctionsForDependencies() to include support for UStruct types and to conform to changes made to FCompactBlueprintDependencyData.
	- Modified FEmitDefaultValueHelper::AddRegisterHelper() to include support for UStruct types.
	- Modified FBlueprintNativeCodeGenModule::FindReplacedClassForObject() to ensure that converted User-Defined Enum types are switched to a UEnumProperty at package save time so that any import tables contain the proper class. This is necessary because we nativize User-Defined Enum types as 'enum class' types, and UHT will generate code for those as a UEnumProperty with an "underlying" property. However, non-nativized User-Defined Enum types are referenced as a UByteProperty with a UEnum reference, so we have to fix up all the import tables before saving. Otherwise, EDL will assert on load (see UE-47519).

Change 3553301 by [REDACTED_PERSON]

	Fix ensure when passing literal None to SoftObjectPath, it now treats them as empty instead

Change 3553631 by [REDACTED_PERSON]

	UField::Serialize no longer serialize's its next ptr,  UStruct::Serialize serializes all Children properties instead. This resolves a hard circular dependency between function libraries that EDL detected. It was resolved in an ad hoc way by the old linker. This change was originally submitted in 3499927, but it was incorrectly clearing the UField::Next pointer in UField::Serialize.

Change 3553799 by [REDACTED_PERSON]

	Fix issue where calling WaitUntilComplete on a combined handle with Stalled children wouldn't work properly. It now forces all stalled children to start immediately. I also added a warning log when this happens and an ensure if somehow the force didn't work
	Copy of CL #3553781

Change 3553896 by [REDACTED_PERSON]

	Blueprints: Allow the autowiring logic to better break and replace existing connections when made (e.g., when dragging a variable onto a compatible pin with an existing connection, break the old connection to allow the new connection to be made)

Change 3553897 by [REDACTED_PERSON]

	Blueprints: Adjust search query when doing 'Find References' on variables from My Blueprints so that bound event nodes show up for components and widgets

Change 3553898 by [REDACTED_PERSON]

	Blueprints: Add the name of the variable directly in the get/set menu options (when dragging from My Blueprints into the graph)

Change 3553909 by [REDACTED_PERSON]

	Blueprints: Added the full name of the type, container type (and value type for maps) to the tooltips for the type picker elements, so long names can be read in full

Change 3554517 by [REDACTED_PERSON]

	Blueprints: Added an option to disable the comment bubble on comment boxes that appears when zoomed out

Change 3554664 by [REDACTED_PERSON]

	Editor: Renamed "Find in CB" command to "Browse" and renamed "Search" (in BP) to "Find", so terminology is consistent and keyboard shortcuts make sense (Ctrl+B for Browse, Ctrl+F for find, not using the term Search anywhere)

Change 3554831 by [REDACTED_PERSON]

	Non editor build fix

Change 3554834 by [REDACTED_PERSON]

	Actor bound event related warnings now show up in blueprint status when opening level blueprint for first time, improved warning message. Removed unused delegate and return value from FixLevelScriptActorBindings. Can now pass raw strings to blueprint results log (no need for Printf, although padding is not great), UClasses in compiler results log will open the generated blueprint when clicked on

Change 3556157 by [REDACTED_PERSON]

	Convert LevelSequenceBindingReference to use FSoftObjectPath so it works properly with redirectors and fixups

Change 3557775 by [REDACTED_PERSON]

	Blueprints: Fixed swapped transaction messages when converting a cast node between pure and impure

Change 3557777 by [REDACTED_PERSON]

	Blueprints: Allow 'Goto Definition' and 'Find References' to be used while stopped at a breakpoint
	PR #3774: Expose GotoFunctionDefinition during BP debugging (Contributed by [REDACTED_PERSON])

Change 3560510 by [REDACTED_PERSON]

	Blueprints: Add support for 'goto definition' on Create Event nodes when the Object pin is not hooked up

Change 3560563 by [REDACTED_PERSON]

	Blueprints: Disallow converting a variable get node to impure/back when debugging (no graph mutating operations should be allowed)

Change 3561443 by [REDACTED_PERSON]

	Restore code to support gc.DumpPoolStats, was accidentally removed when FGCArrayPool was moved to a header.
	Clean up comments around Cleanup function, the functionality to trim the memory pools was integrated into ClearWeakReferences on a prior change

Change 3561658 by [REDACTED_PERSON]

	Blueprints: Refactored 'Goto Definition' so there is no per-class logic in the Blueprint editor or schema any more; any node can opt in individually
	- Added a key binding for Goto Definition (Alt+G)
	- Added a key binding for Find References (Shift+Alt+F)
	- Collapsed 'Goto Code Definition' for variables and functions into the same path, so there aren't separate menu items and commands
	- Added new methods CanJumpToDefinition and JumpToDefinition to UEdGraphNode, the default behavior for UK2Node subclasses is to call GetJumpTargetForDoubleClick and call into FKismetEditorUtilities::BringKismetToFocusAttentionOnObject
	- Going to a native function now goes thru a better code path that will actually put the source code editor on the function definition, rather than just opening the file containing the definition

Change 3562291 by [REDACTED_PERSON]

	Fix issue where calling FSoftObjectPtr::Get during a package save would result in that ptr being forever marked broken, because the ResolveObject fails during save. It's too risky to change that behavior, so change it so the TagAtLastTest doesn't get updated in that case

Change 3562292 by [REDACTED_PERSON]

	When deleting actors that are soft referenced by actor/sequencer it will now warn the same way it does for level script. Added IAssetTools::FindSoftReferencesToObject to perform this search
	Change it so saving a non-primary world does not result it being dirtied due to the temporary physics scene fixup
	Fix issue where the actor name was shown incorrectly in the SSCS tree for actor instances, which meant that if you renamed it you would end up renaming it to the BP's name

Change 3564814 by [REDACTED_PERSON]

Change 3566707 by [REDACTED_PERSON]

	Remove unused code, UClassGenerateCDODuplicatesForHotReload was attempting to create a CDO with a special name, which triggered an ensure. The Duplicated CDO was never used (callign code removed in 3289276 as it was a waste of cycles)

Change 3566717 by [REDACTED_PERSON]

	Core: Remove all defintions that contain "_API" from the command line when compiling .rc files (they do not support repsonse files and a too-long command line will fail the compile)

Change 3566771 by [REDACTED_PERSON]

	Editor: Fixing deprecation warning

Change 3567023 by [REDACTED_PERSON]

	Blueprints: Change various TArray<> uses to TSet<> throughout name validation and related code to enable it to scale better to high component or variable counts
	Adapted from PR #3708: Fast construction of bp (Contributed by [REDACTED_PERSON])

Change 3567304 by [REDACTED_PERSON]

	Add bCheckRecursive option to IsEditorOnlyObject that is enabled by default and will check outer/archetype/class.
	This is needed for places that call this function from outside of SavePackage.cpp when the editor only mark is set, such as the automation test code

Change 3567398 by [REDACTED_PERSON]

	Fix crash when spawning a widget that has no editor WidgetTree, but does have a cooked widget tree template due to tree inheritance

Change 3567729 by [REDACTED_PERSON]

	Blueprints: Clarified message about "{VariableName} is not blueprint visible" to define what that means "(BlueprintReadOnly or BlueprintReadWrite)"

Change 3567739 by [REDACTED_PERSON]

	Don't crash if PropertyStruct cannot find it's struct. The function half handled this before, but Preload crashes with a null parameter

Change 3567741 by [REDACTED_PERSON]

	Disable optimization for a path test that was crashing in VC2015 in a monolithic build

Change 3568332 by [REDACTED_PERSON]

	Prevented UAIPerceptionSystem::GetCurrent causing a BP error when WorldContextObject is null #UE4

Change 3568676 by [REDACTED_PERSON]

	Blueprints: Allow editing the tooltip of each enum value in a user defined enum

	Known issue: Undo/redo is not currently possible on the tooltip as it is directly stored in package metadata

Change 3569128 by [REDACTED_PERSON]

	Blueprints: Removing the experimental profiler as we won't be returning to it any time soon

Change 3569207 by [REDACTED_PERSON]

	Blueprints: Allow drag-dropping component Blueprint assets into the graph to create Add Component nodes and improved the error message when dragging something that cannot be dropped into an actor Blueprint

Change 3569208 by [REDACTED_PERSON]

	Blueprints: Allow specifying a description for user defined enums (shown in the content browser)

Change 3569209 by [REDACTED_PERSON]

	Editor: Allow adjusting the font size for each individual comment box node in Blueprints and Materials

Change 3570177 by [REDACTED_PERSON]

	Blueprints: Fixed discrepancy between the Structure tab name and the menu option for the tab in the user defined structure editor (now both say Structure Editor)

Change 3570179 by [REDACTED_PERSON]

	Blueprints: Fixed the tooltip of a user defined structure not updating immediately in the content browser after being edited

Change 3570192 by [REDACTED_PERSON]

	Blueprints: Added "(selected)" after the label (in the 'debug filter' dropdown in the Blueprint editor) for actors that are selected in the level editor, which should make it easier to choose the specific actor you want to debug

Change 3571203 by [REDACTED_PERSON]

	Blueprints: Cleanup and refactoring to prepare for turning commented out nodes into an official feature
	- Made EnabledState and bUserSetEnabledState private on UEdGraphNode and added new getters/setters
	- Introduced IsAutomaticallyPlacedGhostNode() and MakeAutomaticallyPlacedGhostNode() to start decoupling the concept from commented out nodes
	- Updated a couple of places that used a hardcoded UberGraphPages[0] into instead editing the most recently interacted with event graph if possible
	- Updated 'is data only blueprint' logic to allow multiple ubergraph pages, as long as they're all empty of non-disabled nodes

Change 3571224 by [REDACTED_PERSON]

	Blueprints: Draw banners on development-only and user disabled nodes (excluding 'ghost' nodes like autogenerated event entries in new BPs)
	Adapted from PR #2701: Differentiate development nodes in BP (Contributed by [REDACTED_PERSON])

Change 3571279 by [REDACTED_PERSON]

	Blueprints: Changed UK2Node::GetPassThroughPin so that only the execution wire on nodes with exactly one input and one output exec wire will have a corresponding pass-through pin (when there is ambiguity in which exec would be used, e.g., with a branch or sequence or timeline node, the subsequent nodes are now effectively disabled as well)

Change 3571282 by [REDACTED_PERSON]

	Blueprints: Fixed the tooltip for dragging a variable onto an inherited category not showing the 'move to category' hint

Change 3571284 by [REDACTED_PERSON]

	Blueprints: Made wires into/out of a user-disabled node draw verly dimly (other than the passthrough exec if it exists)

Change 3571311 by [REDACTED_PERSON]

	Add ActorIteratorFlags which allows overriding which types of actors/levels are iterated by ActorIterator, to allow iterating levels that are not visible.
	All of the iteration logic is now in the base and the children just set different flags, which fixes it so TActorIterator does the same level collection check as FActorIterator

Change 3571313 by [REDACTED_PERSON]

	Several fixes to automation framework to allow it to work better with Cooked builds.
	Change it so the automation test list is a single  message. There is no guarantee on order of message packets, so several tests were being missed each time.

Change 3571485 by mason.seay

	Test map for Make Set bug

Change 3571501 by [REDACTED_PERSON]

	Accidentally undid the UHT fixup for TAssetPtr during my bulk rename

Change 3571531 by [REDACTED_PERSON]

	Fix warning messages

Change 3571591 by [REDACTED_PERSON]

	Blueprints: Made drag-dropping assets into a graph to create a component transactional (allowing the action to be undone)

Change 3572938 by [REDACTED_PERSON]

	Blueprints: Fixed a typo in a set function comment

Change 3572941 by [REDACTED_PERSON]

	Blueprints: Made the compact node title for cross and dot product the words cross and dot rather than hard to see . and x symbols

Change 3574816 by mason.seay

	Renamed asset to better reflect name of object reference

Change 3574985 by mason.seay

	Updated comments and string outputs to list Soft Object Reference

Change 3575740 by [REDACTED_PERSON]

Change 3575795 by [REDACTED_PERSON]

Change 3576374 by mason.seay

	Forgot to submit the deleting of a redirector

Change 3576966 by [REDACTED_PERSON]

Change 3577002 by [REDACTED_PERSON]

	Prevent wildcard pins from being connected to exec pins

Change 3577232 by [REDACTED_PERSON]

	Change summary:
	- Modified FNativeClassHeaderGenerator::ExportGeneratedStructBodyMacros() to emit the 'ReplaceConverted' package name for the FCompiledInDeferStruct global associated with the nativized UDS asset in the UHT codegen. This brings nativized UDS to parity with nativized BP class assets (it was likely just an oversight initially).

Change 3577710 by [REDACTED_PERSON]

	Mirror of 3576977:
	Fix for crash when loading cooked uassets that reference functions that are not present

Change 3577723 by [REDACTED_PERSON]

	Prevent deferring of classes that are needed to load subobjects

Change 3577741 by [REDACTED_PERSON]

	Back out changelist 3577723 - causes crash when starting QAGame in Dev-Framework - not in Release-4.17

Change 3578938 by [REDACTED_PERSON]

	Fix it so the redirector cleanup code on map change will not be stopped by non-standalone top level objects, which could be left behind by issues in other systems

Change 3578947 by [REDACTED_PERSON]

	(4.17) Properly expose members of DialogueContext to blueprints

Change 3578952 by [REDACTED_PERSON]

	Fix ensure where ParentHandles on a StreamableHandle could be modified while iterating

Change 3579315 by mason.seay

	Test map for Make Container nodes

Change 3579600 by [REDACTED_PERSON]

	Disable window test on non-desktop platforms as they cannot be resized post launch

Change 3579601 by [REDACTED_PERSON]

Change 3579713 by [REDACTED_PERSON]

	Prevent crashes when bluepints implement an interface that was deleted

Change 3579719 by [REDACTED_PERSON]

	Fix two compilation manager issues: Make sure CDOs are not renamed under a UClass, because they will be considered as possible subobjects, which has bad side effects and make sure that we update references even on empty functions, so that empty UFunctions are not left with references to REINST data

Change 3579745 by [REDACTED_PERSON]

	Blueprints: Improve categorization and reordering support in 'My Blueprints'
	- Allow drag-dropping functions, macros, delegates, etc... to reorder them within a category or to change categories (bringing them to parity with variables)
	- Make category ordering on all categories sticky so you can reorder categories (the relative ordering will be the same within different sections like variables and functions)
	- Added hover cues when drag dropping some items that were missing them (e.g., event dispatchers)
	- Added support for renaming categories using F2

	Known issues (none are regressions):
	- Timelines cannot be moved to other categories or reordered
	- Renaming a nested category will result in it becoming a top level category (discarding the parent category chain)
	- Some actions do not support undo

Change 3579795 by [REDACTED_PERSON]

	PR #3867: Fix for not releasing Local Notification Delegate. (Contributed by [REDACTED_PERSON])

Change 3580463 by [REDACTED_PERSON]

	(4.17) Don't crash if calling PostEditUndo on an Actor in the transient package

Change 3581073 by [REDACTED_PERSON]

	Make UK2Node_SpawnActorFromClass inherit from K2Node_ConstructObjectFromClass and eliminate duplicate code.
	Correctly reconnect split pins when changing class on create widget, construct object, and spawn actor nodes

Change 3581156 by [REDACTED_PERSON]

	Also fix bug where the category of parent pins was being set incorrectly while changing variable type, we only want to override type for wildcard pins

Change 3581473 by [REDACTED_PERSON]

	Properly turn off optimization for PS4 test

Change 3582094 by [REDACTED_PERSON]

	Fix anim nodes not navigating to their graph on double click

Change 3582157 by [REDACTED_PERSON]

	Fix double-clicking on animation asset nodes not opening the asset editors

Change 3582289 by [REDACTED_PERSON]

	(4.17) Don't crash when adding a streaming level that's already in the level

Change 3545435 by [REDACTED_PERSON]

	FStringAssetReference -> FSoftObjectPath
	FStringClassReference -> FSoftClassPath
	TAssetPtr -> TSoftObjectPtr
	TAssetSubclassOf -> TSoftClassPtr
	The old headers are deprecated, and FSoftClassPath is now in the same header has FSoftObjectPath.
	This change increments the UE4 version because FSoftObjectPaths are now stored as a name + substring instead of one giant name, which in practice will improve performance and memory while manipulating them. Also the package table of soft referenced packages is now stored as FNames instead of FStrings as these packages names will already be in the name table due to the AssetRegistry
	Remove automatic support for loading Objectpaths starting with engine-ini:, as it was only partially supported and is very outdated. ResolveIniObjectsReference can still be called manually
	Changed TPersistentObjectPtr and TLazyObjectPtr to be structs instead of classes
	Change UnrealHeaderTool to read configuration such as StructsWithNoPrefix out of inis instead of using a hardcoded list. Add support for TypeRedirects, which is used to make the old type names automatically remap to the new ones
	Clean up FRedirectCollector to remove some of the functionality that is no longer used by the cooker, and disable tracking of redirects in standalone -game builds

Change 3567760 by [REDACTED_PERSON]

	Fix it so EngineTest can be cooked by moving some utility functions to EditorTestsUtilityLibrary and adding an EditorFunctionalTest. The core EngineTest module is safely runtime-only now and can run it's tests locally in windows cooked
	Merge FuncTestManager into FunctionalTestModule to avoid confusion with FunctionalTestingManager UObject
	Add EngineTestAssetManager and override the cook function to cook all maps with runtime functional tests
	Change actor merging tests to be editor only, this stops them from cooking
	Several individual tests crash on cooked builds, I started threads with the owners of those

Change 3575737 by [REDACTED_PERSON]

	As part of this change, connecting a -game client to a PIE server and vice versa is much more likely to work. You may still have game-side problems, look at UEditorEngine::NetworkRemapPath for an example of how to do the PIE prefix conversion
	Remove advanced CreateTemporaryCopiesOfLevels option from sequencer capture, it has not been tested in multiple years and does not work with newer sequencer features

Change 3578806 by [REDACTED_PERSON]

	Fix Construct Object not working correctly with split pins.
	Add Construct Object test cases to functional tests.
	Added split pin expose on spawn test cases.

diff --git a/Engine/Config/BaseEngine.ini b/Engine/Config/BaseEngine.ini
--- a/Engine/Config/BaseEngine.ini
+++ b/Engine/Config/BaseEngine.ini
@@ -16,7 +16,6 @@ Protocol=unreal
 Name=Player
 SaveExt=usa
 Port=7777
-PeerPort=7778
 
 [HTTP]
 HttpTimeout=180
@@ -99,7 +98,6 @@ ArrowMaterialName=/Engine/EditorMaterials/GizmoMaterial.GizmoMaterial
 bAllowHostMigration=false
 HostMigrationTimeout=15
 ParticleEventManagerClassPath="/Script/Engine.ParticleEventManager"
-PlayOnConsoleSaveDir=Autosaves
 LightingOnlyBrightness=(R=0.3,G=0.3,B=0.3,A=1.0)
 +ShaderComplexityColors=(R=0.0,G=1.0,B=0.127,A=1.0)
 +ShaderComplexityColors=(R=0.0,G=1.0,B=0.0,A=1.0)
@@ -444,6 +442,16 @@ bUseStreamingPause=false
 ; 4.18
 +StructRedirects=(OldName="TargetReference",NewName="/Script/AnimGraphRuntime.BoneSocketTarget")
 
+; StringAssetReference->SoftObjectPath/AssetPtr->SoftObjectPtr rename
++StructRedirects=(OldName="StringAssetReference",NewName="/Script/CoreUObject.SoftObjectPath")
++StructRedirects=(OldName="StringClassReference",NewName="/Script/CoreUObject.SoftClassPath")
++ClassRedirects=(OldName="AssetObjectProperty",NewName="/Script/CoreUObject.SoftObjectProperty")
++ClassRedirects=(OldName="AssetClassProperty",NewName="/Script/CoreUObject.SoftClassProperty")
++FunctionRedirects=(OldName="MakeStringAssetReference",NewName="/Script/Engine.KismetSystemLibrary.MakeSoftObjectPath")
++PropertyRedirects=(OldName="MakeSoftObjectPath.AssetLongPathname",NewName="PathString")
++FunctionRedirects=(OldName="SetAssetPropertyByName",NewName="/Script/Engine.KismetSystemLibrary.SetSoftObjectPropertyByName")
++FunctionRedirects=(OldName="SetAssetClassPropertyByName",NewName="/Script/Engine.KismetSystemLibrary.SetSoftClassPropertyByName")
+
 [CoreUObject.Metadata]
 ; Note: UnrealHeaderTool should be rerun after making changes to MetadataRedirects to catch any keys specified in class headers
 +MetadataRedirects=(OldKey="K2Protected", NewKey="BlueprintProtected")
@@ -1293,6 +1301,8 @@ bUseAssetRegistryTagsWhitelistInsteadOfBlacklist=false
 +CookedTagsWhitelist=(Class=*,Tag=AssetBundleData)
 +CookedTagsWhitelist=(Class=*,Tag=PrimaryAssetType)
 +CookedTagsWhitelist=(Class=*,Tag=PrimaryAssetName)
++CookedTagsWhitelist=(Class=World,Tag=Tests)
++CookedTagsWhitelist=(Class=World,Tag=TestNames)
 
 ; These are the options for what parts of the asset registry are cooked and read at runtime. These can be overridden per platform in the appropriate PlatformEngine.ini file
 bSerializeAssetRegistry=true
diff --git a/Engine/Documentation/Source/Gameplay/DataDriven/DataDrivenGameplay.INT.udn b/Engine/Documentation/Source/Gameplay/DataDriven/DataDrivenGameplay.INT.udn
--- a/Engine/Documentation/Source/Gameplay/DataDriven/DataDrivenGameplay.INT.udn
+++ b/Engine/Documentation/Source/Gameplay/DataDriven/DataDrivenGameplay.INT.udn
@@ -68,7 +68,7 @@ As an example, here is a definition for level up experience data and the CSV doc
 
 		/** Icon to use for Achivement */
 		UPROPERTY(EditAnywhere, BlueprintReadWrite, Category=LevelUp)
-		TAssetPtr<UTexture> AchievementIcon;
+		TSoftObjectPtr<UTexture> AchievementIcon;
 	};
 
 
@@ -149,5 +149,5 @@ One final note is that all structures and curves returned should not be cached f
 of a function. This is to ensure that if the table is refreshed via re-import, the data changes take effect 
 immediately and that no invalid pointers are accessed.
 
-In the above DataTable example, the asset that is referenced is a lazy loaded asset (TAssetPtr handles this). 
+In the above DataTable example, the asset that is referenced is a lazy loaded asset (TSoftObjectPtr handles this). 
 If the Asset field type was set to UTexture, all of the assets would be loaded whenever the DataTable was loaded.
diff --git a/Engine/Documentation/Source/Programming/Assets/AsyncLoading/AsyncLoading.INT.udn b/Engine/Documentation/Source/Programming/Assets/AsyncLoading/AsyncLoading.INT.udn
--- a/Engine/Documentation/Source/Programming/Assets/AsyncLoading/AsyncLoading.INT.udn
+++ b/Engine/Documentation/Source/Programming/Assets/AsyncLoading/AsyncLoading.INT.udn
@@ -8,13 +8,13 @@ Version: 4.9
 
 There are several new systems in UE4 that make it much easier to asynchronously load asset data, and they replace a lot of the functionality that existed for seekfree content packages in UE3. These new methods work identically in development and with cooked data on devices, so you do not need to maintain two code paths for loading data on demand. There are two general methods you can use to refer to and load data on demand:
 
-## FStringAssetReferences and TAssetPtr
+## FSoftObjectPaths and TSoftObjectPtr
 
-The easiest way to have an artist or designer reference an asset is to create a UProperty of a hard pointer and give it a category. In UE4, if you have a hard UObject pointer property referencing an asset, that asset will be loaded when the object containing the property is loaded (either by being placed in a map, or referenced from something like a gameinfo). If you are not careful, you can end up loading 100% of your assets at game startup time. If you want artists/designers to be able to make references to specific assets using the same UI as a hard pointer, without always loading the referenced asset, use a `FStringAssetReference` or `TAssetPtr`.
+The easiest way to have an artist or designer reference an asset is to create a UProperty of a hard pointer and give it a category. In UE4, if you have a hard UObject pointer property referencing an asset, that asset will be loaded when the object containing the property is loaded (either by being placed in a map, or referenced from something like a gameinfo). If you are not careful, you can end up loading 100% of your assets at game startup time. If you want artists/designers to be able to make references to specific assets using the same UI as a hard pointer, without always loading the referenced asset, use a `FSoftObjectPath` or `TSoftObjectPtr`.
 
-A `FStringAssetReference` is a simple struct that contains a string with the full name of an asset. If you make a property of that type in your class, it shows up in the editor as if it was a `UObject *` property. It also properly handles cooking and redirects, so if you have a StringAssetReference, it is guaranteed to work properly on a device. A `TAssetPtr` is basically a `TWeakObjectPtr` that wraps around a `FStringAssetReference`, and will template to a specific class so you can restrict the editor UI to only allow selecting certain classes. If the referenced asset exists in memory, `TAssetPtr.Get()` will return it. If it does not, you can call `ToStringReference()` to find out the asset that it refers to, load that using the method described below, then call `TAssetPtr.Get()` again to dereference it. 
+A `FSoftObjectPath` is a simple struct that contains a string with the full name of an asset. If you make a property of that type in your class, it shows up in the editor as if it was a `UObject *` property. It also properly handles cooking and redirects, so if you have a SoftObjectPath, it is guaranteed to work properly on a device. A `TSoftObjectPtr` is basically a `TWeakObjectPtr` that wraps around a `FSoftObjectPath`, and will template to a specific class so you can restrict the editor UI to only allow selecting certain classes. If the referenced asset exists in memory, `TSoftObjectPtr.Get()` will return it. If it does not, you can call `ToSoftObjectPath()` to find out the asset that it refers to, load that using the method described below, then call `TSoftObjectPtr.Get()` again to dereference it. 
 
-`TAssetPtrs` and StringAssetReferences are great if an artist or designer is setting up references manually, but if you want to do something like a query to find an asset satisfying certain requirements, without loading all of those assets, you want to use the asset registry and object libraries.
+`TSoftObjectPtrs` and Soft Object Paths are great if an artist or designer is setting up references manually, but if you want to do something like a query to find an asset satisfying certain requirements, without loading all of those assets, you want to use the asset registry and object libraries.
 
 ## The Asset Registry and Object Libraries
 
@@ -50,15 +50,15 @@ In this example, it creates a new object library, associates a base class, then
 		   }
 	}
 
-In this example, it searches the object library for anything that has a `TypeName` field containing "FooType", then returns the first it finds. Once you have that `AssetData`, you can call `ToStringReference()` to convert that to an `FStringAssetReference`, which you can then load asynchronously using the next system:
+In this example, it searches the object library for anything that has a `TypeName` field containing "FooType", then returns the first it finds. Once you have that `AssetData`, you can call `ToStringReference()` to convert that to an `FSoftObjectPath`, which you can then load asynchronously using the next system:
 
 ## StreamableManager and Asynchronous Loading
 
-Now that you have a `FStringAssetReference` that refers to an asset on disk, how do you actually load it asynchronously? `FStreamableManager` is the easiest way to do this. First, you will need to create an `FStreamableManager`, I would suggest putting it in some sort of global game singleton object, such as one specified using `GameSingletonClassName` in `DefaultEngine.ini`. Then, you can pass your `StringAssetReference` into it and start a load. `SynchronousLoad` will do a simple, blocking load and return the object. This method may be fine for smaller objects, but it could potentially stall your main thread for too long. In that case, you will need to use `RequestAsyncLoad`, which will asynchronously load a group of assets and call a delegate after completion. Here is an example:
+Now that you have a `FSoftObjectPath` that refers to an asset on disk, how do you actually load it asynchronously? `FStreamableManager` is the easiest way to do this. First, you will need to create an `FStreamableManager`, I would suggest putting it in some sort of global game singleton object, such as one specified using `GameSingletonClassName` in `DefaultEngine.ini`. Then, you can pass your `FSoftObjectPath` into it and start a load. `SynchronousLoad` will do a simple, blocking load and return the object. This method may be fine for smaller objects, but it could potentially stall your main thread for too long. In that case, you will need to use `RequestAsyncLoad`, which will asynchronously load a group of assets and call a delegate after completion. Here is an example:
 
 	void UGameCheatManager::GrantItems()
 	{
-		   TArray<FStringAssetReference> ItemsToStream;
+		   TArray<FSoftObjectPath> ItemsToStream;
 		   FStreamableManager& Streamable = UGameGlobals::Get().StreamableManager;
 		   for(int32 i = 0; i < ItemList.Num(); ++i)
 		   {
@@ -79,6 +79,6 @@ Now that you have a `FStringAssetReference` that refers to an asset on disk, how
 		   }
 	}
 
-In this example, ItemList is a `TArray< TAssetPtr<UGameItem> >` that was modified by designers in the editor. The code iterates that list, converts them to `StringReferences`, and queues up a load of them. When all of those items are loaded (or fail to load because they are missing), it calls the passed in delegate. That delegate then iterates the same list of items, dereferences them, and gives them to a player. `StreamableManager` keeps hard references to any assets it loads until the delegate is called, so you can safely know that none of the objects you wanted to asynchronously load will be garbage collected before the delegate is called. It releases those references after the delegate is called, so you need to hard reference them somewhere else if you want to ensure they will stay around.
+In this example, ItemList is a `TArray< TSoftObjectPtr<UGameItem> >` that was modified by designers in the editor. The code iterates that list, converts them to `StringReferences`, and queues up a load of them. When all of those items are loaded (or fail to load because they are missing), it calls the passed in delegate. That delegate then iterates the same list of items, dereferences them, and gives them to a player. `StreamableManager` keeps hard references to any assets it loads until the delegate is called, so you can safely know that none of the objects you wanted to asynchronously load will be garbage collected before the delegate is called. It releases those references after the delegate is called, so you need to hard reference them somewhere else if you want to ensure they will stay around.
 
 You can use the same method to asynchronously load a `FAssetData`, just call `ToStringReference` on them, add them to an array, and call `RequestAsyncLoad` with a delegate. The delegate can be anything you want, so you can pass along payload information if you want. By combining the methods described above, you should be able to set up a system that allows efficient loading of any asset in your game. It will take some time to convert your gameplay code that does direct memory accesses to handle asynchronous loading, but afterwards, your game will have far fewer stalls and much lower memory usage.
diff --git a/Engine/Documentation/Source/Programming/Assets/ReferencingAssets/ReferencingAssets.INT.udn b/Engine/Documentation/Source/Programming/Assets/ReferencingAssets/ReferencingAssets.INT.udn
--- a/Engine/Documentation/Source/Programming/Assets/ReferencingAssets/ReferencingAssets.INT.udn
+++ b/Engine/Documentation/Source/Programming/Assets/ReferencingAssets/ReferencingAssets.INT.udn
@@ -47,24 +47,24 @@ The declaration of the UPROPERTY appears the same as the previous hard reference
 
 ## Indirect Property Reference
 
-One easy way to control when an asset is loaded is to use a TAssetPtr. To the designer, they can work with it just like the direct property reference. However, instead of a direct pointer reference, the property is stored as a string with template code to enable safe checking of whether the asset has been loaded yet or not. Use the IsPending() method to check if the asset is ready to be accessed or not. Note that using TAssetPtr requires you to manually load the asset when you want to use it. You can use the templated LoadObject<>() method, StaticLoadObject(), or the FStreamingManager to load your object (see [](Programming/Assets/AsyncLoading) for more information). The first two methods load the asset synchronously which can cause frame rate spikes, so should only be used if you know it won't affect the gameplay.
+One easy way to control when an asset is loaded is to use a TSoftObjectPtr. To the designer, they can work with it just like the direct property reference. However, instead of a direct pointer reference, the property is stored as a string with template code to enable safe checking of whether the asset has been loaded yet or not. Use the IsPending() method to check if the asset is ready to be accessed or not. Note that using TSoftObjectPtr requires you to manually load the asset when you want to use it. You can use the templated LoadObject<>() method, StaticLoadObject(), or the FStreamingManager to load your object (see [](Programming/Assets/AsyncLoading) for more information). The first two methods load the asset synchronously which can cause frame rate spikes, so should only be used if you know it won't affect the gameplay.
 
 	UPROPERTY(EditDefaultsOnly, BlueprintReadWrite, Category=Building)
-	TAssetPtr<UStaticMesh> BaseMesh;
+	TSoftObjectPtr<UStaticMesh> BaseMesh;
 
 	UStaticMesh* GetLazyLoadedMesh()
 	{
 		if (BaseMesh.IsPending())
 		{
-			const FStringAssetReference& AssetRef = BaseMesh.ToStringReference();
+			const FSoftObjectPath& AssetRef = BaseMesh.ToStringReference();
 			BaseMesh = Cast< UStaticMesh>(Streamable.SynchronousLoad(AssetRef));
 		}
 		return BaseMesh.Get();
 	}
 
-The code above uses a TAssetPtr of UStaticMesh to lazy load the mesh at runtime. The asset is checked to see if the object has been loaded or not. If it has not, a synchronous load happens using the FStreamingManager. Otherwise, the UStaticMesh pointer within the TAssetPtr is returned to the caller.
+The code above uses a TSoftObjectPtr of UStaticMesh to lazy load the mesh at runtime. The asset is checked to see if the object has been loaded or not. If it has not, a synchronous load happens using the FStreamingManager. Otherwise, the UStaticMesh pointer within the TSoftObjectPtr is returned to the caller.
 
-If you want to defer loading a UClass, you use the same approach as the TAssetPtr substituting the class specific version TAssetSubclassOf template type. This functions the same as referring to a specific asset, but instead refers to the UClass for the asset instead of an instance.
+If you want to defer loading a UClass, you use the same approach as the TSoftObjectPtr substituting the class specific version TSoftClassPtr template type. This functions the same as referring to a specific asset, but instead refers to the UClass for the asset instead of an instance.
 
 ## Find/Load Object
 
diff --git a/Engine/Plugins/2D/Paper2D/Source/Paper2D/Classes/PaperSprite.h b/Engine/Plugins/2D/Paper2D/Source/Paper2D/Classes/PaperSprite.h
--- a/Engine/Plugins/2D/Paper2D/Source/Paper2D/Classes/PaperSprite.h
+++ b/Engine/Plugins/2D/Paper2D/Source/Paper2D/Classes/PaperSprite.h
@@ -155,7 +155,7 @@ protected:
 
 	// The previous spritesheet group this belonged to
 	// To make sure we remove ourselves from it if changed or nulled out
-	TAssetPtr<class UPaperSpriteAtlas> PreviousAtlasGroup;
+	TSoftObjectPtr<class UPaperSpriteAtlas> PreviousAtlasGroup;
 
 #endif
 
diff --git a/Engine/Plugins/2D/Paper2D/Source/Paper2D/Classes/PaperSpriteAtlas.h b/Engine/Plugins/2D/Paper2D/Source/Paper2D/Classes/PaperSpriteAtlas.h
--- a/Engine/Plugins/2D/Paper2D/Source/Paper2D/Classes/PaperSpriteAtlas.h
+++ b/Engine/Plugins/2D/Paper2D/Source/Paper2D/Classes/PaperSpriteAtlas.h
@@ -15,7 +15,7 @@ struct FPaperSpriteAtlasSlot
 	GENERATED_USTRUCT_BODY()
 
 	UPROPERTY()
-	TAssetPtr<class UPaperSprite> SpriteRef;
+	TSoftObjectPtr<class UPaperSprite> SpriteRef;
 
 	UPROPERTY()
 	int32 AtlasIndex;
diff --git a/Engine/Plugins/2D/Paper2D/Source/Paper2D/Classes/PaperTileMap.h b/Engine/Plugins/2D/Paper2D/Source/Paper2D/Classes/PaperTileMap.h
--- a/Engine/Plugins/2D/Paper2D/Source/Paper2D/Classes/PaperTileMap.h
+++ b/Engine/Plugins/2D/Paper2D/Source/Paper2D/Classes/PaperTileMap.h
@@ -73,7 +73,7 @@ class PAPER2D_API UPaperTileMap : public UObject
 
 	// Last tile set that was selected when editing the tile map
 	UPROPERTY()
-	TAssetPtr<UPaperTileSet> SelectedTileSet;
+	TSoftObjectPtr<UPaperTileSet> SelectedTileSet;
 
 	// The material to use on a tile map instance if not overridden
 	UPROPERTY(Category=Setup, EditAnywhere, BlueprintReadOnly)
diff --git a/Engine/Plugins/2D/Paper2D/Source/Paper2D/Private/PaperSprite.cpp b/Engine/Plugins/2D/Paper2D/Source/Paper2D/Private/PaperSprite.cpp
--- a/Engine/Plugins/2D/Paper2D/Source/Paper2D/Private/PaperSprite.cpp
+++ b/Engine/Plugins/2D/Paper2D/Source/Paper2D/Private/PaperSprite.cpp
@@ -574,13 +574,7 @@ void UPaperSprite::PostEditChangeProperty(FPropertyChangedEvent& PropertyChanged
 
 	if (PropertyName == GET_MEMBER_NAME_CHECKED(UPaperSprite, AtlasGroup))
 	{
-		auto SpriteAssetPtr = PreviousAtlasGroup;
-		FStringAssetReference AtlasGroupStringRef = SpriteAssetPtr.ToStringReference();
-		UPaperSpriteAtlas* PreviousAtlasGroupPtr = nullptr;
-		if (!AtlasGroupStringRef.ToString().IsEmpty())
-		{
-			PreviousAtlasGroupPtr = Cast<UPaperSpriteAtlas>(StaticLoadObject(UPaperSpriteAtlas::StaticClass(), nullptr, *AtlasGroupStringRef.ToString(), nullptr, LOAD_None, nullptr));
-		}
+		UPaperSpriteAtlas* PreviousAtlasGroupPtr = PreviousAtlasGroup.LoadSynchronous();
 		
 		if (PreviousAtlasGroupPtr != AtlasGroup)
 		{
diff --git a/Engine/Plugins/2D/Paper2D/Source/Paper2DEditor/Classes/PaperImporterSettings.h b/Engine/Plugins/2D/Paper2D/Source/Paper2DEditor/Classes/PaperImporterSettings.h
--- a/Engine/Plugins/2D/Paper2D/Source/Paper2DEditor/Classes/PaperImporterSettings.h
+++ b/Engine/Plugins/2D/Paper2D/Source/Paper2DEditor/Classes/PaperImporterSettings.h
@@ -5,7 +5,7 @@
 #include "CoreMinimal.h"
 #include "UObject/ObjectMacros.h"
 #include "UObject/Object.h"
-#include "Misc/StringAssetReference.h"
+#include "UObject/SoftObjectPath.h"
 #include "Engine/Texture.h"
 #include "PaperImporterSettings.generated.h"
 
@@ -93,27 +93,27 @@ protected:
 
 	// The unlit default masked material for newly created sprites (masked means binary opacity: things are either opaque or see-thru, with nothing in between)
 	UPROPERTY(config, EditAnywhere, Category=MaterialSettings, meta=(AllowedClasses="MaterialInterface", DisplayName="Unlit Default Masked Material"))
-	FStringAssetReference UnlitDefaultMaskedMaterialName;
+	FSoftObjectPath UnlitDefaultMaskedMaterialName;
 
 	// The unlit default translucent material for newly created sprites (translucent means smooth opacity which can vary continuously from 0..1, but translucent rendering is more expensive that opaque or masked rendering and has different sorting rules)
 	UPROPERTY(config, EditAnywhere, Category=MaterialSettings, meta=(AllowedClasses="MaterialInterface", DisplayName="Unlit Default Translucent Material"))
-	FStringAssetReference UnlitDefaultTranslucentMaterialName;
+	FSoftObjectPath UnlitDefaultTranslucentMaterialName;
 
 	// The unlit default opaque material for newly created sprites
 	UPROPERTY(config, EditAnywhere, Category=MaterialSettings, meta=(AllowedClasses="MaterialInterface", DisplayName="Unlit Default Opaque Sprite Material"))
-	FStringAssetReference UnlitDefaultOpaqueMaterialName;
+	FSoftObjectPath UnlitDefaultOpaqueMaterialName;
 
 	// The lit default masked material for newly created sprites (masked means binary opacity: things are either opaque or see-thru, with nothing in between)
 	UPROPERTY(config, EditAnywhere, Category=MaterialSettings, meta=(AllowedClasses="MaterialInterface", DisplayName="Lit Default Masked Material"))
-	FStringAssetReference LitDefaultMaskedMaterialName;
+	FSoftObjectPath LitDefaultMaskedMaterialName;
 
 	// The lit default translucent material for newly created sprites (translucent means smooth opacity which can vary continuously from 0..1, but translucent rendering is more expensive that opaque or masked rendering and has different sorting rules)
 	UPROPERTY(config, EditAnywhere, Category=MaterialSettings, meta=(AllowedClasses="MaterialInterface", DisplayName="Lit Default Translucent Material"))
-	FStringAssetReference LitDefaultTranslucentMaterialName;
+	FSoftObjectPath LitDefaultTranslucentMaterialName;
 
 	// The lit default opaque material for newly created sprites
 	UPROPERTY(config, EditAnywhere, Category=MaterialSettings, meta=(AllowedClasses="MaterialInterface", DisplayName="Lit Default Opaque Material"))
-	FStringAssetReference LitDefaultOpaqueMaterialName;
+	FSoftObjectPath LitDefaultOpaqueMaterialName;
 
 public:
 	UPaperImporterSettings();
diff --git a/Engine/Plugins/2D/Paper2D/Source/Paper2DEditor/Private/Atlasing/PaperAtlasHelpers.h b/Engine/Plugins/2D/Paper2D/Source/Paper2DEditor/Private/Atlasing/PaperAtlasHelpers.h
--- a/Engine/Plugins/2D/Paper2D/Source/Paper2DEditor/Private/Atlasing/PaperAtlasHelpers.h
+++ b/Engine/Plugins/2D/Paper2D/Source/Paper2DEditor/Private/Atlasing/PaperAtlasHelpers.h
@@ -4,7 +4,7 @@
 
 #include "CoreMinimal.h"
 #include "Misc/Guid.h"
-#include "Misc/StringAssetReference.h"
+#include "UObject/SoftObjectPath.h"
 #include "PaperSprite.h"
 #include "PaperSpriteAtlas.h"
 #include "Modules/ModuleManager.h"
@@ -16,7 +16,7 @@ static void LoadAllReferencedSprites(UPaperSpriteAtlas* AtlasGroup, TArray<UPape
 {
 	for (FPaperSpriteAtlasSlot& Slot : AtlasGroup->AtlasSlots)
 	{
-		FStringAssetReference SpriteStringRef = Slot.SpriteRef.ToStringReference();
+		FSoftObjectPath SpriteStringRef = Slot.SpriteRef.ToSoftObjectPath();
 		if (!SpriteStringRef.ToString().IsEmpty())
 		{
 			UPaperSprite* Sprite = Cast<UPaperSprite>(StaticLoadObject(UPaperSprite::StaticClass(), nullptr, *SpriteStringRef.ToString(), nullptr, LOAD_None, nullptr));
@@ -136,7 +136,7 @@ static const FPaperSpriteAtlasSlot& FindBestSlotForTexture(TArray<FPaperSpriteAt
 	bOutSpriteChanged = false;
 
 	// Asset pointer for the sprite we're currently trying to pack
-	TAssetPtr<UPaperSprite> SpriteAssetRef = Sprite;
+	TSoftObjectPtr<UPaperSprite> SpriteAssetRef = Sprite;
 
 	// 1. Find matching texture and see if we can fit it in the existing slot
 	for (int32 SlotIndex = 0; SlotIndex < AtlasSlots.Num(); ++SlotIndex)
diff --git a/Engine/Plugins/2D/Paper2D/Source/Paper2DEditor/Private/PaperImporterSettings.cpp b/Engine/Plugins/2D/Paper2D/Source/Paper2DEditor/Private/PaperImporterSettings.cpp
--- a/Engine/Plugins/2D/Paper2D/Source/Paper2DEditor/Private/PaperImporterSettings.cpp
+++ b/Engine/Plugins/2D/Paper2D/Source/Paper2DEditor/Private/PaperImporterSettings.cpp
@@ -27,13 +27,13 @@ UPaperImporterSettings::UPaperImporterSettings()
 
 	if (!IsRunningCommandlet())
 	{
-		UnlitDefaultMaskedMaterialName = FStringAssetReference("/Paper2D/MaskedUnlitSpriteMaterial.MaskedUnlitSpriteMaterial");
-		UnlitDefaultTranslucentMaterialName = FStringAssetReference("/Paper2D/TranslucentUnlitSpriteMaterial.TranslucentUnlitSpriteMaterial");
-		UnlitDefaultOpaqueMaterialName = FStringAssetReference("/Paper2D/OpaqueUnlitSpriteMaterial.OpaqueUnlitSpriteMaterial");
+		UnlitDefaultMaskedMaterialName = FSoftObjectPath("/Paper2D/MaskedUnlitSpriteMaterial.MaskedUnlitSpriteMaterial");
+		UnlitDefaultTranslucentMaterialName = FSoftObjectPath("/Paper2D/TranslucentUnlitSpriteMaterial.TranslucentUnlitSpriteMaterial");
+		UnlitDefaultOpaqueMaterialName = FSoftObjectPath("/Paper2D/OpaqueUnlitSpriteMaterial.OpaqueUnlitSpriteMaterial");
 
-		LitDefaultMaskedMaterialName = FStringAssetReference("/Paper2D/MaskedLitSpriteMaterial.MaskedLitSpriteMaterial");
-		LitDefaultTranslucentMaterialName = FStringAssetReference("/Paper2D/TranslucentLitSpriteMaterial.TranslucentLitSpriteMaterial");
-		LitDefaultOpaqueMaterialName = FStringAssetReference("/Paper2D/OpaqueLitSpriteMaterial.OpaqueLitSpriteMaterial");
+		LitDefaultMaskedMaterialName = FSoftObjectPath("/Paper2D/MaskedLitSpriteMaterial.MaskedLitSpriteMaterial");
+		LitDefaultTranslucentMaterialName = FSoftObjectPath("/Paper2D/TranslucentLitSpriteMaterial.TranslucentLitSpriteMaterial");
+		LitDefaultOpaqueMaterialName = FSoftObjectPath("/Paper2D/OpaqueLitSpriteMaterial.OpaqueLitSpriteMaterial");
 	}
 }
 
diff --git a/Engine/Plugins/2D/Paper2D/Source/Paper2DEditor/Private/TileSetEditor/SingleTileEditorViewportClient.cpp b/Engine/Plugins/2D/Paper2D/Source/Paper2DEditor/Private/TileSetEditor/SingleTileEditorViewportClient.cpp
--- a/Engine/Plugins/2D/Paper2D/Source/Paper2DEditor/Private/TileSetEditor/SingleTileEditorViewportClient.cpp
+++ b/Engine/Plugins/2D/Paper2D/Source/Paper2DEditor/Private/TileSetEditor/SingleTileEditorViewportClient.cpp
@@ -48,7 +48,7 @@ FSingleTileEditorViewportClient::FSingleTileEditorViewportClient(UPaperTileSet*
 	// Create a render component for the tile preview
 	PreviewTileSpriteComponent = NewObject<UPaperSpriteComponent>();
 
-	FStringAssetReference DefaultTranslucentMaterialName("/Paper2D/TranslucentUnlitSpriteMaterial.TranslucentUnlitSpriteMaterial");
+	FSoftObjectPath DefaultTranslucentMaterialName("/Paper2D/TranslucentUnlitSpriteMaterial.TranslucentUnlitSpriteMaterial");
 	UMaterialInterface* TranslucentMaterial = Cast<UMaterialInterface>(DefaultTranslucentMaterialName.TryLoad());
 	PreviewTileSpriteComponent->SetMaterial(0, TranslucentMaterial);
 
diff --git a/Engine/Plugins/2D/Paper2D/Source/PaperSpriteSheetImporter/Private/PaperJsonSpriteSheetImporter.cpp b/Engine/Plugins/2D/Paper2D/Source/PaperSpriteSheetImporter/Private/PaperJsonSpriteSheetImporter.cpp
--- a/Engine/Plugins/2D/Paper2D/Source/PaperSpriteSheetImporter/Private/PaperJsonSpriteSheetImporter.cpp
+++ b/Engine/Plugins/2D/Paper2D/Source/PaperSpriteSheetImporter/Private/PaperJsonSpriteSheetImporter.cpp
@@ -360,22 +360,18 @@ FPaperJsonSpriteSheetImporter::FPaperJsonSpriteSheetImporter()
 {
 }
 
-void FPaperJsonSpriteSheetImporter::SetReimportData(const TArray<FString>& ExistingSpriteNames, const TArray< TAssetPtr<class UPaperSprite> >& ExistingSpriteAssetPtrs)
+void FPaperJsonSpriteSheetImporter::SetReimportData(const TArray<FString>& ExistingSpriteNames, const TArray< TSoftObjectPtr<class UPaperSprite> >& ExistingSpriteSoftPtrs)
 {
-	check(ExistingSpriteNames.Num() == ExistingSpriteAssetPtrs.Num());
-	if (ExistingSpriteNames.Num() == ExistingSpriteAssetPtrs.Num())
+	check(ExistingSpriteNames.Num() == ExistingSpriteSoftPtrs.Num());
+	if (ExistingSpriteNames.Num() == ExistingSpriteSoftPtrs.Num())
 	{
-		for (int i = 0; i < ExistingSpriteAssetPtrs.Num(); ++i)
+		for (int i = 0; i < ExistingSpriteSoftPtrs.Num(); ++i)
 		{
-			const TAssetPtr<class UPaperSprite> SpriteAssetPtr = ExistingSpriteAssetPtrs[i];
-			FStringAssetReference SpriteStringRef = SpriteAssetPtr.ToStringReference();
-			if (!SpriteStringRef.ToString().IsEmpty())
+			const TSoftObjectPtr<class UPaperSprite> SpriteSoftPtr = ExistingSpriteSoftPtrs[i];
+			UPaperSprite* LoadedSprite = SpriteSoftPtr.LoadSynchronous();
+			if (LoadedSprite != nullptr)
 			{
-				UPaperSprite* LoadedSprite = Cast<UPaperSprite>(StaticLoadObject(UPaperSprite::StaticClass(), nullptr, *SpriteStringRef.ToString(), nullptr, LOAD_None, nullptr));
-				if (LoadedSprite != nullptr)
-				{
-					ExistingSprites.Add(ExistingSpriteNames[i], LoadedSprite);
-				}
+				ExistingSprites.Add(ExistingSpriteNames[i], LoadedSprite);
 			}
 		}
 	}
diff --git a/Engine/Plugins/2D/Paper2D/Source/PaperSpriteSheetImporter/Private/PaperJsonSpriteSheetImporter.h b/Engine/Plugins/2D/Paper2D/Source/PaperSpriteSheetImporter/Private/PaperJsonSpriteSheetImporter.h
--- a/Engine/Plugins/2D/Paper2D/Source/PaperSpriteSheetImporter/Private/PaperJsonSpriteSheetImporter.h
+++ b/Engine/Plugins/2D/Paper2D/Source/PaperSpriteSheetImporter/Private/PaperJsonSpriteSheetImporter.h
@@ -3,7 +3,7 @@
 #pragma once
 
 #include "CoreMinimal.h"
-#include "UObject/AssetPtr.h"
+#include "UObject/SoftObjectPtr.h"
 
 class FJsonObject;
 class UPaperSprite;
@@ -50,7 +50,7 @@ public:
 
 	bool PerformImport(const FString& LongPackagePath, EObjectFlags Flags, UPaperSpriteSheet* SpriteSheet);
 
-	void SetReimportData(const TArray<FString>& ExistingSpriteNames, const TArray< TAssetPtr<class UPaperSprite> >& ExistingSpriteAssetPtrs);
+	void SetReimportData(const TArray<FString>& ExistingSpriteNames, const TArray< TSoftObjectPtr<class UPaperSprite> >& ExistingSpriteSoftPtrs);
 
 	static UTexture2D* ImportOrReimportTexture(UTexture2D* ExistingTexture, const FString& TextureSourcePath, const FString& DestinationAssetFolder);
 	static UTexture2D* ImportTexture(const FString& TextureSourcePath, const FString& DestinationAssetFolder);
diff --git a/Engine/Plugins/2D/Paper2D/Source/PaperSpriteSheetImporter/Private/PaperSpriteSheet.h b/Engine/Plugins/2D/Paper2D/Source/PaperSpriteSheetImporter/Private/PaperSpriteSheet.h
--- a/Engine/Plugins/2D/Paper2D/Source/PaperSpriteSheetImporter/Private/PaperSpriteSheet.h
+++ b/Engine/Plugins/2D/Paper2D/Source/PaperSpriteSheetImporter/Private/PaperSpriteSheet.h
@@ -19,7 +19,7 @@ public:
 	TArray<FString> SpriteNames;
 
 	UPROPERTY(VisibleAnywhere, Category=Data)
-	TArray< TAssetPtr<class UPaperSprite> > Sprites;
+	TArray< TSoftObjectPtr<class UPaperSprite> > Sprites;
 
 	// The name of the default or diffuse texture during import
 	UPROPERTY(VisibleAnywhere, Category=Data)
diff --git a/Engine/Plugins/2D/Paper2D/Source/PaperSpriteSheetImporter/Private/PaperSpriteSheetAssetTypeActions.cpp b/Engine/Plugins/2D/Paper2D/Source/PaperSpriteSheetImporter/Private/PaperSpriteSheetAssetTypeActions.cpp
--- a/Engine/Plugins/2D/Paper2D/Source/PaperSpriteSheetImporter/Private/PaperSpriteSheetAssetTypeActions.cpp
+++ b/Engine/Plugins/2D/Paper2D/Source/PaperSpriteSheetImporter/Private/PaperSpriteSheetAssetTypeActions.cpp
@@ -1,7 +1,6 @@
 // Copyright 1998-2017 Epic Games, Inc. All Rights Reserved.
 
 #include "PaperSpriteSheetAssetTypeActions.h"
-#include "Misc/StringAssetReference.h"
 #include "Misc/PackageName.h"
 #include "Framework/MultiBox/MultiBoxBuilder.h"
 #include "Misc/FeedbackContext.h"
@@ -99,13 +98,8 @@ void FPaperSpriteSheetAssetTypeActions::ExecuteCreateFlipbooks(TArray<TWeakObjec
 
 				for (int SpriteIndex = 0; SpriteIndex < SpriteSheet->Sprites.Num(); ++SpriteIndex)
 				{
-					auto SpriteAssetPtr = SpriteSheet->Sprites[SpriteIndex];
-					FStringAssetReference SpriteStringRef = SpriteAssetPtr.ToStringReference();
-					UPaperSprite* Sprite = nullptr;
-					if (!SpriteStringRef.ToString().IsEmpty())
-					{
-						Sprite = Cast<UPaperSprite>(StaticLoadObject(UPaperSprite::StaticClass(), nullptr, *SpriteStringRef.ToString(), nullptr, LOAD_None, nullptr));
-					}
+					TSoftObjectPtr<class UPaperSprite> SpriteSoftPtr = SpriteSheet->Sprites[SpriteIndex];
+					UPaperSprite* Sprite = SpriteSoftPtr.LoadSynchronous();
 					if (Sprite != nullptr)
 					{
 						const FString SpriteName = useSpriteNames ? SpriteSheet->SpriteNames[SpriteIndex] : Sprite->GetName();
diff --git a/Engine/Plugins/2D/Paper2D/Source/SpriterImporter/Private/PaperSpriterImportData.h b/Engine/Plugins/2D/Paper2D/Source/SpriterImporter/Private/PaperSpriterImportData.h
--- a/Engine/Plugins/2D/Paper2D/Source/SpriterImporter/Private/PaperSpriterImportData.h
+++ b/Engine/Plugins/2D/Paper2D/Source/SpriterImporter/Private/PaperSpriterImportData.h
@@ -18,15 +18,15 @@ public:
 
 // 	// Should be a TMap...
 // 	TArray<FString> FileToTextureList;
-// 	TAssetPtr<class UTexture2D> ImportedTextureList;
-// 	TAssetPtr<class UPaperSprite> ImportedSpriteList;
+// 	TSoftObjectPtr<class UTexture2D> ImportedTextureList;
+// 	TSoftObjectPtr<class UPaperSprite> ImportedSpriteList;
 
 // 	// The names of sprites during import
 // 	UPROPERTY(VisibleAnywhere, Category=Data)
 // 	TArray<FString> SpriteNames;
 // 
 // 	UPROPERTY(VisibleAnywhere, Category = Data)
-// 	TArray< TAssetPtr<class UPaperSprite> > Sprites;
+// 	TArray< TSoftObjectPtr<class UPaperSprite> > Sprites;
 // 
 // 	// The name of the texture during import
 // 	UPROPERTY(VisibleAnywhere, Category = Data)
diff --git a/Engine/Plugins/Editor/AssetManagerEditor/Source/AssetManagerEditor/Private/PrimaryAssetIdCustomization.cpp b/Engine/Plugins/Editor/AssetManagerEditor/Source/AssetManagerEditor/Private/PrimaryAssetIdCustomization.cpp
--- a/Engine/Plugins/Editor/AssetManagerEditor/Source/AssetManagerEditor/Private/PrimaryAssetIdCustomization.cpp
+++ b/Engine/Plugins/Editor/AssetManagerEditor/Source/AssetManagerEditor/Private/PrimaryAssetIdCustomization.cpp
@@ -85,7 +85,7 @@ FString FPrimaryAssetIdCustomization::OnGetObjectPath() const
 
 	UAssetManager& Manager = UAssetManager::Get();
 
-	FStringAssetReference FoundPath = Manager.GetPrimaryAssetPath(FPrimaryAssetId(StringReference));
+	FSoftObjectPath FoundPath = Manager.GetPrimaryAssetPath(FPrimaryAssetId(StringReference));
 
 	return FoundPath.ToString();
 }
diff --git a/Engine/Plugins/Editor/AssetManagerEditor/Source/AssetManagerEditor/Private/SAssetAuditBrowser.cpp b/Engine/Plugins/Editor/AssetManagerEditor/Source/AssetManagerEditor/Private/SAssetAuditBrowser.cpp
--- a/Engine/Plugins/Editor/AssetManagerEditor/Source/AssetManagerEditor/Private/SAssetAuditBrowser.cpp
+++ b/Engine/Plugins/Editor/AssetManagerEditor/Source/AssetManagerEditor/Private/SAssetAuditBrowser.cpp
@@ -461,11 +461,11 @@ FReply SAssetAuditBrowser::OnKeyDown(const FGeometry& MyGeometry, const FKeyEven
 	return FReply::Unhandled();
 }
 
-void SAssetAuditBrowser::AddAssetsToList(const TArray<FStringAssetReference>& AssetsToView, bool bReplaceExisting)
+void SAssetAuditBrowser::AddAssetsToList(const TArray<FSoftObjectPath>& AssetsToView, bool bReplaceExisting)
 {
 	TArray<FName> AssetNames;
 
-	for (const FStringAssetReference& AssetToView : AssetsToView)
+	for (const FSoftObjectPath& AssetToView : AssetsToView)
 	{
 		AssetNames.Add(FName(*AssetToView.GetLongPackageName()));
 	}
@@ -629,7 +629,7 @@ void SAssetAuditBrowser::AddAssetsOfType(FPrimaryAssetType AssetType)
 {
 	if (AssetType.IsValid())
 	{
-		TArray<FStringAssetReference> AssetArray;
+		TArray<FSoftObjectPath> AssetArray;
 		
 		AssetManager->GetPrimaryAssetPathList(AssetType, AssetArray);
 
diff --git a/Engine/Plugins/Editor/AssetManagerEditor/Source/AssetManagerEditor/Private/SAssetAuditBrowser.h b/Engine/Plugins/Editor/AssetManagerEditor/Source/AssetManagerEditor/Private/SAssetAuditBrowser.h
--- a/Engine/Plugins/Editor/AssetManagerEditor/Source/AssetManagerEditor/Private/SAssetAuditBrowser.h
+++ b/Engine/Plugins/Editor/AssetManagerEditor/Source/AssetManagerEditor/Private/SAssetAuditBrowser.h
@@ -32,7 +32,7 @@ public:
 
 	/** Adds assets to current management view */
 	void AddAssetsToList(const TArray<FAssetData>& AssetsToView, bool bReplaceExisting);
-	void AddAssetsToList(const TArray<FStringAssetReference>& AssetsToView, bool bReplaceExisting);
+	void AddAssetsToList(const TArray<FSoftObjectPath>& AssetsToView, bool bReplaceExisting);
 	void AddAssetsToList(const TArray<FName>& PackageNamesToView, bool bReplaceExisting);
 
 protected:
diff --git a/Engine/Plugins/Experimental/ControlRig/Source/ControlRig/Public/Sequencer/ControlRigSequence.h b/Engine/Plugins/Experimental/ControlRig/Source/ControlRig/Public/Sequencer/ControlRigSequence.h
--- a/Engine/Plugins/Experimental/ControlRig/Source/ControlRig/Public/Sequencer/ControlRigSequence.h
+++ b/Engine/Plugins/Experimental/ControlRig/Source/ControlRig/Public/Sequencer/ControlRigSequence.h
@@ -36,11 +36,11 @@ public:
 public:
 	/** The last animation sequence this control rig sequence was exported to */
 	UPROPERTY(AssetRegistrySearchable)
-	TAssetPtr<UAnimSequence> LastExportedToAnimationSequence;
+	TSoftObjectPtr<UAnimSequence> LastExportedToAnimationSequence;
 
 	/** The skeletal mesh that was used the last time we exported this sequence */
 	UPROPERTY(AssetRegistrySearchable)
-	TAssetPtr<USkeletalMesh> LastExportedUsingSkeletalMesh;
+	TSoftObjectPtr<USkeletalMesh> LastExportedUsingSkeletalMesh;
 
 	/** The frame rate that was used the last time we exported this sequence */
 	UPROPERTY(AssetRegistrySearchable)
diff --git a/Engine/Plugins/Experimental/ControlRig/Source/ControlRigEditor/Private/ControlRigEditorModule.cpp b/Engine/Plugins/Experimental/ControlRig/Source/ControlRigEditor/Private/ControlRigEditorModule.cpp
--- a/Engine/Plugins/Experimental/ControlRig/Source/ControlRigEditor/Private/ControlRigEditorModule.cpp
+++ b/Engine/Plugins/Experimental/ControlRig/Source/ControlRigEditor/Private/ControlRigEditorModule.cpp
@@ -316,7 +316,7 @@ void FControlRigEditorModule::HandleNewBlueprintCreated(UBlueprint* InBlueprint)
 		OutputNode->NodePosX = 0;
 		OutputNode->NodePosY = 0;
 		UEdGraphSchema_K2::SetNodeMetaData(OutputNode, FNodeMetadata::DefaultGraphNode);
-		OutputNode->DisableNode();
+		OutputNode->MakeAutomaticallyPlacedGhostNode();
 		OutputNode->NodeComment = LOCTEXT("AnimationOutputComment", "This node acts as the output for this animation controller.\nTo add or remove an output pin, enable or disable the \"Animation Output\" checkbox for a variable.").ToString();
 		OutputNode->bCommentBubbleVisible = true;
 		OutputNode->bCommentBubblePinned = true;
diff --git a/Engine/Plugins/Experimental/ControlRig/Source/ControlRigEditor/Private/InputOutput/K2Node_ControlRig.cpp b/Engine/Plugins/Experimental/ControlRig/Source/ControlRigEditor/Private/InputOutput/K2Node_ControlRig.cpp
--- a/Engine/Plugins/Experimental/ControlRig/Source/ControlRigEditor/Private/InputOutput/K2Node_ControlRig.cpp
+++ b/Engine/Plugins/Experimental/ControlRig/Source/ControlRigEditor/Private/InputOutput/K2Node_ControlRig.cpp
@@ -147,7 +147,7 @@ UClass* UK2Node_ControlRig::GetControlRigClass() const
 	const UClass* FoundClass = GetControlRigClassImpl();
 	if (FoundClass)
 	{
-		ControlRigClass = FStringClassReference(FoundClass);
+		ControlRigClass = FSoftClassPath(FoundClass);
 	}
 
 	if (ControlRigClass.IsValid())
diff --git a/Engine/Plugins/Experimental/ControlRig/Source/ControlRigEditor/Private/InputOutput/K2Node_ControlRig.h b/Engine/Plugins/Experimental/ControlRig/Source/ControlRigEditor/Private/InputOutput/K2Node_ControlRig.h
--- a/Engine/Plugins/Experimental/ControlRig/Source/ControlRigEditor/Private/InputOutput/K2Node_ControlRig.h
+++ b/Engine/Plugins/Experimental/ControlRig/Source/ControlRigEditor/Private/InputOutput/K2Node_ControlRig.h
@@ -7,7 +7,7 @@
 #include "UObject/ObjectMacros.h"
 #include "EdGraph/EdGraphNodeUtils.h"
 #include "K2Node.h"
-#include "Misc/StringClassReference.h"
+#include "UObject/SoftObjectPath.h"
 #include "K2Node_ControlRig.generated.h"
 
 class FBlueprintActionDatabaseRegistrar;
@@ -44,7 +44,7 @@ class UK2Node_ControlRig : public UK2Node
 	 * We need to cache this here as we need to access it during compilation when CDOs etc. may be in flux.
 	 */
 	UPROPERTY()
-	mutable FStringClassReference ControlRigClass;
+	mutable FSoftClassPath ControlRigClass;
 
 	/** Disabled input pins */
 	UPROPERTY(EditAnywhere, Category = "Inputs")
diff --git a/Engine/Plugins/MovieScene/LevelSequenceEditor/Source/LevelSequenceEditor/Private/LevelSequenceEditorToolkit.cpp b/Engine/Plugins/MovieScene/LevelSequenceEditor/Source/LevelSequenceEditor/Private/LevelSequenceEditorToolkit.cpp
--- a/Engine/Plugins/MovieScene/LevelSequenceEditor/Source/LevelSequenceEditor/Private/LevelSequenceEditorToolkit.cpp
+++ b/Engine/Plugins/MovieScene/LevelSequenceEditor/Source/LevelSequenceEditor/Private/LevelSequenceEditorToolkit.cpp
@@ -385,7 +385,7 @@ void FLevelSequenceEditorToolkit::AddDefaultTracksForActor(AActor& Actor, const
 		}
 
 		// add tracks by type
-		for (const FStringClassReference& DefaultTrack : TrackSettings.DefaultTracks)
+		for (const FSoftClassPath& DefaultTrack : TrackSettings.DefaultTracks)
 		{
 			UClass* TrackClass = DefaultTrack.ResolveClass();
 
@@ -399,7 +399,7 @@ void FLevelSequenceEditorToolkit::AddDefaultTracksForActor(AActor& Actor, const
 					continue;
 				}
 				
-				for (const FStringClassReference& ExcludeDefaultTrack : ExcludeTrackSettings.ExcludeDefaultTracks)
+				for (const FSoftClassPath& ExcludeDefaultTrack : ExcludeTrackSettings.ExcludeDefaultTracks)
 				{
 					if (ExcludeDefaultTrack == DefaultTrack)
 					{
diff --git a/Engine/Plugins/MovieScene/LevelSequenceEditor/Source/LevelSequenceEditor/Private/Misc/LevelSequenceEditorSettings.h b/Engine/Plugins/MovieScene/LevelSequenceEditor/Source/LevelSequenceEditor/Private/Misc/LevelSequenceEditorSettings.h
--- a/Engine/Plugins/MovieScene/LevelSequenceEditor/Source/LevelSequenceEditor/Private/Misc/LevelSequenceEditorSettings.h
+++ b/Engine/Plugins/MovieScene/LevelSequenceEditor/Source/LevelSequenceEditor/Private/Misc/LevelSequenceEditorSettings.h
@@ -6,7 +6,7 @@
 #include "UObject/ObjectMacros.h"
 #include "UObject/Object.h"
 #include "Engine/EngineTypes.h"
-#include "Misc/StringClassReference.h"
+#include "UObject/SoftObjectPath.h"
 #include "LevelSequence.h"
 #include "LevelSequenceEditorSettings.generated.h"
 
@@ -32,15 +32,15 @@ struct FLevelSequenceTrackSettings
 
 	/** The Actor class to create movie scene tracks for. */
 	UPROPERTY(config, noclear, EditAnywhere, Category=TrackSettings, meta=(MetaClass="Actor"))
-	FStringClassReference MatchingActorClass;
+	FSoftClassPath MatchingActorClass;
 
 	/** List of movie scene track classes to be added automatically. */
 	UPROPERTY(config, noclear, EditAnywhere, Category=TrackSettings, meta=(MetaClass="MovieSceneTrack"))
-	TArray<FStringClassReference> DefaultTracks;
+	TArray<FSoftClassPath> DefaultTracks;
 
 	/** List of movie scene track classes not to be added automatically. */
 	UPROPERTY(config, noclear, EditAnywhere, Category=TrackSettings, meta=(MetaClass="MovieSceneTrack"))
-	TArray<FStringClassReference> ExcludeDefaultTracks;
+	TArray<FSoftClassPath> ExcludeDefaultTracks;
 
 	/** List of property names for which movie scene tracks will be created automatically. */
 	UPROPERTY(config, EditAnywhere, Category=TrackSettings)
diff --git a/Engine/Plugins/NetcodeUnitTest/NetcodeUnitTest/Source/NetcodeUnitTest/Public/NUTUtilReflection.h b/Engine/Plugins/NetcodeUnitTest/NetcodeUnitTest/Source/NetcodeUnitTest/Public/NUTUtilReflection.h
--- a/Engine/Plugins/NetcodeUnitTest/NetcodeUnitTest/Source/NetcodeUnitTest/Public/NUTUtilReflection.h
+++ b/Engine/Plugins/NetcodeUnitTest/NetcodeUnitTest/Source/NetcodeUnitTest/Public/NUTUtilReflection.h
@@ -132,8 +132,8 @@ class FFuncReflection;
  *	readwriteupcast		- UIntProperty									arraycheck	assign
  *	readwriteupcast		- UInt64Property								arraycheck	assign
  *					- ObjectPropertyBase
- *						- AssetObjectProperty
- *							- AssetClassProperty
+ *						- SoftObjectProperty
+ *							- SoftClassProperty
  *						- LazyObjectProperty
  *	readwrite			- ObjectProperty								arraycheck
  *	inherited				- ClassProperty								inherited
diff --git a/Engine/Plugins/Online/OnlineFramework/Source/Lobby/Public/LobbyBeaconHost.h b/Engine/Plugins/Online/OnlineFramework/Source/Lobby/Public/LobbyBeaconHost.h
--- a/Engine/Plugins/Online/OnlineFramework/Source/Lobby/Public/LobbyBeaconHost.h
+++ b/Engine/Plugins/Online/OnlineFramework/Source/Lobby/Public/LobbyBeaconHost.h
@@ -4,7 +4,6 @@
 
 #include "CoreMinimal.h"
 #include "UObject/ObjectMacros.h"
-#include "UObject/AssetPtr.h"
 #include "OnlineBeaconHostObject.h"
 
 #include "LobbyBeaconHost.generated.h"
@@ -93,7 +92,7 @@ protected:
 
 	/** Class to use for the lobby beacon state */
 	UPROPERTY()
-	TAssetSubclassOf<ALobbyBeaconState> LobbyStateClass;
+	TSoftClassPtr<ALobbyBeaconState> LobbyStateClass;
 
 	/** Actor representing the state of the lobby (similar to AGameState) */
 	UPROPERTY()
diff --git a/Engine/Plugins/Online/OnlineSubsystemUtils/Source/OnlineSubsystemUtils/Private/OnlineSubsystemUtils.cpp b/Engine/Plugins/Online/OnlineSubsystemUtils/Source/OnlineSubsystemUtils/Private/OnlineSubsystemUtils.cpp
--- a/Engine/Plugins/Online/OnlineSubsystemUtils/Source/OnlineSubsystemUtils/Private/OnlineSubsystemUtils.cpp
+++ b/Engine/Plugins/Online/OnlineSubsystemUtils/Source/OnlineSubsystemUtils/Private/OnlineSubsystemUtils.cpp
@@ -56,7 +56,7 @@ UAudioComponent* CreateVoiceAudioComponent(uint32 SampleRate)
 				AudioComponent->bAllowSpatialization = false;
 				AudioComponent->SetVolumeMultiplier(1.5f);
 
-				const FStringAssetReference VoiPSoundClassName = GetDefault<UAudioSettings>()->VoiPSoundClass;
+				const FSoftObjectPath VoiPSoundClassName = GetDefault<UAudioSettings>()->VoiPSoundClass;
 				if (VoiPSoundClassName.IsValid())
 				{
 					AudioComponent->SoundClassOverride = LoadObject<USoundClass>(nullptr, *VoiPSoundClassName.ToString());
diff --git a/Engine/Plugins/Runtime/GameplayAbilities/Source/GameplayAbilities/Private/AbilitySystemGlobals.cpp b/Engine/Plugins/Runtime/GameplayAbilities/Source/GameplayAbilities/Private/AbilitySystemGlobals.cpp
--- a/Engine/Plugins/Runtime/GameplayAbilities/Source/GameplayAbilities/Private/AbilitySystemGlobals.cpp
+++ b/Engine/Plugins/Runtime/GameplayAbilities/Source/GameplayAbilities/Private/AbilitySystemGlobals.cpp
@@ -17,7 +17,7 @@
 UAbilitySystemGlobals::UAbilitySystemGlobals(const FObjectInitializer& ObjectInitializer)
 : Super(ObjectInitializer)
 {
-	AbilitySystemGlobalsClassName = FStringClassReference(TEXT("/Script/GameplayAbilities.AbilitySystemGlobals"));
+	AbilitySystemGlobalsClassName = FSoftClassPath(TEXT("/Script/GameplayAbilities.AbilitySystemGlobals"));
 
 	PredictTargetGameplayEffects = true;
 
@@ -316,7 +316,7 @@ void UAbilitySystemGlobals::InitAttributeDefaults()
 	}
 
 	// Handle array of global curve tables for attribute defaults
- 	for (const FStringAssetReference& AttribDefaultTableName : GlobalAttributeSetDefaultsTableNames)
+ 	for (const FSoftObjectPath& AttribDefaultTableName : GlobalAttributeSetDefaultsTableNames)
  	{
 		if (AttribDefaultTableName.IsValid())
 		{
diff --git a/Engine/Plugins/Runtime/GameplayAbilities/Source/GameplayAbilities/Private/GameplayAbilitiesModule.cpp b/Engine/Plugins/Runtime/GameplayAbilities/Source/GameplayAbilities/Private/GameplayAbilitiesModule.cpp
--- a/Engine/Plugins/Runtime/GameplayAbilities/Source/GameplayAbilities/Private/GameplayAbilitiesModule.cpp
+++ b/Engine/Plugins/Runtime/GameplayAbilities/Source/GameplayAbilities/Private/GameplayAbilitiesModule.cpp
@@ -2,7 +2,7 @@
 
 #include "GameplayAbilitiesModule.h"
 #include "UObject/Object.h"
-#include "Misc/StringClassReference.h"
+#include "UObject/SoftObjectPath.h"
 #include "GameFramework/HUD.h"
 #include "AbilitySystemGlobals.h"
 #include "AbilitySystemComponent.h"
@@ -25,7 +25,7 @@ class FGameplayAbilitiesModule : public IGameplayAbilitiesModule
 		// Defer loading of globals to the first time it is requested
 		if (!AbilitySystemGlobals)
 		{
-			FStringClassReference AbilitySystemClassName = (UAbilitySystemGlobals::StaticClass()->GetDefaultObject<UAbilitySystemGlobals>())->AbilitySystemGlobalsClassName;
+			FSoftClassPath AbilitySystemClassName = (UAbilitySystemGlobals::StaticClass()->GetDefaultObject<UAbilitySystemGlobals>())->AbilitySystemGlobalsClassName;
 
 			UClass* SingletonClass = AbilitySystemClassName.TryLoadClass<UObject>();
 			checkf(SingletonClass != nullptr, TEXT("Ability config value AbilitySystemGlobalsClassName is not a valid class name."));
diff --git a/Engine/Plugins/Runtime/GameplayAbilities/Source/GameplayAbilities/Private/GameplayCueManager.cpp b/Engine/Plugins/Runtime/GameplayAbilities/Source/GameplayAbilities/Private/GameplayCueManager.cpp
--- a/Engine/Plugins/Runtime/GameplayAbilities/Source/GameplayAbilities/Private/GameplayCueManager.cpp
+++ b/Engine/Plugins/Runtime/GameplayAbilities/Source/GameplayAbilities/Private/GameplayCueManager.cpp
@@ -276,7 +276,7 @@ bool UGameplayCueManager::HandleMissingGameplayCue(UGameplayCueSet* OwningSet, s
 	return false;
 }
 
-void UGameplayCueManager::OnMissingCueAsyncLoadComplete(FStringAssetReference LoadedObject, TWeakObjectPtr<UGameplayCueSet> OwningSet, FGameplayTag GameplayCueTag, TWeakObjectPtr<AActor> TargetActor, EGameplayCueEvent::Type EventType, FGameplayCueParameters Parameters)
+void UGameplayCueManager::OnMissingCueAsyncLoadComplete(FSoftObjectPath LoadedObject, TWeakObjectPtr<UGameplayCueSet> OwningSet, FGameplayTag GameplayCueTag, TWeakObjectPtr<AActor> TargetActor, EGameplayCueEvent::Type EventType, FGameplayCueParameters Parameters)
 {
 	if (!LoadedObject.ResolveObject())
 	{
@@ -647,7 +647,7 @@ void UGameplayCueManager::RefreshObjectLibraries()
 	}
 }
 
-static void SearchDynamicClassCues(const FName PropertyName, const TArray<FString>& Paths, TArray<FGameplayCueReferencePair>& CuesToAdd, TArray<FStringAssetReference>& AssetsToLoad)
+static void SearchDynamicClassCues(const FName PropertyName, const TArray<FString>& Paths, TArray<FGameplayCueReferencePair>& CuesToAdd, TArray<FSoftObjectPath>& AssetsToLoad)
 {
 	// Iterate over all Dynamic Classes (nativized Blueprints). Search for ones with GameplayCueName tag.
 
@@ -675,7 +675,7 @@ static void SearchDynamicClassCues(const FName PropertyName, const TArray<FStrin
 			FGameplayTag GameplayCueTag = Manager.RequestGameplayTag(*FoundGameplayTag, false);
 			if (GameplayCueTag.IsValid())
 			{
-				FStringAssetReference StringRef(ClassPath); // TODO: is there any translation needed?
+				FSoftObjectPath StringRef(ClassPath); // TODO: is there any translation needed?
 				ensure(StringRef.IsValid());
 
 				CuesToAdd.Add(FGameplayCueReferencePair(GameplayCueTag, StringRef));
@@ -758,7 +758,7 @@ void UGameplayCueManager::InitObjectLibrary(FGameplayCueObjectLibrary& Lib)
 	Lib.StaticObjectLibrary->GetAssetDataList(StaticAssetDatas);
 
 	TArray<FGameplayCueReferencePair> CuesToAdd;
-	TArray<FStringAssetReference> AssetsToLoad;
+	TArray<FSoftObjectPath> AssetsToLoad;
 
 	// ------------------------------------------------------------------------------------------------------------------
 	// Build Cue lists for loading. Determines what from the obj library needs to be loaded
@@ -786,7 +786,7 @@ void UGameplayCueManager::InitObjectLibrary(FGameplayCueObjectLibrary& Lib)
 	// --------------------------------------------
 	if (Lib.bShouldAsyncLoad)
 	{
-		auto ForwardLambda = [](TArray<FStringAssetReference> AssetList, FOnGameplayCueNotifySetLoaded OnLoadedDelegate)
+		auto ForwardLambda = [](TArray<FSoftObjectPath> AssetList, FOnGameplayCueNotifySetLoaded OnLoadedDelegate)
 		{
 			OnLoadedDelegate.ExecuteIfBound(AssetList);
 		};
@@ -808,7 +808,7 @@ void UGameplayCueManager::InitObjectLibrary(FGameplayCueObjectLibrary& Lib)
 
 static FAutoConsoleVariable CVarGameplyCueAddToGlobalSetDebug(TEXT("GameplayCue.AddToGlobalSet.DebugTag"), TEXT(""), TEXT("Debug Tag adding to global set"), ECVF_Default	);
 
-void UGameplayCueManager::BuildCuesToAddToGlobalSet(const TArray<FAssetData>& AssetDataList, FName TagPropertyName, TArray<FGameplayCueReferencePair>& OutCuesToAdd, TArray<FStringAssetReference>& OutAssetsToLoad, FShouldLoadGCNotifyDelegate ShouldLoad)
+void UGameplayCueManager::BuildCuesToAddToGlobalSet(const TArray<FAssetData>& AssetDataList, FName TagPropertyName, TArray<FGameplayCueReferencePair>& OutCuesToAdd, TArray<FSoftObjectPath>& OutAssetsToLoad, FShouldLoadGCNotifyDelegate ShouldLoad)
 {
 	UGameplayTagsManager& Manager = UGameplayTagsManager::Get();
 
@@ -850,7 +850,7 @@ void UGameplayCueManager::BuildCuesToAddToGlobalSet(const TArray<FAssetData>& As
 			if (GameplayCueTag.IsValid())
 			{
 				// Add a new NotifyData entry to our flat list for this one
-				FStringAssetReference StringRef;
+				FSoftObjectPath StringRef;
 				StringRef.SetPath(FPackageName::ExportTextPathToObjectPath(*GeneratedClassTag));
 
 				OutCuesToAdd.Add(FGameplayCueReferencePair(GameplayCueTag, StringRef));
@@ -920,9 +920,9 @@ void UGameplayCueManager::CheckForTooManyRPCs(FName FuncName, const FGameplayCue
 	}
 }
 
-void UGameplayCueManager::OnGameplayCueNotifyAsyncLoadComplete(TArray<FStringAssetReference> AssetList)
+void UGameplayCueManager::OnGameplayCueNotifyAsyncLoadComplete(TArray<FSoftObjectPath> AssetList)
 {
-	for (FStringAssetReference StringRef : AssetList)
+	for (FSoftObjectPath StringRef : AssetList)
 	{
 		UClass* GCClass = FindObject<UClass>(nullptr, *StringRef.ToString());
 		if (ensure(GCClass))
@@ -977,7 +977,7 @@ void UGameplayCueManager::HandleAssetAdded(UObject *Object)
 		{
 			if (VerifyNotifyAssetIsInValidPath(Blueprint->GetOuter()->GetPathName()))
 			{
-				FStringAssetReference StringRef;
+				FSoftObjectPath StringRef;
 				StringRef.SetPath(Blueprint->GeneratedClass->GetPathName());
 
 				TArray<FGameplayCueReferencePair> CuesToAdd;
@@ -1004,7 +1004,7 @@ void UGameplayCueManager::HandleAssetAdded(UObject *Object)
 /** Handles cleaning up an object library if it matches the passed in object */
 void UGameplayCueManager::HandleAssetDeleted(UObject *Object)
 {
-	FStringAssetReference StringRefToRemove;
+	FSoftObjectPath StringRefToRemove;
 	UBlueprint* Blueprint = Cast<UBlueprint>(Object);
 	if (Blueprint && Blueprint->GeneratedClass)
 	{
@@ -1019,7 +1019,7 @@ void UGameplayCueManager::HandleAssetDeleted(UObject *Object)
 
 	if (StringRefToRemove.IsValid())
 	{
-		TArray<FStringAssetReference> StringRefs;
+		TArray<FSoftObjectPath> StringRefs;
 		StringRefs.Add(StringRefToRemove);
 		
 		
diff --git a/Engine/Plugins/Runtime/GameplayAbilities/Source/GameplayAbilities/Private/GameplayCueSet.cpp b/Engine/Plugins/Runtime/GameplayAbilities/Source/GameplayAbilities/Private/GameplayCueSet.cpp
--- a/Engine/Plugins/Runtime/GameplayAbilities/Source/GameplayAbilities/Private/GameplayCueSet.cpp
+++ b/Engine/Plugins/Runtime/GameplayAbilities/Source/GameplayAbilities/Private/GameplayCueSet.cpp
@@ -45,7 +45,7 @@ void UGameplayCueSet::AddCues(const TArray<FGameplayCueReferencePair>& CuesToAdd
 		for (const FGameplayCueReferencePair& CueRefPair : CuesToAdd)
 		{
 			const FGameplayTag& GameplayCueTag = CueRefPair.GameplayCueTag;
-			const FStringAssetReference& StringRef = CueRefPair.StringRef;
+			const FSoftObjectPath& StringRef = CueRefPair.StringRef;
 
 			// Check for duplicates: we may want to remove this eventually (allow multiple GC notifies to handle same event)
 			bool bDupe = false;
@@ -82,9 +82,9 @@ void UGameplayCueSet::RemoveCuesByTags(const FGameplayTagContainer& TagsToRemove
 {
 }
 
-void UGameplayCueSet::RemoveCuesByStringRefs(const TArray<FStringAssetReference>& CuesToRemove)
+void UGameplayCueSet::RemoveCuesByStringRefs(const TArray<FSoftObjectPath>& CuesToRemove)
 {
-	for (const FStringAssetReference& StringRefToRemove : CuesToRemove)
+	for (const FSoftObjectPath& StringRefToRemove : CuesToRemove)
 	{
 		for (int32 idx = 0; idx < GameplayCueData.Num(); ++idx)
 		{
@@ -118,7 +118,7 @@ void UGameplayCueSet::GetFilenames(TArray<FString>& Filenames) const
 	}
 }
 
-void UGameplayCueSet::GetStringAssetReferences(TArray<FStringAssetReference>& List) const
+void UGameplayCueSet::GetSoftObjectPaths(TArray<FSoftObjectPath>& List) const
 {
 	List.Reserve(GameplayCueData.Num());
 	for (const FGameplayCueNotifyData& Data : GameplayCueData)
@@ -128,7 +128,7 @@ void UGameplayCueSet::GetStringAssetReferences(TArray<FStringAssetReference>& Li
 }
 
 #if WITH_EDITOR
-void UGameplayCueSet::UpdateCueByStringRefs(const FStringAssetReference& CueToRemove, FString NewPath)
+void UGameplayCueSet::UpdateCueByStringRefs(const FSoftObjectPath& CueToRemove, FString NewPath)
 {
 	for (int32 idx = 0; idx < GameplayCueData.Num(); ++idx)
 	{
diff --git a/Engine/Plugins/Runtime/GameplayAbilities/Source/GameplayAbilities/Public/AbilitySystemGlobals.h b/Engine/Plugins/Runtime/GameplayAbilities/Source/GameplayAbilities/Public/AbilitySystemGlobals.h
--- a/Engine/Plugins/Runtime/GameplayAbilities/Source/GameplayAbilities/Public/AbilitySystemGlobals.h
+++ b/Engine/Plugins/Runtime/GameplayAbilities/Source/GameplayAbilities/Public/AbilitySystemGlobals.h
@@ -5,8 +5,7 @@
 #include "CoreMinimal.h"
 #include "UObject/ObjectMacros.h"
 #include "UObject/Object.h"
-#include "Misc/StringAssetReference.h"
-#include "Misc/StringClassReference.h"
+#include "UObject/SoftObjectPath.h"
 #include "GameplayTagContainer.h"
 #include "GameplayEffectTypes.h"
 #include "GameplayAbilitiesModule.h"
@@ -121,7 +120,7 @@ class GAMEPLAYABILITIES_API UAbilitySystemGlobals : public UObject
 
 	/** The class to instantiate as the globals object. Defaults to this class but can be overridden */
 	UPROPERTY(config)
-	FStringClassReference AbilitySystemGlobalsClassName;
+	FSoftClassPath AbilitySystemGlobalsClassName;
 
 	void AutomationTestOnly_SetGlobalCurveTable(UCurveTable *InTable)
 	{
@@ -276,27 +275,27 @@ protected:
 
 	/** Name of global curve table to use as the default for scalable floats, etc. */
 	UPROPERTY(config)
-	FStringAssetReference GlobalCurveTableName;
+	FSoftObjectPath GlobalCurveTableName;
 
 	/** Holds information about the valid attributes' min and max values and stacking rules */
 	UPROPERTY(config)
-	FStringAssetReference GlobalAttributeMetaDataTableName;
+	FSoftObjectPath GlobalAttributeMetaDataTableName;
 
 	/** Holds default values for attribute sets, keyed off of Name/Levels. NOTE: Preserved for backwards compatibility, should use the array version below now */
 	UPROPERTY(config)
-	FStringAssetReference GlobalAttributeSetDefaultsTableName;
+	FSoftObjectPath GlobalAttributeSetDefaultsTableName;
 
 	/** Array of curve table names to use for default values for attribute sets, keyed off of Name/Levels */
 	UPROPERTY(config)
-	TArray<FStringAssetReference> GlobalAttributeSetDefaultsTableNames;
+	TArray<FSoftObjectPath> GlobalAttributeSetDefaultsTableNames;
 
 	/** Class reference to gameplay cue manager. Use this if you want to just instantiate a class for your gameplay cue manager without having to create an asset. */
 	UPROPERTY(config)
-	FStringAssetReference GlobalGameplayCueManagerClass;
+	FSoftObjectPath GlobalGameplayCueManagerClass;
 
 	/** Object reference to gameplay cue manager (E.g., reference to a specific blueprint of your GameplayCueManager class. This is not necessary unless you want to have data or blueprints in your gameplay cue manager. */
 	UPROPERTY(config)
-	FStringAssetReference GlobalGameplayCueManagerName;
+	FSoftObjectPath GlobalGameplayCueManagerName;
 
 	/** Look in these paths for GameplayCueNotifies. These are your "always loaded" set. */
 	UPROPERTY(config)
@@ -304,7 +303,7 @@ protected:
 
 	/** The class to instantiate as the GameplayTagResponseTable. */
 	UPROPERTY(config)
-	FStringAssetReference GameplayTagResponseTableName;
+	FSoftObjectPath GameplayTagResponseTableName;
 
 	UPROPERTY()
 	UGameplayTagReponseTable* GameplayTagResponseTable;
diff --git a/Engine/Plugins/Runtime/GameplayAbilities/Source/GameplayAbilities/Public/GameplayCueManager.h b/Engine/Plugins/Runtime/GameplayAbilities/Source/GameplayAbilities/Public/GameplayCueManager.h
--- a/Engine/Plugins/Runtime/GameplayAbilities/Source/GameplayAbilities/Public/GameplayCueManager.h
+++ b/Engine/Plugins/Runtime/GameplayAbilities/Source/GameplayAbilities/Public/GameplayCueManager.h
@@ -4,7 +4,7 @@
 
 #include "CoreMinimal.h"
 #include "UObject/ObjectMacros.h"
-#include "Misc/StringAssetReference.h"
+#include "UObject/SoftObjectPath.h"
 #include "GameplayTagContainer.h"
 #include "Engine/DataAsset.h"
 #include "GameplayEffectTypes.h"
@@ -22,7 +22,7 @@ class AGameplayCueNotify_Actor;
 class UAbilitySystemComponent;
 class UObjectLibrary;
 
-DECLARE_DELEGATE_OneParam(FOnGameplayCueNotifySetLoaded, TArray<FStringAssetReference>);
+DECLARE_DELEGATE_OneParam(FOnGameplayCueNotifySetLoaded, TArray<FSoftObjectPath>);
 DECLARE_DELEGATE_OneParam(FGameplayCueProxyTick, float);
 DECLARE_DELEGATE_RetVal_TwoParams(bool, FShouldLoadGCNotifyDelegate, const FAssetData&, FName);
 
@@ -299,14 +299,14 @@ protected:
 	virtual void ReloadObjectLibrary(UWorld* World, const UWorld::InitializationValues IVS);
 #endif
 
-	void BuildCuesToAddToGlobalSet(const TArray<FAssetData>& AssetDataList, FName TagPropertyName, TArray<struct FGameplayCueReferencePair>& OutCuesToAdd, TArray<FStringAssetReference>& OutAssetsToLoad, FShouldLoadGCNotifyDelegate = FShouldLoadGCNotifyDelegate());
+	void BuildCuesToAddToGlobalSet(const TArray<FAssetData>& AssetDataList, FName TagPropertyName, TArray<struct FGameplayCueReferencePair>& OutCuesToAdd, TArray<FSoftObjectPath>& OutAssetsToLoad, FShouldLoadGCNotifyDelegate = FShouldLoadGCNotifyDelegate());
 
 	/** The cue manager has a tendency to produce a lot of RPCs. This logs out when we are attempting to fire more RPCs than will actually go off */
 	void CheckForTooManyRPCs(FName FuncName, const FGameplayCuePendingExecute& PendingCue, const FString& CueID, const FGameplayEffectContext* EffectContext);
 
-	void OnGameplayCueNotifyAsyncLoadComplete(TArray<FStringAssetReference> StringRef);
+	void OnGameplayCueNotifyAsyncLoadComplete(TArray<FSoftObjectPath> StringRef);
 
-	void OnMissingCueAsyncLoadComplete(FStringAssetReference LoadedObject, TWeakObjectPtr<UGameplayCueSet> OwningSet, FGameplayTag GameplayCueTag, TWeakObjectPtr<AActor> TargetActor, EGameplayCueEvent::Type EventType, FGameplayCueParameters Parameters);
+	void OnMissingCueAsyncLoadComplete(FSoftObjectPath LoadedObject, TWeakObjectPtr<UGameplayCueSet> OwningSet, FGameplayTag GameplayCueTag, TWeakObjectPtr<AActor> TargetActor, EGameplayCueEvent::Type EventType, FGameplayCueParameters Parameters);
 
 	void CheckForPreallocation(UClass* GCClass);
 
diff --git a/Engine/Plugins/Runtime/GameplayAbilities/Source/GameplayAbilities/Public/GameplayCueSet.h b/Engine/Plugins/Runtime/GameplayAbilities/Source/GameplayAbilities/Public/GameplayCueSet.h
--- a/Engine/Plugins/Runtime/GameplayAbilities/Source/GameplayAbilities/Public/GameplayCueSet.h
+++ b/Engine/Plugins/Runtime/GameplayAbilities/Source/GameplayAbilities/Public/GameplayCueSet.h
@@ -4,7 +4,7 @@
 
 #include "CoreMinimal.h"
 #include "UObject/ObjectMacros.h"
-#include "Misc/StringAssetReference.h"
+#include "UObject/SoftObjectPath.h"
 #include "GameplayTagContainer.h"
 #include "Engine/DataAsset.h"
 #include "GameplayEffectTypes.h"
@@ -25,7 +25,7 @@ struct FGameplayCueNotifyData
 	FGameplayTag GameplayCueTag;
 
 	UPROPERTY(EditAnywhere, Category=GameplayCue, meta=(AllowedClasses="GameplayCueNotify"))
-	FStringAssetReference GameplayCueNotifyObj;
+	FSoftObjectPath GameplayCueNotifyObj;
 
 	UPROPERTY(transient)
 	UClass* LoadedGameplayCueClass;
@@ -36,9 +36,9 @@ struct FGameplayCueNotifyData
 struct FGameplayCueReferencePair
 {
 	FGameplayTag GameplayCueTag;
-	FStringAssetReference StringRef;
+	FSoftObjectPath StringRef;
 
-	FGameplayCueReferencePair(const FGameplayTag& InGameplayCueTag, const FStringAssetReference& InStringRef)
+	FGameplayCueReferencePair(const FGameplayTag& InGameplayCueTag, const FSoftObjectPath& InStringRef)
 		: GameplayCueTag(InGameplayCueTag)
 		, StringRef(InStringRef)
 	{}
@@ -62,7 +62,7 @@ class GAMEPLAYABILITIES_API UGameplayCueSet : public UDataAsset
 	virtual void RemoveCuesByTags(const FGameplayTagContainer& TagsToRemove);
 
 	/** Removes all cues from the set matching the supplied string refs */
-	virtual void RemoveCuesByStringRefs(const TArray<FStringAssetReference>& CuesToRemove);
+	virtual void RemoveCuesByStringRefs(const TArray<FSoftObjectPath>& CuesToRemove);
 
 	/** Nulls reference to the loaded class. Note this doesn't remove the entire cue from the internal data structure, just the hard ref to the loaded class */
 	virtual void RemoveLoadedClass(UClass* Class);
@@ -70,14 +70,15 @@ class GAMEPLAYABILITIES_API UGameplayCueSet : public UDataAsset
 	/** Returns filenames of everything we know about (loaded or not) */
 	virtual void GetFilenames(TArray<FString>& Filenames) const;
 
-	virtual void GetStringAssetReferences(TArray<FStringAssetReference>& List) const;
+	/** Extracts all soft object paths pointing to Cues */
+	virtual void GetSoftObjectPaths(TArray<FSoftObjectPath>& List) const;
 
 #if WITH_EDITOR
 
 	void CopyCueDataToSetForEditorPreview(FGameplayTag Tag, UGameplayCueSet* DestinationSet);
 
 	/** Updates an existing cue */
-	virtual void UpdateCueByStringRefs(const FStringAssetReference& CueToRemove, FString NewPath);
+	virtual void UpdateCueByStringRefs(const FSoftObjectPath& CueToRemove, FString NewPath);
 
 #endif
 
diff --git a/Engine/Plugins/Runtime/GameplayAbilities/Source/GameplayAbilitiesEditor/Private/GameplayAbilitiesEditor.cpp b/Engine/Plugins/Runtime/GameplayAbilities/Source/GameplayAbilitiesEditor/Private/GameplayAbilitiesEditor.cpp
--- a/Engine/Plugins/Runtime/GameplayAbilities/Source/GameplayAbilitiesEditor/Private/GameplayAbilitiesEditor.cpp
+++ b/Engine/Plugins/Runtime/GameplayAbilities/Source/GameplayAbilitiesEditor/Private/GameplayAbilitiesEditor.cpp
@@ -45,8 +45,7 @@ void FGameplayAbilitiesEditor::InitGameplayAbilitiesEditor(const EToolkitMode::T
 void FGameplayAbilitiesEditor::EnsureGameplayAbilityBlueprintIsUpToDate(UBlueprint* Blueprint)
 {
 #if WITH_EDITORONLY_DATA
-	int32 Count = Blueprint->UbergraphPages.Num();
-	for (auto Graph : Blueprint->UbergraphPages)
+	for (UEdGraph* Graph : Blueprint->UbergraphPages)
 	{
 		// remove the default event graph, if it exists, from existing Gameplay Ability Blueprints
 		if (Graph->GetName() == "EventGraph" && Graph->Nodes.Num() == 0)
diff --git a/Engine/Plugins/Runtime/GameplayAbilities/Source/GameplayAbilitiesEditor/Private/GameplayCueTagDetails.cpp b/Engine/Plugins/Runtime/GameplayAbilities/Source/GameplayAbilitiesEditor/Private/GameplayCueTagDetails.cpp
--- a/Engine/Plugins/Runtime/GameplayAbilities/Source/GameplayAbilitiesEditor/Private/GameplayCueTagDetails.cpp
+++ b/Engine/Plugins/Runtime/GameplayAbilities/Source/GameplayAbilitiesEditor/Private/GameplayCueTagDetails.cpp
@@ -88,7 +88,7 @@ void FGameplayCueTagDetails::CustomizeChildren( TSharedRef<IPropertyHandle> Stru
 			.Padding(2.0f, 2.0f)
 			.AutoWidth()
 			[
-				SAssignNew(ListView, SListView < TSharedRef<FStringAssetReference> >)
+				SAssignNew(ListView, SListView < TSharedRef<FSoftObjectPath> >)
 				.ItemHeight(48)
 				.SelectionMode(ESelectionMode::None)
 				.ListItemsSource(&NotifyList)
@@ -109,7 +109,7 @@ void FGameplayCueTagDetails::CustomizeChildren( TSharedRef<IPropertyHandle> Stru
 }
 END_SLATE_FUNCTION_BUILD_OPTIMIZATION
 
-TSharedRef<ITableRow> FGameplayCueTagDetails::GenerateListRow(TSharedRef<FStringAssetReference> NotifyName, const TSharedRef<STableViewBase>& OwnerTable)
+TSharedRef<ITableRow> FGameplayCueTagDetails::GenerateListRow(TSharedRef<FSoftObjectPath> NotifyName, const TSharedRef<STableViewBase>& OwnerTable)
 {	
 	FString ShortName = NotifyName->ToString();
 	int32 idx;
@@ -120,7 +120,7 @@ TSharedRef<ITableRow> FGameplayCueTagDetails::GenerateListRow(TSharedRef<FString
 	}
 
 	return
-	SNew(STableRow< TSharedRef<FStringAssetReference> >, OwnerTable)
+	SNew(STableRow< TSharedRef<FSoftObjectPath> >, OwnerTable)
 	[
 		SNew(SBox)
 		.HAlign(HAlign_Left)
@@ -133,7 +133,7 @@ TSharedRef<ITableRow> FGameplayCueTagDetails::GenerateListRow(TSharedRef<FString
 	];
 }
 
-void FGameplayCueTagDetails::NavigateToHandler(TSharedRef<FStringAssetReference> AssetRef)
+void FGameplayCueTagDetails::NavigateToHandler(TSharedRef<FSoftObjectPath> AssetRef)
 {
 	SGameplayCueEditor::OpenEditorForNotify(AssetRef->ToString());
 }
@@ -178,7 +178,7 @@ bool FGameplayCueTagDetails::UpdateNotifyList()
 					if (CueSet->GameplayCueData.IsValidIndex(idx))
 					{
 						FGameplayCueNotifyData& Data = CueSet->GameplayCueData[*idxPtr];
-						TSharedRef<FStringAssetReference> Item(MakeShareable(new FStringAssetReference(Data.GameplayCueNotifyObj)));
+						TSharedRef<FSoftObjectPath> Item(MakeShareable(new FSoftObjectPath(Data.GameplayCueNotifyObj)));
 						NotifyList.Add(Item);
 					}
 				}
diff --git a/Engine/Plugins/Runtime/GameplayAbilities/Source/GameplayAbilitiesEditor/Private/GameplayCueTagDetails.h b/Engine/Plugins/Runtime/GameplayAbilities/Source/GameplayAbilitiesEditor/Private/GameplayCueTagDetails.h
--- a/Engine/Plugins/Runtime/GameplayAbilities/Source/GameplayAbilitiesEditor/Private/GameplayCueTagDetails.h
+++ b/Engine/Plugins/Runtime/GameplayAbilities/Source/GameplayAbilitiesEditor/Private/GameplayCueTagDetails.h
@@ -3,7 +3,7 @@
 #pragma once
 
 #include "CoreMinimal.h"
-#include "Misc/StringAssetReference.h"
+#include "UObject/SoftObjectPath.h"
 #include "Layout/Visibility.h"
 #include "Input/Reply.h"
 #include "IPropertyTypeCustomization.h"
@@ -32,11 +32,11 @@ private:
 
 	TSharedPtr<IPropertyHandle> GameplayTagProperty;
 
-	TSharedRef<ITableRow> GenerateListRow(TSharedRef<FStringAssetReference> NotifyName, const TSharedRef<STableViewBase>& OwnerTable);
-	TArray<TSharedRef<FStringAssetReference> > NotifyList;
-	TSharedPtr< SListView < TSharedRef<FStringAssetReference> > > ListView;
+	TSharedRef<ITableRow> GenerateListRow(TSharedRef<FSoftObjectPath> NotifyName, const TSharedRef<STableViewBase>& OwnerTable);
+	TArray<TSharedRef<FSoftObjectPath> > NotifyList;
+	TSharedPtr< SListView < TSharedRef<FSoftObjectPath> > > ListView;
 
-	void NavigateToHandler(TSharedRef<FStringAssetReference> AssetRef);
+	void NavigateToHandler(TSharedRef<FSoftObjectPath> AssetRef);
 
 	FReply OnAddNewNotifyClicked();
 
diff --git a/Engine/Plugins/Runtime/GameplayAbilities/Source/GameplayAbilitiesEditor/Private/SGameplayCueEditor.cpp b/Engine/Plugins/Runtime/GameplayAbilities/Source/GameplayAbilitiesEditor/Private/SGameplayCueEditor.cpp
--- a/Engine/Plugins/Runtime/GameplayAbilities/Source/GameplayAbilitiesEditor/Private/SGameplayCueEditor.cpp
+++ b/Engine/Plugins/Runtime/GameplayAbilities/Source/GameplayAbilitiesEditor/Private/SGameplayCueEditor.cpp
@@ -7,7 +7,6 @@
 #include "Stats/StatsMisc.h"
 #include "Misc/ScopedSlowTask.h"
 #include "UObject/Package.h"
-#include "Misc/StringAssetReference.h"
 #include "Misc/PackageName.h"
 #include "Widgets/SBoxPanel.h"
 #include "Styling/SlateTypes.h"
@@ -78,8 +77,8 @@ struct FGCTreeItem : public TSharedFromThis<FGCTreeItem>
 	FGameplayTag GameplayCueTag;
 	FString Description;
 
-	FStringAssetReference GameplayCueNotifyObj;
-	FStringAssetReference ParentGameplayCueNotifyObj;
+	FSoftObjectPath GameplayCueNotifyObj;
+	FSoftObjectPath ParentGameplayCueNotifyObj;
 	TWeakObjectPtr<UFunction> FunctionPtr;
 
 	int32 TranslationUniqueID;
@@ -1386,7 +1385,7 @@ void SGameplayCueEditor::OpenEditorForNotify(FString NotifyFullPath)
 	// This nonsense is to handle the case where the asset only exists in memory
 	// and there for does not have a linker/exist on disk. (The FString version
 	// of OpenEditorForAsset does not handle this).
-	FStringAssetReference AssetRef(NotifyFullPath);
+	FSoftObjectPath AssetRef(NotifyFullPath);
 
 	UObject* Obj = AssetRef.ResolveObject();
 	if (!Obj)
diff --git a/Engine/Plugins/Runtime/Oculus/OculusVR/Source/OculusHMD/Public/OculusHMDTypes.h b/Engine/Plugins/Runtime/Oculus/OculusVR/Source/OculusHMD/Public/OculusHMDTypes.h
--- a/Engine/Plugins/Runtime/Oculus/OculusVR/Source/OculusHMD/Public/OculusHMDTypes.h
+++ b/Engine/Plugins/Runtime/Oculus/OculusVR/Source/OculusHMD/Public/OculusHMDTypes.h
@@ -7,7 +7,7 @@
 #include "RHI.h"
 #include "RHIResources.h"
 #include "Engine/Texture2D.h"
-#include "Misc/StringAssetReference.h"
+#include "UObject/SoftObjectPath.h"
 #include "OculusHMDTypes.generated.h"
 
 USTRUCT()
diff --git a/Engine/Plugins/Runtime/SignificanceManager/Source/SignificanceManager/Private/SignificanceManager.cpp b/Engine/Plugins/Runtime/SignificanceManager/Source/SignificanceManager/Private/SignificanceManager.cpp
--- a/Engine/Plugins/Runtime/SignificanceManager/Source/SignificanceManager/Private/SignificanceManager.cpp
+++ b/Engine/Plugins/Runtime/SignificanceManager/Source/SignificanceManager/Private/SignificanceManager.cpp
@@ -115,7 +115,7 @@ void FSignificanceManagerModule::PopulateAutoCompleteEntries(TArray<FAutoComplet
 USignificanceManager::USignificanceManager()
 	: Super()
 {
-	SignificanceManagerClassName = FStringClassReference(GetClass()); 
+	SignificanceManagerClassName = FSoftClassPath(GetClass()); 
 
 	bCreateOnClient = true;
 	bCreateOnServer = true;
diff --git a/Engine/Plugins/Runtime/SignificanceManager/Source/SignificanceManager/Public/SignificanceManager.h b/Engine/Plugins/Runtime/SignificanceManager/Source/SignificanceManager/Public/SignificanceManager.h
--- a/Engine/Plugins/Runtime/SignificanceManager/Source/SignificanceManager/Public/SignificanceManager.h
+++ b/Engine/Plugins/Runtime/SignificanceManager/Source/SignificanceManager/Public/SignificanceManager.h
@@ -9,7 +9,7 @@
 #include "Templates/SubclassOf.h"
 #include "Modules/ModuleManager.h"
 #include "UObject/GCObject.h"
-#include "Misc/StringClassReference.h"
+#include "UObject/SoftObjectPath.h"
 #include "Engine/World.h"
 #include "SignificanceManager.generated.h"
 
@@ -191,7 +191,7 @@ private:
 
 	// Game specific significance class to instantiate
 	UPROPERTY(globalconfig, noclear, EditAnywhere, Category=DefaultClasses, meta=(MetaClass="SignificanceManager", DisplayName="Significance Manager Class"))
-	FStringClassReference SignificanceManagerClassName;
+	FSoftClassPath SignificanceManagerClassName;
 
 	// Callback function registered with HUD to supply debug info when ShowDebug SignificanceManager has been entered on the console
 	void OnShowDebugInfo(AHUD* HUD, UCanvas* Canvas, const FDebugDisplayInfo& DisplayInfo, float& YL, float& YPos);
diff --git a/Engine/Plugins/Runtime/SoundMod/Source/SoundMod/Public/SoundNodeModPlayer.h b/Engine/Plugins/Runtime/SoundMod/Source/SoundMod/Public/SoundNodeModPlayer.h
--- a/Engine/Plugins/Runtime/SoundMod/Source/SoundMod/Public/SoundNodeModPlayer.h
+++ b/Engine/Plugins/Runtime/SoundMod/Source/SoundMod/Public/SoundNodeModPlayer.h
@@ -5,7 +5,7 @@
 
 #include "CoreMinimal.h"
 #include "UObject/ObjectMacros.h"
-#include "UObject/AssetPtr.h"
+#include "UObject/SoftObjectPtr.h"
 #include "Sound/SoundNodeAssetReferencer.h"
 #include "SoundNodeModPlayer.generated.h"
 
@@ -26,7 +26,7 @@ class USoundNodeModPlayer : public USoundNodeAssetReferencer
 
 private:
 	UPROPERTY(EditAnywhere, Category=ModPlayer, meta=(DisplayName="Sound Mod"))
-	TAssetPtr<USoundMod> SoundModAssetPtr;
+	TSoftObjectPtr<USoundMod> SoundModAssetPtr;
 
 	UPROPERTY(transient)
 	USoundMod* SoundMod;
diff --git a/Engine/Plugins/Runtime/Steam/SteamVR/Source/SteamVR/Private/SteamVRAssetManager.h b/Engine/Plugins/Runtime/Steam/SteamVR/Source/SteamVR/Private/SteamVRAssetManager.h
--- a/Engine/Plugins/Runtime/Steam/SteamVR/Source/SteamVR/Private/SteamVRAssetManager.h
+++ b/Engine/Plugins/Runtime/Steam/SteamVR/Source/SteamVR/Private/SteamVRAssetManager.h
@@ -3,7 +3,7 @@
 #pragma once
 
 #include "IXRDeviceAssets.h"
-#include "UObject/AssetPtr.h"
+#include "UObject/SoftObjectPtr.h"
 
 class  FSteamVRAsyncMeshLoader;
 struct FSteamVRMeshData;
@@ -36,5 +36,5 @@ protected:
 
 private:
 	TArray< TSharedPtr<FSteamVRAsyncMeshLoader> > AsyncMeshLoaders;
-	TAssetPtr<UMaterial> DefaultDeviceMat;
+	TSoftObjectPtr<UMaterial> DefaultDeviceMat;
 };
diff --git a/Engine/Plugins/ScriptPlugin/Source/ScriptGeneratorPlugin/Private/LuaScriptCodeGenerator.cpp b/Engine/Plugins/ScriptPlugin/Source/ScriptGeneratorPlugin/Private/LuaScriptCodeGenerator.cpp
--- a/Engine/Plugins/ScriptPlugin/Source/ScriptGeneratorPlugin/Private/LuaScriptCodeGenerator.cpp
+++ b/Engine/Plugins/ScriptPlugin/Source/ScriptGeneratorPlugin/Private/LuaScriptCodeGenerator.cpp
@@ -315,8 +315,8 @@ bool FLuaScriptCodeGenerator::IsPropertyTypeSupported(UProperty* Property) const
 		}
 	}
 	else if (Property->IsA(ULazyObjectProperty::StaticClass()) ||
-		Property->IsA(UAssetObjectProperty::StaticClass()) ||
-		Property->IsA(UAssetClassProperty::StaticClass()) ||
+		Property->IsA(USoftObjectProperty::StaticClass()) ||
+		Property->IsA(USoftClassProperty::StaticClass()) ||
 		Property->IsA(UWeakObjectProperty::StaticClass()))
 	{
 		bSupported = false;
diff --git a/Engine/Plugins/Tests/EditorTests/Source/EditorTests/EditorTests.Build.cs b/Engine/Plugins/Tests/EditorTests/Source/EditorTests/EditorTests.Build.cs
--- a/Engine/Plugins/Tests/EditorTests/Source/EditorTests/EditorTests.Build.cs
+++ b/Engine/Plugins/Tests/EditorTests/Source/EditorTests/EditorTests.Build.cs
@@ -6,32 +6,32 @@ public class EditorTests : ModuleRules
 {
 	public EditorTests(ReadOnlyTargetRules Target) : base(Target)
 	{
-		
+
 		PublicIncludePaths.AddRange(
 			new string[] {
 				"EditorTests/Public"
 				// ... add public include paths required here ...
 			}
-			);
-				
-		
+		);
+
+
 		PrivateIncludePaths.AddRange(
 			new string[] {
 				"EditorTests/Private",
 				// ... add other private include paths required here ...
 			}
-			);
-			
-		
+		);
+
+
 		PublicDependencyModuleNames.AddRange(
 			new string[]
 			{
 				"Core",
 				// ... add other public dependencies that you statically link with here ...
 			}
-			);
-			
-		
+		);
+
+
 		PrivateDependencyModuleNames.AddRange(
 			new string[]
 			{
@@ -39,36 +39,38 @@ public class EditorTests : ModuleRules
 				"InputCore",
 				"LevelEditor",
 				"CoreUObject",
-                "ShaderCore",
+				"ShaderCore",
 				"Engine",
 				"Slate",
 				"SlateCore",
-                "AssetTools",
-                "MainFrame",
-                "MaterialEditor",
-                "JsonUtilities",
-                "Analytics",
-                "ContentBrowser",
-                "EditorStyle",
-                "SourceControl",
-                "RHI",
-                "BlueprintGraph",
-                "AddContentDialog",
-                "GraphEditor",
-                "DirectoryWatcher",
+				"AssetTools",
+				"MainFrame",
+				"MaterialEditor",
+				"JsonUtilities",
+				"Analytics",
+				"ContentBrowser",
+				"EditorStyle",
+				"SourceControl",
+				"RHI",
+				"BlueprintGraph",
+				"AddContentDialog",
+				"GraphEditor",
+				"DirectoryWatcher",
 				"Projects",
-                "UnrealEd",
-                "AudioEditor",
-                "AnimGraphRuntime",
-            }
-			);
-		
-		
+				"UnrealEd",
+				"AudioEditor",
+				"AnimGraphRuntime",
+				"MeshMergeUtilities",
+				"MaterialBaking"
+			}
+		);
+
+
 		DynamicallyLoadedModuleNames.AddRange(
 			new string[]
 			{
 				// ... add any modules that your module loads dynamically here ...
 			}
-			);
+		);
 	}
 }
diff --git a/Engine/Plugins/Tests/EditorTests/Source/EditorTests/Private/UnrealEd/BlueprintAutomationTests.cpp b/Engine/Plugins/Tests/EditorTests/Source/EditorTests/Private/UnrealEd/BlueprintAutomationTests.cpp
--- a/Engine/Plugins/Tests/EditorTests/Source/EditorTests/Private/UnrealEd/BlueprintAutomationTests.cpp
+++ b/Engine/Plugins/Tests/EditorTests/Source/EditorTests/Private/UnrealEd/BlueprintAutomationTests.cpp
@@ -13,7 +13,7 @@
 #include "UObject/Class.h"
 #include "UObject/UObjectIterator.h"
 #include "UObject/Package.h"
-#include "Misc/StringAssetReference.h"
+#include "UObject/SoftObjectPath.h"
 #include "UObject/MetaData.h"
 #include "UObject/UnrealType.h"
 #include "Serialization/ObjectWriter.h"
@@ -913,7 +913,7 @@ bool FBlueprintCompileOnLoadTest::RunTest(const FString& BlueprintAssetPath)
 	struct FReplaceInnerData
 	{
 		TWeakObjectPtr<UClass> Class;
-		FStringAssetReference BlueprintAsset;
+		FSoftObjectPath BlueprintAsset;
 	};
 	TArray<FReplaceInnerData> ReplaceInnerData;
 	for (auto BPToUnloadWP : BlueprintDependencies)
@@ -924,7 +924,7 @@ bool FBlueprintCompileOnLoadTest::RunTest(const FString& BlueprintAssetPath)
 		{
 			FReplaceInnerData Data;
 			Data.Class = OldClass;
-			Data.BlueprintAsset = FStringAssetReference(BPToUnload);
+			Data.BlueprintAsset = FSoftObjectPath(BPToUnload);
 			ReplaceInnerData.Add(Data);
 		}
 	}
diff --git a/Engine/Plugins/Tests/EditorTests/Source/EditorTests/Private/UnrealEd/EditorOpenAssetAutomationTests.cpp b/Engine/Plugins/Tests/EditorTests/Source/EditorTests/Private/UnrealEd/EditorOpenAssetAutomationTests.cpp
--- a/Engine/Plugins/Tests/EditorTests/Source/EditorTests/Private/UnrealEd/EditorOpenAssetAutomationTests.cpp
+++ b/Engine/Plugins/Tests/EditorTests/Source/EditorTests/Private/UnrealEd/EditorOpenAssetAutomationTests.cpp
@@ -20,7 +20,7 @@ IMPLEMENT_COMPLEX_AUTOMATION_TEST(FOpenAssetEditors, "Project.Editor.Open Assets
 void FOpenAssetEditors::GetTests(TArray<FString>& OutBeautifiedNames, TArray<FString>& OutTestCommands) const
 {
 	const UAutomationTestSettings* AutomationTestSettings = GetDefault<UAutomationTestSettings>();
-	for ( FStringAssetReference AssetRef : AutomationTestSettings->AssetsToOpen )
+	for ( FSoftObjectPath AssetRef : AutomationTestSettings->AssetsToOpen )
 	{
 		OutBeautifiedNames.Add(AssetRef.GetAssetName());
 		OutTestCommands.Add(AssetRef.GetLongPackageName());
@@ -33,7 +33,7 @@ bool FOpenAssetEditors::RunTest(const FString& LongAssetPath)
 	FAssetEditorManager::Get().CloseAllAssetEditors();
 
 	// below is all latent action, so before sending there, verify the asset exists
-	UObject* Object = FStringAssetReference(LongAssetPath).TryLoad();
+	UObject* Object = FSoftObjectPath(LongAssetPath).TryLoad();
 	if ( Object == nullptr )
 	{
 		UE_LOG(LogEditorAutomationTests, Error, TEXT("Failed to Open Asset '%s'."), *LongAssetPath);
diff --git a/Engine/Plugins/Tests/EditorTests/Source/EditorTests/Private/UnrealEd/EditorPerformanceTests.cpp b/Engine/Plugins/Tests/EditorTests/Source/EditorTests/Private/UnrealEd/EditorPerformanceTests.cpp
--- a/Engine/Plugins/Tests/EditorTests/Source/EditorTests/Private/UnrealEd/EditorPerformanceTests.cpp
+++ b/Engine/Plugins/Tests/EditorTests/Source/EditorTests/Private/UnrealEd/EditorPerformanceTests.cpp
@@ -238,7 +238,7 @@ void FMapPerformanceInEditor::GetTests(TArray<FString>& OutBeautifiedNames, TArr
 
 	for ( const FEditorMapPerformanceTestDefinition& PerfDefinition : AutomationTestSettings->EditorPerformanceTestMaps )
 	{
-		const FStringAssetReference& PerfMap = PerfDefinition.PerformanceTestmap;
+		const FSoftObjectPath& PerfMap = PerfDefinition.PerformanceTestmap;
 
 		if ( PerfMap.IsValid() )
 		{
diff --git a/Engine/Plugins/Tests/FunctionalTestingEditor/Source/Private/FunctionalTestingEditorModule.cpp b/Engine/Plugins/Tests/FunctionalTestingEditor/Source/Private/FunctionalTestingEditorModule.cpp
new file mode 100644
--- /dev/null
+++ b/Engine/Plugins/Tests/FunctionalTestingEditor/Source/Private/FunctionalTestingEditorModule.cpp
@@ -0,0 +1,205 @@
+// Copyright 1998-2017 Epic Games, Inc. All Rights Reserved.
+
+#include "FunctionalTestingEditorModule.h"
+#include "HAL/IConsoleManager.h"
+#include "Modules/ModuleManager.h"
+#include "UObject/UObjectHash.h"
+#include "Textures/SlateIcon.h"
+#include "Framework/Commands/UIAction.h"
+#include "Framework/MultiBox/MultiBoxExtender.h"
+#include "Framework/MultiBox/MultiBoxBuilder.h"
+#include "Framework/Docking/TabManager.h"
+#include "EditorStyleSet.h"
+#include "GameFramework/Actor.h"
+#include "AssetData.h"
+#include "EdGraph/EdGraphSchema.h"
+#include "LevelEditor.h"
+#include "ISessionFrontendModule.h"
+#include "IPlacementModeModule.h"
+#include "FunctionalTest.h"
+#include "ScreenshotFunctionalTest.h"
+#include "AssetRegistryModule.h"
+#include "WorkspaceMenuStructure.h"
+#include "WorkspaceMenuStructureModule.h"
+#include "Widgets/Docking/SDockTab.h"
+#include "AssetTypeActions_GroundTruthData.h"
+#include "IAssetTools.h"
+
+#define LOCTEXT_NAMESPACE "EditorAutomation"
+
+void OpenMapAndFocusActor(const TArray<FString>& Args)
+{
+	if ( Args.Num() != 2 )
+	{
+		UE_LOG(LogConsoleResponse, Display, TEXT("Automate.OpenMapAndFocusActor failed, the number of arguments is wrong.  Automate.OpenMapAndFocusActor MapObjectPath ActorName\n"));
+		return;
+	}
+
+	FString AssetPath(Args[0]);
+	FString ActorName(Args[1]);
+
+	FAssetRegistryModule& AssetRegistryModule = FModuleManager::LoadModuleChecked<FAssetRegistryModule>(TEXT("AssetRegistry"));
+
+	FAssetData MapAssetData = AssetRegistryModule.Get().GetAssetByObjectPath(*AssetPath);
+
+	if ( MapAssetData.IsValid() )
+	{
+		bool bIsWorlAlreadyOpened = false;
+		if ( UWorld* EditorWorld = GEditor->GetEditorWorldContext().World() )
+		{
+			if ( FAssetData(EditorWorld).PackageName == MapAssetData.PackageName )
+			{
+				bIsWorlAlreadyOpened = true;
+			}
+		}
+
+		if ( !bIsWorlAlreadyOpened )
+		{
+			UObject* ObjectToEdit = MapAssetData.GetAsset();
+			if ( ObjectToEdit )
+			{
+				GEditor->EditObject(ObjectToEdit);
+			}
+		}
+	}
+
+	if ( UWorld* EditorWorld = GEditor->GetEditorWorldContext().World() )
+	{
+		AActor* ActorToFocus = nullptr;
+		for ( int32 LevelIndex=0; LevelIndex < EditorWorld->GetNumLevels(); LevelIndex++ )
+		{
+			ULevel* Level = EditorWorld->GetLevel(LevelIndex);
+			ActorToFocus = static_cast<AActor*>( FindObjectWithOuter(Level, AActor::StaticClass(), *Args[1]) );
+			if ( ActorToFocus )
+			{
+				break;
+			}
+		}
+
+		if ( ActorToFocus )
+		{
+			GEditor->SelectNone(/*bNoteSelectionChange=*/ false, false, false);
+			GEditor->SelectActor(ActorToFocus, true, true);
+			GEditor->NoteSelectionChange();
+
+			const bool bActiveViewportOnly = false;
+			GEditor->MoveViewportCamerasToActor(*ActorToFocus, bActiveViewportOnly);
+
+			FLevelEditorModule& LevelEditorModule = FModuleManager::LoadModuleChecked<FLevelEditorModule>("LevelEditor");
+			TSharedPtr<SDockTab> LevelEditorTab = LevelEditorModule.GetLevelEditorInstanceTab().Pin();
+			if ( LevelEditorTab.IsValid() )
+			{
+				LevelEditorTab->DrawAttention();
+			}
+		}
+	}
+}
+
+FAutoConsoleCommand OpenMapAndFocusActorCmd(
+	TEXT("Automate.OpenMapAndFocusActor"),
+	TEXT("Opens a map and focuses a particular actor by name."),
+	FConsoleCommandWithArgsDelegate::CreateStatic(OpenMapAndFocusActor)
+);
+
+class FFunctionalTestingEditorModule : public IFunctionalTestingEditorModule
+{
+	void StartupModule()
+	{
+		// make an extension to add the Orion function menu
+		Extender = MakeShareable(new FExtender());
+		Extender->AddMenuExtension(
+			"General",
+			EExtensionHook::After,
+			nullptr,
+			FMenuExtensionDelegate::CreateRaw(this, &FFunctionalTestingEditorModule::OnAutomationToolsMenuCreation));
+
+		// add the menu extension to the editor
+		FLevelEditorModule& LevelEditorModule = FModuleManager::LoadModuleChecked<FLevelEditorModule>("LevelEditor");
+		LevelEditorModule.GetMenuExtensibilityManager()->AddExtender(Extender);
+
+		FModuleManager::Get().OnModulesChanged().AddRaw(this, &FFunctionalTestingEditorModule::OnModulesChanged);
+
+		if ( IPlacementModeModule::IsAvailable() )
+		{
+			OnModulesChanged("PlacementMode", EModuleChangeReason::ModuleLoaded);
+		}
+
+		IAssetTools& AssetTools = FModuleManager::LoadModuleChecked<FAssetToolsModule>("AssetTools").Get();
+		AssetTools.RegisterAssetTypeActions(MakeShareable(new FAssetTypeActions_GroundTruthData()));
+	}
+
+	void ShutdownModule()
+	{
+		if ( FModuleManager::Get().IsModuleLoaded("LevelEditor") )
+		{
+			FLevelEditorModule& LevelEditorModule = FModuleManager::LoadModuleChecked<FLevelEditorModule>("LevelEditor");
+			LevelEditorModule.GetMenuExtensibilityManager()->RemoveExtender(Extender);
+		}
+
+		FModuleManager::Get().OnModulesChanged().RemoveAll(this);
+
+		if ( IPlacementModeModule::IsAvailable() )
+		{
+			IPlacementModeModule::Get().UnregisterPlacementCategory("Testing");
+		}
+	}
+
+	void OnAutomationToolsMenuCreation(FMenuBuilder& MenuBuilder)
+	{
+		MenuBuilder.BeginSection("Testing", LOCTEXT("Testing", "Testing"));
+		MenuBuilder.AddMenuEntry(
+			LOCTEXT("AutomationLabel", "Test Automation"),
+			LOCTEXT("Tooltip", "Launch the Testing Automation Frontend."),
+			FSlateIcon(FEditorStyle::GetStyleSetName(), "AutomationTools.MenuIcon"),
+			FUIAction(FExecuteAction::CreateStatic(&FFunctionalTestingEditorModule::OnShowAutomationFrontend)));
+		MenuBuilder.EndSection();
+
+		// TODO Come up with a way to hide this if no tools are registered.
+		if (WorkspaceMenu::GetMenuStructure().GetAutomationToolsCategory()->GetChildItems().Num() > 0)
+		{
+			MenuBuilder.AddSubMenu(
+				LOCTEXT("AutomationTools", "Automation Tools"),
+				LOCTEXT("AutomationToolsToolTip", "Assorted tools to help generate data for some of the automation tests."),
+				FNewMenuDelegate::CreateRaw(this, &FFunctionalTestingEditorModule::PopulateAutomationTools)
+			);
+		}
+	}
+
+	void PopulateAutomationTools(FMenuBuilder& MenuBuilder)
+	{
+		MenuBuilder.BeginSection("AutomationTools", LOCTEXT("AutomationTools", "Automation Tools"));
+		const bool bAutoAndOrphanedMenus = false;
+		FGlobalTabmanager::Get()->PopulateTabSpawnerMenu(MenuBuilder, WorkspaceMenu::GetMenuStructure().GetAutomationToolsCategory(), bAutoAndOrphanedMenus);
+		MenuBuilder.EndSection();
+	}
+
+	static void OnShowAutomationFrontend()
+	{
+		ISessionFrontendModule& SessionFrontend = FModuleManager::LoadModuleChecked<ISessionFrontendModule>("SessionFrontend");
+		SessionFrontend.InvokeSessionFrontend(FName("AutomationPanel"));
+	}
+
+	void OnModulesChanged(FName Module, EModuleChangeReason Reason)
+	{
+		if ( Module == TEXT("PlacementMode") && Reason == EModuleChangeReason::ModuleLoaded )
+		{
+			FPlacementCategoryInfo Info(
+				LOCTEXT("FunctionalTestingCategoryName", "Testing"),
+				"Testing",
+				TEXT("PMTesting"),
+				25
+				);
+
+			IPlacementModeModule::Get().RegisterPlacementCategory(Info);
+			IPlacementModeModule::Get().RegisterPlaceableItem(Info.UniqueHandle, MakeShareable(new FPlaceableItem(nullptr, FAssetData(AFunctionalTest::StaticClass()))));
+			IPlacementModeModule::Get().RegisterPlaceableItem(Info.UniqueHandle, MakeShareable(new FPlaceableItem(nullptr, FAssetData(AScreenshotFunctionalTest::StaticClass()))));
+		}
+	}
+
+private:
+	TSharedPtr<FExtender> Extender;
+};
+
+#undef LOCTEXT_NAMESPACE
+
+IMPLEMENT_MODULE(FFunctionalTestingEditorModule, FunctionalTestingEditor);
diff --git a/Engine/Programs/UnrealHeaderTool/Config/DefaultEngine.ini b/Engine/Programs/UnrealHeaderTool/Config/DefaultEngine.ini
--- a/Engine/Programs/UnrealHeaderTool/Config/DefaultEngine.ini
+++ b/Engine/Programs/UnrealHeaderTool/Config/DefaultEngine.ini
@@ -4,3 +4,26 @@
 +ScriptSupportedModules=ScriptPlugin
 
 [UnrealHeaderTool]
++StructsWithNoPrefix="uint64"
++StructsWithNoPrefix="uint32"
++StructsWithNoPrefix="double"
++StructsWithTPrefix="IndirectArray"
++StructsWithTPrefix="BitArray"
++StructsWithTPrefix="SparseArray"
++StructsWithTPrefix="Set"
++StructsWithTPrefix="Map"
++StructsWithTPrefix="MultiMap"
++StructsWithTPrefix="SharedPtr"
++DelegateParameterCountStrings="_OneParam"
++DelegateParameterCountStrings="_TwoParams"
++DelegateParameterCountStrings="_ThreeParams"
++DelegateParameterCountStrings="_FourParams"
++DelegateParameterCountStrings="_FiveParams"
++DelegateParameterCountStrings="_SixParams"
++DelegateParameterCountStrings="_SevenParams"
++DelegateParameterCountStrings="_EightParams"
++DelegateParameterCountStrings="_NineParams"
++TypeRedirects=(OldType="FStringAssetReference", NewType="FSoftObjectPath")
++TypeRedirects=(OldType="FStringClassReference", NewType="FSoftClassPath")
++TypeRedirects=(OldType="TAssetPtr", NewType="TSoftObjectPtr")
++TypeRedirects=(OldType="TAssetSubclassOf", NewType="TSoftClassPtr")
diff --git a/Engine/Source/Developer/AssetTools/Private/AssetFixUpRedirectors.cpp b/Engine/Source/Developer/AssetTools/Private/AssetFixUpRedirectors.cpp
--- a/Engine/Source/Developer/AssetTools/Private/AssetFixUpRedirectors.cpp
+++ b/Engine/Source/Developer/AssetTools/Private/AssetFixUpRedirectors.cpp
@@ -101,8 +101,8 @@ void FAssetFixUpRedirectors::ExecuteFixUp(TArray<TWeakObjectPtr<UObjectRedirecto
 				// If any referencing packages are left read-only, the checkout failed or SCC was not enabled. Trim them from the save list and leave redirectors.
 				DetectReadOnlyPackages(RedirectorRefsList, ReferencingPackagesToSave);
 
-				// Fix up referencing FStringAssetReferences
-				FixUpStringAssetReferences(RedirectorRefsList, ReferencingPackagesToSave);
+				// Fix up referencing FSoftObjectPaths
+				FixUpSoftObjectPaths(RedirectorRefsList, ReferencingPackagesToSave);
 
 				// Save all packages that were referencing any of the assets that were moved without redirectors
 				TArray<UPackage*> FailedToSave;
@@ -455,29 +455,32 @@ void FAssetFixUpRedirectors::ReportFailures(const TArray<FRedirectorRefs>& Redir
 	EditorErrors.Open();
 }
 
-void FAssetFixUpRedirectors::FixUpStringAssetReferences(const TArray<FRedirectorRefs>& RedirectorsToFix, const TArray<UPackage*>& InReferencingPackagesToSave) const
+void FAssetFixUpRedirectors::FixUpSoftObjectPaths(const TArray<FRedirectorRefs>& RedirectorsToFix, const TArray<UPackage*>& InReferencingPackagesToSave) const
 {
 	TArray<UPackage *> PackagesToCheck(InReferencingPackagesToSave);
 
 	FEditorFileUtils::GetDirtyWorldPackages(PackagesToCheck);
 	FEditorFileUtils::GetDirtyContentPackages(PackagesToCheck);
 
-	TMap<FString, FString> RedirectorMap;
+	TMap<FSoftObjectPath, FSoftObjectPath> RedirectorMap;
 
 	for (const FRedirectorRefs& RedirectorRef : RedirectorsToFix)
 	{
 		UObjectRedirector* Redirector = RedirectorRef.Redirector;
-		RedirectorMap.Add(Redirector->GetPathName(), Redirector->DestinationObject->GetPathName());
+		FSoftObjectPath OldPath = FSoftObjectPath(Redirector);
+		FSoftObjectPath NewPath = FSoftObjectPath(Redirector->DestinationObject);
+
+		RedirectorMap.Add(OldPath, NewPath);
 
 		if (UBlueprint* Blueprint = Cast<UBlueprint>(Redirector->DestinationObject))
 		{
 			// Add redirect for class and default as well
-			RedirectorMap.Add(FString::Printf(TEXT("%s_C"), *Redirector->GetPathName()), FString::Printf(TEXT("%s_C"), *Redirector->DestinationObject->GetPathName()));
-			RedirectorMap.Add(FString::Printf(TEXT("Default__%s_C"), *Redirector->GetPathName()), FString::Printf(TEXT("Default__%s_C"), *Redirector->DestinationObject->GetPathName()));
+			RedirectorMap.Add(FString::Printf(TEXT("%s_C"), *OldPath.ToString()), FString::Printf(TEXT("%s_C"), *NewPath.ToString()));
+			RedirectorMap.Add(FString::Printf(TEXT("%s.Default__%s_C"), *OldPath.GetLongPackageName(), *OldPath.GetAssetName()), FString::Printf(TEXT("%s.Default__%s_C"), *NewPath.GetLongPackageName(), *NewPath.GetAssetName()));
 		}
 	}
 
-	FAssetRenameManager::RenameReferencingStringAssetReferences(PackagesToCheck, RedirectorMap);
+	FAssetRenameManager::RenameReferencingSoftObjectPaths(PackagesToCheck, RedirectorMap);
 }
 
 #undef LOCTEXT_NAMESPACE
diff --git a/Engine/Source/Developer/AssetTools/Private/AssetFixUpRedirectors.h b/Engine/Source/Developer/AssetTools/Private/AssetFixUpRedirectors.h
--- a/Engine/Source/Developer/AssetTools/Private/AssetFixUpRedirectors.h
+++ b/Engine/Source/Developer/AssetTools/Private/AssetFixUpRedirectors.h
@@ -39,8 +39,8 @@ private:
 	/** Finds any read only packages and removes them from the save list. Redirectors referenced by these packages will not be fixed up. */ 
 	void DetectReadOnlyPackages(TArray<FRedirectorRefs>& RedirectorsToFix, TArray<UPackage*>& InOutReferencingPackagesToSave) const;
 
-	/** FixUp string asset references. */
-	void FixUpStringAssetReferences(const TArray<FRedirectorRefs>& RedirectorsToFix, const TArray<UPackage*>& InReferencingPackagesToSave) const;
+	/** FixUp soft object paths */
+	void FixUpSoftObjectPaths(const TArray<FRedirectorRefs>& RedirectorsToFix, const TArray<UPackage*>& InReferencingPackagesToSave) const;
 
 	/** Saves all the referencing packages and updates SCC state */
 	void SaveReferencingPackages(const TArray<UPackage*>& ReferencingPackagesToSave, TArray<UPackage*>& OutFailedToSave) const;
