Documentation for build tools

diff --git a/Engine/Documentation/Source/Programming/UnrealBuildSystem/ModuleFiles/ModuleFilesProperties/ModuleFilesProperties.INT.udn b/Engine/Documentation/Source/Programming/UnrealBuildSystem/ModuleFiles/ModuleFilesProperties/ModuleFilesProperties.INT.udn
--- a/Engine/Documentation/Source/Programming/UnrealBuildSystem/ModuleFiles/ModuleFilesProperties/ModuleFilesProperties.INT.udn
+++ b/Engine/Documentation/Source/Programming/UnrealBuildSystem/ModuleFiles/ModuleFilesProperties/ModuleFilesProperties.INT.udn
@@ -27,6 +27,8 @@ $ bEnableBufferSecurityChecks (Boolean): Enable buffer security checks.  This sh
 
 $ bEnableExceptions (Boolean): Enable exception handling
 
+$ bEnableObjCExceptions (Boolean): Enable objective C exception handling
+
 $ bEnableShadowVariableWarnings (Boolean): Enable warnings for shadowed variables
 
 $ bEnableUndefinedIdentifierWarnings (Boolean): Enable warnings for using undefined identifiers in #if expressions
diff --git a/Engine/Documentation/Source/Programming/UnrealBuildSystem/TargetFiles/TargetFilesProperties/TargetFilesProperties.INT.udn b/Engine/Documentation/Source/Programming/UnrealBuildSystem/TargetFiles/TargetFilesProperties/TargetFilesProperties.INT.udn
--- a/Engine/Documentation/Source/Programming/UnrealBuildSystem/TargetFiles/TargetFilesProperties/TargetFilesProperties.INT.udn
+++ b/Engine/Documentation/Source/Programming/UnrealBuildSystem/TargetFiles/TargetFilesProperties/TargetFilesProperties.INT.udn
@@ -96,6 +96,8 @@ $ bCompileSpeedTree (Boolean): Whether to compile SpeedTree support.
 
 $ bForceEnableExceptions (Boolean): Enable exceptions for all modules.
 
+$ bForceEnableObjCExceptions (Boolean): Enable exceptions for all modules.
+
 $ bForceEnableRTTI (Boolean): Enable RTTI for all modules.
 
 $ bWithServerCode (Boolean): Compile server-only code.
