Documentation for build tools

diff --git a/Engine/Documentation/Source/Programming/UnrealBuildSystem/Configuration/BuildConfigProperties/BuildConfigProperties.INT.udn b/Engine/Documentation/Source/Programming/UnrealBuildSystem/Configuration/BuildConfigProperties/BuildConfigProperties.INT.udn
--- a/Engine/Documentation/Source/Programming/UnrealBuildSystem/Configuration/BuildConfigProperties/BuildConfigProperties.INT.udn
+++ b/Engine/Documentation/Source/Programming/UnrealBuildSystem/Configuration/BuildConfigProperties/BuildConfigProperties.INT.udn
@@ -160,7 +160,7 @@ $ bUseDangerouslyFastMode : If this is set, then we don't do any post-compile st
 
 ### HTML5Platform
 
-$ HTML5Architecture : Architecture to build for. Use -win32 for win32 builds (build html5 platform as a win32 binary for debugging).
+$ HTML5Architecture : Architecture to build for.
 
 ### UEBuildConfiguration
 
