Adjusting position of image in Actor Lifecycle page. Also updating meta tags.

diff --git a/Engine/Documentation/Source/Programming/UnrealArchitecture/Actors/ActorLifecycle/ActorLifecycle.INT.udn b/Engine/Documentation/Source/Programming/UnrealArchitecture/Actors/ActorLifecycle/ActorLifecycle.INT.udn
--- a/Engine/Documentation/Source/Programming/UnrealArchitecture/Actors/ActorLifecycle/ActorLifecycle.INT.udn
+++ b/Engine/Documentation/Source/Programming/UnrealArchitecture/Actors/ActorLifecycle/ActorLifecycle.INT.udn
@@ -1,8 +1,15 @@
 Availability:Public
+Crumbs: %ROOT%
 Title: Actor Lifecycle
-Crumbs:%ROOT%, Programming, Programming/UnrealArchitecture, Programming/UnrealArchitecture/Actors
+Type: Overview
 Description:What actually happens when an Actor is loaded or spawned, and eventually dies.
 version: 4.11
+SkillLevel: Intermediate
+Parent: Programming/UnrealArchitecture/Actors
+Order: 0
+Tags: Actors
+Tags: Architecture
+
 
 [TOC (start:2 end:3)]
 
@@ -122,4 +129,5 @@ games that stream content, such as open-world games, as many clusters may merge
 in the group is set for destruction, this can result in large clusters of objects that are kept in memory due to a small number of those objects being in constant use. However, enabling cluster merging and manually adding 
 clusters, in code, for special cases like assets with large numbers of referenced objects that are not shared by other objects, can improve performance by unifying cleanup of the subobjects and decreasing the 
 individual number of entries and dependencies the garbage collector must maintain and check during the game.
+
 [![](AdvancedGC.png)](AdvancedGC.png)
