Documentation for build tools

diff --git a/Engine/Documentation/Source/Programming/UnrealBuildSystem/Configuration/BuildConfigProperties/BuildConfigProperties.INT.udn b/Engine/Documentation/Source/Programming/UnrealBuildSystem/Configuration/BuildConfigProperties/BuildConfigProperties.INT.udn
--- a/Engine/Documentation/Source/Programming/UnrealBuildSystem/Configuration/BuildConfigProperties/BuildConfigProperties.INT.udn
+++ b/Engine/Documentation/Source/Programming/UnrealBuildSystem/Configuration/BuildConfigProperties/BuildConfigProperties.INT.udn
@@ -8,6 +8,18 @@ Version: 4.18
 
 $ bGeneratedSYMFile : Whether to generate dSYM files Lists Architectures that you want to build
 
+$ bEnableAddressSanitizer : Enables address sanitizer (ASan)
+
+$ bEnableThreadSanitizer : Enables thread sanitizer (TSan)
+
+$ bEnableUndefinedBehaviorSanitizer : Enables undefined behavior sanitizer (UBSan)
+
+$ bAddFastPDBToProjects : Whether to add the -FastPDB option to build command lines by default
+
+$ bUsePerFileIntellisense : Whether to generate per-file intellisense data
+
+$ bEditorDependsOnShaderCompileWorker : Whether to include a dependency on ShaderCompileWorker when generating project files for the editor.
+
 $ bDebugBuildsActuallyUseDebugCRT : Enables the debug C++ runtime (CRT) for debug builds. By default we always use the release runtime, since the debug version isn't particularly useful when debugging Unreal Engine projects, and linking against the debug CRT libraries forces our third party library dependencies to also be compiled using the debug CRT (and often perform more slowly). Often it can be inconvenient to require a separate copy of the debug versions of third party static libraries simply so that you can debug your program's code.
 
 $ bUseUnityBuild : Whether to unify C++ code into larger files for faster compilation.
@@ -138,18 +150,26 @@ $ ProcessorCountMultiplier : Processor count multiplier for local execution. Can
 
 $ MaxProcessorCount : Maximum processor count for local execution.
 
-$ bAddFastPDBToProjects : Whether to add the -FastPDB option to build command lines by default
-
-$ bUsePerFileIntellisense : Whether to generate per-file intellisense data
-
-$ bEditorDependsOnShaderCompileWorker : Whether to include a dependency on ShaderCompileWorker when generating project files for the editor.
-
 $ bXGENoWatchdogThread : Whether to use the no_watchdog_thread option to prevent VS2015 toolchain stalls.
 
 $ bShowXGEMonitor : Whether to display the XGE build monitor.
 
 $ bStopXGECompilationAfterErrors : When enabled, XGE will stop compiling targets after a compile error occurs.  Recommended, as it saves computing resources for others.
 
+### ProjectFileGenerator
+
+### VCProjectFileGenerator
+
+$ Version : The version of Visual Studio to generate project files for.
+
+### SourceFileWorkingSet
+
+$ Provider : Sets the provider to use for determining the working set
+
+$ RepositoryPath : Sets the path to use for the repository. Interpreted relative to the UE root directory (ie. folder above the Engine folder) if relative.
+
+$ GitPath : Sets the path to use for the Git executable. Defaults to "git" (assuming it's in the PATH).
+
 ### IOSToolChain
 
 $ IOSSDKVersion : Which version of the iOS SDK to target at build time
@@ -180,12 +200,6 @@ $ ProcessorCountMultiplier : Processor count multiplier for local execution. Can
 
 $ MaxProcessorCount : Maximum processor count for local execution.
 
-### ProjectFileGenerator
-
-### VCProjectFileGenerator
-
-$ Version : The version of Visual Studio to generate project files for.
-
 ### RemoteToolChain
 
 $ RemoteServerName : These two variables will be loaded from XML config file in XmlConfigLoader.Init()
@@ -210,3 +224,9 @@ $ Compiler : Version of the compiler toolchain to use on Windows platform. A val
 
 $ StaticAnalyzer : The static analyzer to use
 
+$ bStrictConformanceMode : Enables strict standard conformance mode (/permissive-) in VS2017+.
+
+### WindowsTargetRules
+
+$ ObjSrcMapFile : Whether we should export a file containing .obj->source file mappings.
+
