Copying //UE4/Dev-Documentation to //UE4/Samples-Main (Source: //UE4/Dev-Documentation @ 3623665)

============================
  MAJOR FEATURES & CHANGES
============================

Change 3178213 by [REDACTED_PERSON]

	#UE4 Docs: UEDOC-3973 - Corrected issues based on Peer Review Feedback.

Change 3318201 by [REDACTED_PERSON]

	Submitted with folder in wrong location. Moved to HowTo section.

Change 3346002 by [REDACTED_PERSON]

	#UE4 Docs: UEDOC-4635 - Mesh Decals page Peer Review Edits and fixed the attachment .zip that had the wrong FBX file.

Change 3376692 by [REDACTED_PERSON]

	#UE4 Docs: UEDOC-4640 - Moving Files up a section to Engine/Rendering/Materials/

Change 3378497 by [REDACTED_PERSON]

	#UE4 Docs: UEDOC-4640 - Mesh Decals Reorganization

Change 3454533 by [REDACTED_PERSON]

	Added global OS toggle

Change 3488613 by [REDACTED_PERSON]

	Updated File against INTSourceChangelist:3482403

Change 3488693 by [REDACTED_PERSON]

	#loc UE4DocJpn File updated against INT#3487380

Change 3488695 by [REDACTED_PERSON]

	#loc UE4DocJpn File updated against INT#3487380

Change 3488700 by [REDACTED_PERSON]

	#loc UE4DocJpn File updated against INT#3481084

Change 3488782 by [REDACTED_PERSON]

	#loc UE4DocJpn Initial translation against INT#3472589

Change 3488801 by [REDACTED_PERSON]

	#loc UE4DocJpn File updated against INT#3478445

Change 3488802 by [REDACTED_PERSON]

	#loc UE4DocJpn File updated against INT#3477566

Change 3488827 by [REDACTED_PERSON]

	#loc UE4DocJpn File updated against INT#3481084

Change 3488837 by [REDACTED_PERSON]

	#loc UE4DocJpn File updated against INT#3481084

Change 3488842 by [REDACTED_PERSON]

	#loc UE4DocJpn File updated against INT#3481084

Change 3488845 by [REDACTED_PERSON]

	#loc UE4DocJpn File updated against INT#3481084

Change 3488859 by [REDACTED_PERSON]

	#loc UE4DocJpn Initial translation against INT#3484673

Change 3489715 by [REDACTED_PERSON]

	Publishing BP Maps and Sets documentation.

	#UE4Doc

Change 3489752 by [REDACTED_PERSON]

	Fixing image file case issue

Change 3489755 by [REDACTED_PERSON]

	Fixed image not showing

Change 3489757 by [REDACTED_PERSON]

	Updated topic images for platforms

Change 3490313 by [REDACTED_PERSON]

	#ue4 docs

	Animation Montage section revamp
	- Split existing Animation Montage page into 4 pages
	- Overview page added
	- Creation page added
	- Editing Montages page added
	- User Guide page added
	- Landing page added with links to the 4 new pages.

Change 3490420 by [REDACTED_PERSON]

	Made a minor grammar fix during Editorial Review process.
	#UE4Doc

Change 3490604 by [REDACTED_PERSON]

	Removing author's question (comment) from source.
	#UE4Doc

Change 3490639 by [REDACTED_PERSON]

	Flipping doc to 'Public'.
	#UE4Doc

Change 3491065 by [REDACTED_PERSON]

	#loc UE4DocJPN moved or deleted

Change 3491123 by [REDACTED_PERSON]

	#loc UE4DocJpn File updated against INT#3490639

Change 3491144 by [REDACTED_PERSON]

	Updated File against INTSoruceChangelist:3482403

Change 3491155 by [REDACTED_PERSON]

	#loc UE4DocJpn File updated against INT#3484673

Change 3491158 by [REDACTED_PERSON]

	#loc UE4DocJpn File updated against INT#3484673

Change 3491161 by [REDACTED_PERSON]

	#loc UE4DocJpn File updated against INT#3489755

Change 3491287 by [REDACTED_PERSON]

	Fixed publish error, essentials list not shown

Change 3491299 by [REDACTED_PERSON]

	#loc UE4DocKOR moved or deleted

Change 3491305 by [REDACTED_PERSON]

	Updated File against INTSourceChangelist:3452211

Change 3491310 by [REDACTED_PERSON]

	Updated File against INTSourceChangelist:3489715

Change 3491315 by [REDACTED_PERSON]

	Updated File against INTSourceChangelist:3489715

Change 3491320 by [REDACTED_PERSON]

	Updated File against INTSourceChangelist:3489715

Change 3491335 by [REDACTED_PERSON]

	Updated File against INTSourceChangelist:3489715

Change 3491342 by [REDACTED_PERSON]

	Updated File against INTSourceChangelist:3489715

Change 3491345 by [REDACTED_PERSON]

	Updated File against INTSourceChangelist:3489715

Change 3491348 by [REDACTED_PERSON]

	Updated File against INTSourceChangelist:3489715

Change 3491350 by [REDACTED_PERSON]

	Updated File against INTSourceChangelist:3489715

Change 3491351 by [REDACTED_PERSON]

	Updated File against INTSourceChangelist:3489715

Change 3491353 by [REDACTED_PERSON]

	Updated File against INTSourceChangelist:3489715

Change 3491358 by [REDACTED_PERSON]

	Updated File against INTSourceChangelist:3482332

Change 3491362 by [REDACTED_PERSON]

	Updated File against INTSourceChangelist:3482332

Change 3491363 by [REDACTED_PERSON]

	#loc UE4DocJpn File updated against INT#3484673

Change 3491409 by [REDACTED_PERSON]

	Updated File against INTSourceChangelist:3482309

Change 3491415 by [REDACTED_PERSON]

	Updated File against INTSourceChangelist:3482294

Change 3491429 by [REDACTED_PERSON]

	Updated File against INTSourceChangelist:3482294

Change 3491682 by [REDACTED_PERSON]

	#ue4 docs

	apply editorial review feedback.

Change 3491704 by [REDACTED_PERSON]

	#ue4 docs

	additional updates based on editorial review feedback

Change 3491718 by [REDACTED_PERSON]

	#ue4 docs

	Added missing image and updated based on editorial review feedback.

Change 3492042 by [REDACTED_PERSON]

	Applying minor edit.

	#UE4Doc

Change 3492049 by [REDACTED_PERSON]

	Adding Feature Highlight video, covering Blueprint Nativization, to documentation.
	#UE4Doc

Change 3494250 by [REDACTED_PERSON]

	#loc UE4DocCHN moved or deleted

Change 3494733 by [REDACTED_PERSON]

	Fixed typo

Change 3494771 by [REDACTED_PERSON]

	#ue4 docs

	Applying peer review feedback.

Change 3495719 by [REDACTED_PERSON]

	Checking in 1 image for ReducePackageSize
	#UE4 Docs

Change 3497877 by [REDACTED_PERSON]

	#loc UE4DocJpn Initial translation against INT#3472589

Change 3498163 by [REDACTED_PERSON]

	#ue4 docs

	applying peer review feedback

Change 3498439 by [REDACTED_PERSON]

	#ue4 docs

	Adding missing image.

Change 3500087 by [REDACTED_PERSON]

	#loc UE4DocKOR

Change 3500128 by [REDACTED_PERSON]

	#loc UE4DocKOR

Change 3500166 by [REDACTED_PERSON]

	Updated File against INTSourceChangelist:3108692

Change 3500224 by [REDACTED_PERSON]

	#loc UE4DocKOR

Change 3500287 by [REDACTED_PERSON]

	#loc UE4DocKOR

Change 3501272 by [REDACTED_PERSON]

	Fixed broken links

Change 3502096 by [REDACTED_PERSON]

	#loc UE4DocJpn File updated against INT#3498163

Change 3502127 by [REDACTED_PERSON]

	#loc UE4DocJpn Initial translation against INT#:3498163

Change 3502137 by [REDACTED_PERSON]

	#loc UE4DocJpn Initial translation against INT#3498163

Change 3502150 by [REDACTED_PERSON]

	#loc UE4DocJpn File updated against INT#3498163

Change 3502159 by [REDACTED_PERSON]

	#loc UE4DocJpn File updated against INT#:3498163

Change 3502168 by [REDACTED_PERSON]

	#loc UE4DocJpn File updated against INT#3494771

Change 3502187 by [REDACTED_PERSON]

	#loc UE4DocJpn File updated against INT#3492042

Change 3502190 by [REDACTED_PERSON]

	#loc UE4DocJpn File updated against INT#3491718

Change 3504279 by [REDACTED_PERSON]

	#loc UE4DocJpn File updated against INT#3501272

Change 3504288 by [REDACTED_PERSON]

	#loc UE4DocJpn File updated against INT#3484673

Change 3504965 by [REDACTED_PERSON]

	Updated topic images

Change 3505571 by [REDACTED_PERSON]

	Updated search boxes for new search engine

Change 3505632 by [REDACTED_PERSON]

	#ue4 docs

	Updating Pose Driver Documentation
	- Split Anim Pose Nodes into two pages (Pose Blender and Pose Driver)
	- Refactored both Pose Blender and Pose Driver pages.
	- Updated AnimPose Page to include new Pose Blender and Pose Driver pages.

Change 3505698 by [REDACTED_PERSON]

	Updated based on SME feedback.

Change 3505757 by [REDACTED_PERSON]

	#ue4 docs

	Applying SME feedback, updated Create Animation description to reflect the inclusion of retargeting.

Change 3506226 by [REDACTED_PERSON]

	#loc UE4DocJPN moved or deleted

Change 3506259 by [REDACTED_PERSON]

	#loc UE4DocJpn File updated against INT#3505632

Change 3506264 by [REDACTED_PERSON]

	#loc UE4DocJpn File updated against INT#3505632

Change 3506281 by [REDACTED_PERSON]

	#loc UE4DocJpn File updated against INT#3473488

Change 3506310 by [REDACTED_PERSON]

	#loc UE4DocJPN moved or deleted

Change 3506350 by [REDACTED_PERSON]

	#loc UE4DocJpn File updated against INT#3469699

Change 3506370 by [REDACTED_PERSON]

	#loc UE4DocJpn File updated against INT#3505757

Change 3506478 by [REDACTED_PERSON]

	#loc UE4DocJpn Initial translation against INT#3505632

Change 3506499 by [REDACTED_PERSON]

	#loc UE4DocJpn Initial translation against INT#3505632

Change 3506584 by [REDACTED_PERSON]

	#loc UE4DocJpn File updated against INT#3454520

Change 3507178 by [REDACTED_PERSON]

	#UE4 Docs: Applying SME Feedback.

Change 3507208 by [REDACTED_PERSON]

	#UE4 Docs: Applying SME feedback.

Change 3509002 by [REDACTED_PERSON]

	#loc UE4DocJpn File updated against INT#3481084

Change 3509009 by [REDACTED_PERSON]

	Updated file against INTSourceChangelist:3350329

Change 3509013 by [REDACTED_PERSON]

	#loc UE4DocKOR

Change 3509014 by [REDACTED_PERSON]

	#loc UE4DocKOR moved or deleted

Change 3509081 by [REDACTED_PERSON]

	#loc UE4DocJpn File updated against INT#3487496

Change 3509107 by [REDACTED_PERSON]

	#loc UE4DocJpn File updated against INT#3471591

Change 3509115 by [REDACTED_PERSON]

	#loc UE4DocJpn File updated against INT#3471591

Change 3509126 by [REDACTED_PERSON]

	#loc UE4DocJpn File updated against INT#3429254

Change 3509139 by [REDACTED_PERSON]

	#loc UE4DocJpn File updated against INT#3429254

Change 3509142 by [REDACTED_PERSON]

	#loc UE4DocJpn File updated against INT#3429254

Change 3509154 by [REDACTED_PERSON]

	#loc UE4DocJpn File updated against INT#3429254

Change 3509156 by [REDACTED_PERSON]

	Updated File against  INTSourceChangelist:3508009

Change 3509177 by [REDACTED_PERSON]

	#loc UE4DocJpn File updated against INT#3481084

Change 3509488 by [REDACTED_PERSON]

	Documentation says "texture streamer cannot automatically fit different memory budgets without the hassle of manual tweaks" - looks like this should say it "can" do that. Code reviewing you guys in case I'm wrong.

Change 3509572 by [REDACTED_PERSON]

	#ue4 docs

	Applying Peer Review feedback.

Change 3510083 by [REDACTED_PERSON]

	#ue4 docs

	Applying SME feedback.

Change 3510882 by [REDACTED_PERSON]

	Updated File against INTSourceChangelist:3507178

Change 3510891 by [REDACTED_PERSON]

	Updated File against INTSourceChangelist:3454520

Change 3510892 by [REDACTED_PERSON]

	Fixed translation of description

Change 3510893 by [REDACTED_PERSON]

	Fixed translation of description

Change 3510894 by [REDACTED_PERSON]

	Fixed again (add period)

Change 3510905 by [REDACTED_PERSON]

	Updated File against INTSourceChangelist:3429238

Change 3510908 by [REDACTED_PERSON]

	Fixed again (add period)

Change 3510910 by [REDACTED_PERSON]

	Fixed translation of title and description

Change 3510916 by [REDACTED_PERSON]

	Fixed translation of title

Change 3510923 by [REDACTED_PERSON]

	Updated File against INTSourceChangelist:3505571

Change 3511033 by [REDACTED_PERSON]

	Updated File against INTSourceChangelist:3360829

Change 3511039 by [REDACTED_PERSON]

	Updated File against INTSourceChangelist:3482177

Change 3511044 by [REDACTED_PERSON]

	Updated File against INTSourceChangelist:3477690

Change 3511084 by [REDACTED_PERSON]

	Updated File against INTSourceChangelist:3473556

Change 3511096 by [REDACTED_PERSON]

	Updated File against INTSourceChangelist:3360829

Change 3511100 by [REDACTED_PERSON]

	Updated File against INTSourceChangelist:3360829

Change 3511115 by [REDACTED_PERSON]

	Updated File against INTSourceChangelist:3350329

Change 3512321 by [REDACTED_PERSON]

	Fixed broken link.
	#UE4Doc

Change 3512328 by [REDACTED_PERSON]

	#UE4 Docs: UEDOC-4958 - Add Custom Stencil to Mobile Post Processing Effects page.

	* Added Custom Stencil section
	* Updated page metadata
	* Adding usage of Custom Stencil to PostProcess Materials page
	* Updating page metadata
	* Updating images for Custom Depth as well.
	* Added Editorial Review suggested edits.

Change 3512866 by [REDACTED_PERSON]

	Updated File against INTSourceChangelist:3512328

Change 3512966 by [REDACTED_PERSON]

	Updated File against INTSourceChangelist:3512321

Change 3513060 by [REDACTED_PERSON]

	Updated File against INTSourceChangelist:3507208

Change 3513074 by [REDACTED_PERSON]

	Updated File against INTSourceChangelist:3494733

Change 3513114 by [REDACTED_PERSON]

	Updated File against INTSourceChangelist:3492049

Change 3513116 by [REDACTED_PERSON]

	#loc UE4DocJpn File updated against INT#3510083

Change 3513127 by [REDACTED_PERSON]

	#loc UE4DocJpn File updated against INT#:3509572

Change 3513140 by [REDACTED_PERSON]

	#loc UE4DocJpn File updated against INT#3509572

Change 3513144 by [REDACTED_PERSON]

	Updated File against INTSourceChangelist:3360829

Change 3513145 by [REDACTED_PERSON]

	#loc UE4DocJpn File updated against INT#3509488

Change 3513582 by [REDACTED_PERSON]

	#UE4 Docs: UEDOC-4952 - Updating Mobile Rendeirng Previewer page.

	* Updating Page metadata
	* Reorganizing page to be easier to read.
	* Adding new Mobile Preview Modes and removing outdated info
	* Added section to show how to enable/disable exprimental/high-end targets

Change 3513619 by [REDACTED_PERSON]

	#UE4 Docs: UEDOC-4858 and UEDOC-4951 - Renaming folders for NVIDIA sections.

Change 3513646 by [REDACTED_PERSON]

	#ue4 docs

	Sequencer Updates
	- New How-to page for WhenFinished property
	- Added section to Overview page for Hierarchical Bias explanation.
	- Adjusted Overview page layout.
	- Removed outdated Enabling Sequencer page since it is now out of experimental and enabled by default.

Change 3513740 by [REDACTED_PERSON]

	#UE4 Docs: UEDOC-4853 and UEDOC-4946 - Renamed pages based on Editorial Review feedback.

Change 3513744 by [REDACTED_PERSON]

	Rollback //UE4/Dev-Documentation/Engine/Documentation/Source/Engine/Rendering/Nvidia/Nvidia.INT.udn to revision 3

Change 3513963 by [REDACTED_PERSON]

	#UE4 Docs: UEDOC-4946 and UEDOC-4853 - NVIDIA pages updates Aftermath and AFR

Change 3514120 by [REDACTED_PERSON]

	#ue4 docs

	Refactored page for Jira UEDOC-3221

Change 3514192 by [REDACTED_PERSON]

	#ue4 docs

	Applying mark ups to images for how-to guide.

Change 3514316 by [REDACTED_PERSON]

	#ue4 docs

	Applying Editorial Review feedback.

Change 3514489 by [REDACTED_PERSON]

	Setting doc to 'Public'.
	#UE4Doc

Change 3514628 by [REDACTED_PERSON]

	Updated metadata

Change 3514777 by [REDACTED_PERSON]

	#UE4 Docs: UEDOC-4957 - Mobile Previewer Editorial Review edits.

Change 3515326 by [REDACTED_PERSON]

	#loc UE4DocJPN moved or deleted

Change 3515366 by [REDACTED_PERSON]

	#loc UE4DocJpn File updated against INT#3471995

Change 3515421 by [REDACTED_PERSON]

	Updated File against INTSourceChangelist:3360829

Change 3515453 by [REDACTED_PERSON]

	#loc UE4DocCHN. do translation

Change 3515467 by [REDACTED_PERSON]

	#loc UE4DocKOR moved or deleted

Change 3515480 by [REDACTED_PERSON]

	#loc UE4DocJpn Initial translation against INT#3513646

Change 3515481 by [REDACTED_PERSON]

	#loc UE4DocJpn File updated against INT#3513646

Change 3515492 by [REDACTED_PERSON]

	#loc UE4DocJpn File updated against INT#3513646

Change 3515493 by [REDACTED_PERSON]

	#loc UE4DocJpn File updated against INT#3514489

Change 3515494 by [REDACTED_PERSON]

	#loc UE4DocJpn File updated against INT#3514316

Change 3515503 by [REDACTED_PERSON]

	#loc UE4DocJpn File updated against INT#3481084

Change 3515514 by [REDACTED_PERSON]

	#loc UE4DocJpn File updated against INT#3481084

Change 3515526 by [REDACTED_PERSON]

	#loc UE4DocJpn File updated against INT#:3481084

Change 3515542 by [REDACTED_PERSON]

	Updated File against INTSourceChangelist:3514777

Change 3515558 by [REDACTED_PERSON]

	Updated File against INTSourceChangelist:3484583

Change 3515571 by [REDACTED_PERSON]

	Updated File against INTSourceChangelist:3484583

Change 3515575 by [REDACTED_PERSON]

	Updated File against INTSourceChangelist:3483116

Change 3517038 by [REDACTED_PERSON]

	Writing first draft of UE4 Brand Glossary.
	#UE4Doc

Change 3517196 by [REDACTED_PERSON]

	#loc UE4DocJPN moved or deleted

Change 3517457 by [REDACTED_PERSON]

	#loc UE4DocCHN moved or deleted

Change 3517485 by [REDACTED_PERSON]

	#loc UE4DocKOR

Change 3517486 by [REDACTED_PERSON]

	#loc UE4DocKOR moved or deleted

Change 3518154 by [REDACTED_PERSON]

	Minor adjustment to IWYU page, linking API docs since they tell users what to include for any given class.

Change 3519874 by [REDACTED_PERSON]

	#loc UE4DocJpn File updated against INT#3514120

Change 3519893 by [REDACTED_PERSON]

	#loc UE4DocJpn Initial translation against INT# INT#3517038

Change 3519902 by [REDACTED_PERSON]

	#loc UE4DocKOR

Change 3519919 by [REDACTED_PERSON]

	#loc UE4DocJpn Initial translation against INT#3517038

Change 3519945 by [REDACTED_PERSON]

	#loc UE4DocJpn Initial translation against INT#:3517038

Change 3520085 by [REDACTED_PERSON]

	#UE4 Docs: UEDOC-4933 - Volumetric Fog Editorial Review edits.

Change 3520317 by [REDACTED_PERSON]

	#loc UE4DocJpn Initial translation against INT#3517038

Change 3520327 by [REDACTED_PERSON]

	#loc UE4DocJpn Initial translation against INT#3517038

Change 3520347 by [REDACTED_PERSON]

	#loc UE4DocJpn Initial translation against INT#3517038

Change 3520382 by [REDACTED_PERSON]

	Updated File against INTSourceChangelist:3482475

Change 3520390 by [REDACTED_PERSON]

	#loc UE4DocJpn Initial translation against INT#3517038

Change 3520392 by [REDACTED_PERSON]

	Updated File against INTSourceChangelist:3517038

Change 3520409 by [REDACTED_PERSON]

	Updated File against INTSourceChangelist:3514628

Change 3520418 by [REDACTED_PERSON]

	#loc UE4DocJpn Initial translation against INT#3517038

Change 3520420 by [REDACTED_PERSON]

	#loc UE4DocJpn Initial translation against INT#:3517038

Change 3520423 by [REDACTED_PERSON]

	#loc UE4DocJpn Initial translation against INT#3517038

Change 3520428 by [REDACTED_PERSON]

	Updatd File against INTSourceChangelist:3482332

Change 3520442 by [REDACTED_PERSON]

	#loc UE4DocJpn Initial translation against INT#3517038

Change 3520446 by [REDACTED_PERSON]

	Updated File against INTSourceChangelist:3108692

Change 3520449 by [REDACTED_PERSON]

	#loc UE4DocJpn Initial translation against INT#3517038

Change 3520495 by [REDACTED_PERSON]

	#loc UE4DocJpn Initial translation against INT#3517038

Change 3520509 by [REDACTED_PERSON]

	#loc UE4DocJpn Initial translation against INT#3517038

Change 3520580 by [REDACTED_PERSON]

	#loc UE4DocKOR

Change 3520887 by [REDACTED_PERSON]

	Updated File against INTSourceChecklist:3520085

Change 3520893 by [REDACTED_PERSON]

	#loc UE4DocJpn Initial translation against INT#3517038

Change 3520896 by [REDACTED_PERSON]

	#loc UE4DocJpn Initial translation against INT#3517038

Change 3520903 by [REDACTED_PERSON]

	#loc UE4DocJpn Initial translation against INT#3517038

Change 3520915 by [REDACTED_PERSON]

	#loc UE4DocJpn Initial translation against INT#3517038

Change 3520922 by [REDACTED_PERSON]

	#loc UE4DocJpn Initial translation against INT#3517038

Change 3520966 by [REDACTED_PERSON]

	#loc UE4DocJpn File updated against INT#3467293

Change 3520990 by [REDACTED_PERSON]

	#loc UE4DocJpn File updated against INT#3487496

Change 3520999 by [REDACTED_PERSON]

	#loc UE4DocJpn File updated against INT#3454520

Change 3521021 by [REDACTED_PERSON]

	Updated File against INTSourceChangelist:3518154

Change 3521235 by [REDACTED_PERSON]

	#UE4 Docs: Making a few changes based on feedback form Rob.

Change 3521730 by [REDACTED_PERSON]

	Adding link to API Reference.

	#UE4Doc

Change 3521824 by [REDACTED_PERSON]

	#UE4 Docs: Was using the wrong format, changed to the right one.

Change 3523701 by [REDACTED_PERSON]

	#loc UE4DocJpn Initial translation against INT#3469699

Change 3523722 by [REDACTED_PERSON]

	#loc UE4DocJpn File updated against INT#3452211

Change 3523743 by [REDACTED_PERSON]

	Updated File against INTSourceChangelist:3487496

Change 3523784 by [REDACTED_PERSON]

	#loc UE4DocJpn File updated against INT#3426963

Change 3523785 by [REDACTED_PERSON]

	Updated File against INTSourceChangelist:3521235

Change 3523836 by [REDACTED_PERSON]

	#loc UE4DocJpn File updated against INT#3426963

Change 3523845 by [REDACTED_PERSON]

	Updated File against INTSourceChangelist:3521824

Change 3523879 by [REDACTED_PERSON]

	Updated File against INTSourceChangelist:3521235

Change 3523894 by [REDACTED_PERSON]

	#loc UE4DocJpn File updated against INT#3426963

Change 3523910 by [REDACTED_PERSON]

	#loc UE4DocJpn File updated against INT#3521730

Change 3523916 by [REDACTED_PERSON]

	#loc UE4DocJpn File updated against INT#3481084

Change 3523952 by [REDACTED_PERSON]

	#loc UE4DocJpn File updated against INT#3367474

Change 3523964 by [REDACTED_PERSON]

	#loc UE4DocJpn File updated against INT#3473462

Change 3523969 by [REDACTED_PERSON]

	#loc UE4DocJpn File updated against INT#3481084

Change 3523984 by [REDACTED_PERSON]

	Updated File against  INTSourceChangelist:3514628

Change 3524197 by [REDACTED_PERSON]

	#ue4 docs

	Applying Peer Feedback

Change 3524216 by [REDACTED_PERSON]

	#ue4 docs

	Applying Peer Review feedback.

Change 3524231 by [REDACTED_PERSON]

	Removing CHM reference from API page.

Change 3524929 by [REDACTED_PERSON]

	#ue4 docs

	Applying Peer Review feedback.

Change 3525849 by [REDACTED_PERSON]

	Updated File against INTSourceChangelist:3512328

Change 3525881 by [REDACTED_PERSON]

	Updated File against INTSourceChangelist:3487496

Change 3525922 by [REDACTED_PERSON]

	Updated File against INTSourceChangelist:3487496

Change 3525964 by [REDACTED_PERSON]

	#loc UE4DocCHN. update against latest INT version

Change 3528506 by [REDACTED_PERSON]

	Updated File against INTSourceChangelist:3489715

Change 3528507 by [REDACTED_PERSON]

	Updated File against INTSourceChangelist:3489715

Change 3528508 by [REDACTED_PERSON]

	Updated File against INTSourceChangelist:3489715

Change 3528510 by [REDACTED_PERSON]

	Updated File against INTSourceChangelist:3489715

Change 3528511 by [REDACTED_PERSON]

	Updated File against INTSourceChangelist:3489715

Change 3528516 by [REDACTED_PERSON]

	Updated File against INTSourceChangelist:3489715

Change 3528517 by [REDACTED_PERSON]

	Updated File against INTSourceChangelist:3489715

Change 3528518 by [REDACTED_PERSON]

	Updated File against INTSourceChangelist:3489715

Change 3528520 by [REDACTED_PERSON]

	#loc UE4DocKOR

Change 3528522 by [REDACTED_PERSON]

	Updated File against INTSourceChangelist:3489715

Change 3528523 by [REDACTED_PERSON]

	Updated File against INTSourceChangelist:3489715

Change 3528528 by [REDACTED_PERSON]

	Udpated File against INTSourceChangelist:3489715

Change 3528529 by [REDACTED_PERSON]

	Updated File against INTSourceChangelist:3489715

Change 3528530 by [REDACTED_PERSON]

	Updated File against INTSourceChangelist:3489715

Change 3528532 by [REDACTED_PERSON]

	Updated File agaisnt  INTSourceChangelist:3489715

Change 3528534 by [REDACTED_PERSON]

	Updated File against INTSourceChangelist:3489715

Change 3528535 by [REDACTED_PERSON]

	Updated File against INTSourceChangelist:3489715

Change 3528541 by [REDACTED_PERSON]

	Updated File against INTSourceChangelist:3489715

Change 3528543 by [REDACTED_PERSON]

	Updated File against INTSourceChangelist:3489715

Change 3528546 by [REDACTED_PERSON]

	Updated File against INTSourceChangelist:3489715

Change 3528550 by [REDACTED_PERSON]

	Updated File against INTSourceChangelist:3489715

Change 3528554 by [REDACTED_PERSON]

	Updated File against INTSourceChangelist:3489715

Change 3528555 by [REDACTED_PERSON]

	Updated File against INTSourceChangelist:3489715

Change 3528558 by [REDACTED_PERSON]

	Updated File against INTSourceChangelist:3482332

Change 3528974 by [REDACTED_PERSON]

	#UE4 Docs: Adding in a few updates and fixing some spelling and grammer errors.

Change 3528976 by [REDACTED_PERSON]

	Adding this missing file.

Change 3530743 by [REDACTED_PERSON]

	Adding some fields to Properties page. Including some tags to make bitmask/bitfield enums easier to find.

Change 3531730 by [REDACTED_PERSON]

	Implementing editorial review comments.

Change 3531989 by [REDACTED_PERSON]

	#UE4 Docs: Applying Editorial feedback.

Change 3531998 by [REDACTED_PERSON]

	Applied bold text style to all key terms belonging to tables.
	#UE4Doc

Change 3532405 by [REDACTED_PERSON]

	#UE4Doc
	Updating Engine Version

Change 3532595 by [REDACTED_PERSON]

	Updated File against INTSourceChangelist:3528974

Change 3532731 by [REDACTED_PERSON]

	#loc UE4DocJpn File updated against INT#3524216

Change 3532735 by [REDACTED_PERSON]

	#loc UE4DocJpn File updated against INT#3524197

Change 3532745 by [REDACTED_PERSON]

	#loc UE4DocJpn File updated against INT#3524929

Change 3532750 by [REDACTED_PERSON]

	#loc UE4DocJpn File updated against INT#3524929

Change 3533383 by [REDACTED_PERSON]

	Updated in response to editorial review.

Change 3533392 by [REDACTED_PERSON]

	Set page to publish

Change 3534796 by [REDACTED_PERSON]

	Moving to new scope per Editorial Review.

	#UE4Doc

Change 3534801 by [REDACTED_PERSON]

	Fixing Parent metadata.
	#UE4Doc

Change 3534809 by [REDACTED_PERSON]

	Fixed Parent Metadata in localized docs.

	#UE4Doc

Change 3534891 by [REDACTED_PERSON]

	#loc UE4DocJpn File updated against INT#3533392

Change 3534895 by [REDACTED_PERSON]

	#loc UE4DocJpn File updated against INT#3533392

Change 3534899 by [REDACTED_PERSON]

	#loc UE4DocJpn File updated against INT#3533392

Change 3534903 by [REDACTED_PERSON]

	#loc UE4DocJpn File updated against INT#3533392

Change 3534906 by [REDACTED_PERSON]

	#loc UE4DocJpn File updated against INT#3532405

Change 3537135 by [REDACTED_PERSON]

	correction for 4.16 - Window to Edit, Head Mount Display to Virtual Reality

Change 3537137 by [REDACTED_PERSON]

	#loc UE4DocKOR

Change 3537352 by [REDACTED_PERSON]

	#loc UE4DocJpn File updated against INT#3532405

Change 3537359 by [REDACTED_PERSON]

	#loc UE4DocJpn File updated against INT#3532405

Change 3537362 by [REDACTED_PERSON]

	#loc UE4DocJpn File updated against INT#3532405

Change 3537534 by [REDACTED_PERSON]

	#loc UE4DocKOR

Change 3539678 by [REDACTED_PERSON]

	#loc UE4DocCHN. do translation.

Change 3540046 by [REDACTED_PERSON]

	Setting documentation to 'Public'.

	#UE4Doc

Change 3541529 by [REDACTED_PERSON]

	#loc UE4DocJpn File updated against INT#3531730

Change 3541692 by [REDACTED_PERSON]

	#loc UE4DocJpn File updated against INT#3533383

Change 3541720 by [REDACTED_PERSON]

	#loc UE4DocJpn File updated against INT#3532405

Change 3543533 by [REDACTED_PERSON]

	#loc UE4DocJpn File updated against INT#3532405

Change 3543535 by [REDACTED_PERSON]

	#loc UE4DocJpn File updated against INT#3532405

Change 3543537 by [REDACTED_PERSON]

	#loc UE4DocJpn File updated against INT#3532405

Change 3543716 by [REDACTED_PERSON]

	#loc UE4DocJpn File updated against INT#3531998

Change 3544224 by [REDACTED_PERSON]

	re-naming this file to make it easier to know what it is.

Change 3544387 by [REDACTED_PERSON]

	#UE4 Docs: First Check in for this doc

Change 3544397 by [REDACTED_PERSON]

	Submitting first draft.

	#UE4Doc

Change 3544575 by [REDACTED_PERSON]

	#UE4 Docs: Applying SME Feedback

Change 3544589 by [REDACTED_PERSON]

	Updated metadata

Change 3545553 by [REDACTED_PERSON]

	#loc UE4DocJPN moved or deleted

Change 3545593 by [REDACTED_PERSON]

	#loc UE4DocJpn File updated against INT#3530743

Change 3545845 by [REDACTED_PERSON]

	#loc UE4DocJpn File updated against INT#3532405

Change 3545848 by [REDACTED_PERSON]

	#loc UE4DocJpn File updated against INT#3532405

Change 3547788 by [REDACTED_PERSON]

	#loc UE4DocJpn File updated against INT#3429264

Change 3547833 by [REDACTED_PERSON]

	#loc UE4DocJpn File updated against INT#3454520

Change 3547843 by [REDACTED_PERSON]

	#loc UE4DocJpn File updated against INT#3532405

Change 3547846 by [REDACTED_PERSON]

	INTSourceChangelist:3532405

Change 3547847 by [REDACTED_PERSON]

	#loc UE4DocJpn File updated against INT#3532405

Change 3547858 by [REDACTED_PERSON]

	#loc UE4DocJpn File updated against INT#3532405

Change 3548232 by [REDACTED_PERSON]

	Initial Asset Management document and related new/updated materials.

Change 3548353 by [REDACTED_PERSON]

	#UE4 Docs: UEDOC-5073 - Slate Clipping System documentation

	* Added Slate Clipping System doc for programming
	* Added Clipping System Feature doc for UMG Designer.

Change 3548675 by [REDACTED_PERSON]

	Added tag as part of Editorial Review.
	#UE4Doc

Change 3550094 by [REDACTED_PERSON]

	Updated File against INTSourceChangelist:3544589

Change 3550236 by [REDACTED_PERSON]

	#loc UE4DocJpn File updated against INT#3485038

Change 3550817 by [REDACTED_PERSON]

	#UE4 Docs: UEDOC-5086 - Adding Opacity/Opacity Mask for Material Flattening to the following pages for HLOD and Actor Merging.

	Actor Merging Reference Page updates:
	* Updating Page UI
	* Removing parts referencing this feature as experimental as it no longer is.
	* Broke up table into different sections with their own descriptions and added missing table information.
	* Updated and added images for the 'Actor Merging Workflow' section and moved to the bottom of the reference page.

	HLOD Reference page udates:
	* Updated images to reflect latest UI additions.
	* Added new changes to the section tables.

Change 3552394 by [REDACTED_PERSON]

	Updated File against INTSourceChangelist:3544589

Change 3552524 by [REDACTED_PERSON]

	Updated File against INTSourceChangelist:3544589

Change 3552532 by [REDACTED_PERSON]

	Updated File against INTSourceChangelist:3151855

Change 3552667 by [REDACTED_PERSON]

	#loc UE4DocKOR

Change 3554961 by [REDACTED_PERSON]

	Updated File against INTSourceChangelist:3544387

Change 3554968 by [REDACTED_PERSON]

	Updated File against INTSourceChangelist:3544387

Change 3554970 by [REDACTED_PERSON]

	#loc UE4DocJpn Initial translation against INT#3548353

Change 3555015 by [REDACTED_PERSON]

	#loc UE4DocKOR

Change 3555683 by [REDACTED_PERSON]

	Chaning the titles so they are not the same.

Change 3555716 by [REDACTED_PERSON]

	Hand to change the" to " as the current ones would not publish.

Change 3556577 by [REDACTED_PERSON]

	#UE4 Docs: Chris B pointed out a small issue in the Google VR QS that could result in people not getting their project to work.

Change 3556634 by [REDACTED_PERSON]

	#UE4 Docs: fixing NVIDIA captialization

Change 3556651 by [REDACTED_PERSON]

	#UE4 Docs - fixed NVIDIA captilization for NVIDIA Aftermath page.

Change 3560592 by [REDACTED_PERSON]

	#loc UE4DocJpn Initial translation against INT#3548232

Change 3560640 by [REDACTED_PERSON]

	Updated File against INTSourceChangelist:3544397

Change 3560709 by [REDACTED_PERSON]

	#loc UE4DocJpn Initial translation against INT#3548232

Change 3561534 by [REDACTED_PERSON]

	#UE4 Docs: UEDOC-5090 - Editorial Review edits for Opacity and Opacity Masks for Material Flattening.

Change 3563711 by [REDACTED_PERSON]

	#loc UE4DocJpn File updated against INT#3561534

Change 3564038 by [REDACTED_PERSON]

	Implemented peer review suggestions.

Change 3565221 by [REDACTED_PERSON]

	Adding relevant page links and applying basis S&S best practices to these documents.

	#UE4Doc

Change 3565273 by [REDACTED_PERSON]

	#UE Docs: Applying some grammar and typo fixes to StereoLayers page

Change 3565868 by [REDACTED_PERSON]

	#loc UE4DocKOR

Change 3566803 by [REDACTED_PERSON]

	SME review fixes.

Change 3567151 by [REDACTED_PERSON]

	Peer review suggestions.

Change 3567323 by [REDACTED_PERSON]

	#UE4 Docs: aPPLYING PEER FEEDBACK.

Change 3568376 by [REDACTED_PERSON]

	#UE4 Docs: Applying Editorial feedback.

Change 3568433 by [REDACTED_PERSON]

	#UE4 Docs:Applying SME feedback.

Change 3568601 by [REDACTED_PERSON]

	#UE4 Docs: Changing the liminations section to Feature Limitations

Change 3568702 by [REDACTED_PERSON]

	#UE Docs: Grammatical fixes for ShaderInPlugins docs

Change 3568890 by [REDACTED_PERSON]

	Initial check in of 4.17 Release Notes.

Change 3569033 by [REDACTED_PERSON]

	#UE4 Docs: UEDOC-4970 - Updating Mobile Performance with new optimizations in Project Settings.

	* Added necessary images and recaptured others for Mobile Project Settings section
	* Updated page metadata
	* Restructured and rewrote parts of the page to update it to latest S&S guidelines.

Change 3569052 by [REDACTED_PERSON]

	#UE Docs: UEDOC-5076 - Typo fixes to ClippingSystem Docs

Change 3569253 by [REDACTED_PERSON]

	Formatting updates

Change 3570173 by [REDACTED_PERSON]

	Updated public docs
	 -4.17  release notes update

Change 3570322 by [REDACTED_PERSON]

	#loc UE4DocKOR moved or deleted

Change 3570398 by [REDACTED_PERSON]

	#loc UE4DocKOR

Change 3570570 by [REDACTED_PERSON]

	Fixed a few typos.

Change 3571140 by [REDACTED_PERSON]

	#UE4 Docs: UEDOC-4975 - Mobile Performance Editorial Review edits.

Change 3571159 by [REDACTED_PERSON]

	#UE4Doc

	Updating related pages.

Change 3573283 by [REDACTED_PERSON]

	#UE4 Docs: UEDOC-5397 - Clothing Tool updates for 4.17

	* Updated the Overview Page
	* Added a Landing Page for Clothing
	* Added a Property and Tool Reference Page
	* Added link to reference page on Physics Landing page

Change 3574383 by [REDACTED_PERSON]

	#loc UE4DocJpn File updated against INT#3561534

Change 3574393 by [REDACTED_PERSON]

	#loc UE4DocJpn File updated against INT#3569052

Change 3574395 by [REDACTED_PERSON]

	#loc UE4DocJpn File updated against INT#3573283

Change 3574458 by [REDACTED_PERSON]

	Updated File against INTSourceChangelist:3569151

Change 3574459 by [REDACTED_PERSON]

	Updated File against INTSourceChangelist:3569151

Change 3574461 by [REDACTED_PERSON]

	Updated File against INTSourceChangelist:3569151

Change 3574584 by [REDACTED_PERSON]

	#loc UE4DocKOR

Change 3574716 by [REDACTED_PERSON]

	#UE4 Docs

	Animation Optimization documentation
	- omitted URO documentation as 4.18 is slated for improvements.
	- already reviewed by [REDACTED_PERSON]

Change 3574744 by [REDACTED_PERSON]

	#ue4 docs

	Applying Editoral Review feedback.

Change 3574749 by [REDACTED_PERSON]

	#ue4 docs

	Applying Editoral Review feedback.

Change 3575019 by [REDACTED_PERSON]

	#UE4 Docs

	Applying Editorial Review feedback

Change 3575037 by [REDACTED_PERSON]

	#ue4 docs

	Applying Editorial Review feedback

Change 3575635 by [REDACTED_PERSON]

	#UE4 Docs: UEDOC-5399 - Clothing Tool Peer Review edits

Change 3575783 by [REDACTED_PERSON]

	#UE4DOC #jira UEDOC-5400 - Fixing Typos in ClothOverview

Change 3575912 by [REDACTED_PERSON]

	#loc UE4DocJPN moved or deleted

Change 3575939 by [REDACTED_PERSON]

	#loc UE4DocKOR moved or deleted

Change 3576254 by [REDACTED_PERSON]

	#loc UE4DocKOR

Change 3576470 by [REDACTED_PERSON]

	#UE4 Docs: UEDOC-5401 - Clothing Tool 4.17 updates SME review edits

Change 3576482 by [REDACTED_PERSON]

	updating DB file

Change 3576484 by [REDACTED_PERSON]

	Minor updates

Change 3576492 by [REDACTED_PERSON]

	Initial draft of Composure Quick Start

Change 3576958 by [REDACTED_PERSON]

	Updating Unreal Engine logo on Introduction to C++ Programming in UE4 doc.
	#UEDoc

Change 3577867 by [REDACTED_PERSON]

	#loc UE4DocJPN moved or deleted

Change 3578075 by [REDACTED_PERSON]

	#loc UE4DocKOR moved or deleted

Change 3578289 by [REDACTED_PERSON]

	#loc UE4DocJpn File updated against INT#3575037

Change 3578362 by [REDACTED_PERSON]

	#loc UE4DocKOR

Change 3578363 by [REDACTED_PERSON]

	#loc UE4DocCHN. do translation.

Change 3579093 by [REDACTED_PERSON]

	#UE4 Docs:Applying SME feedback.

Change 3579346 by [REDACTED_PERSON]

	#UE4 Docs: Fixing NVIDIA capitalization issues on the NVIDIA Ansel documentation pages

Change 3579707 by [REDACTED_PERSON]

	Updated File against INTSourceChangelist:3567323

Change 3579843 by [REDACTED_PERSON]

	Updated File against INTSourceChangelist:3565221

Change 3579859 by [REDACTED_PERSON]

	Updated File against INTSourceChangelist:3565221

Change 3579873 by [REDACTED_PERSON]

	Updated File against INTSourceChangelist:3565221

Change 3579918 by [REDACTED_PERSON]

	#loc UE4DocJpn Initial translation against INT#:3574716

Change 3579929 by [REDACTED_PERSON]

	#loc UE4DocJpn File updated against INT#3579346

Change 3579934 by [REDACTED_PERSON]

	#loc UE4DocJpn File updated against INT#3579346

Change 3579937 by [REDACTED_PERSON]

	#loc UE4DocJpn File updated against INT#3579346

Change 3579942 by [REDACTED_PERSON]

	#loc UE4DocJpn File updated against INT#:3579346

Change 3579946 by [REDACTED_PERSON]

	#loc UE4DocJpn File updated against INT#:3579346

Change 3579953 by [REDACTED_PERSON]

	Updated File against INTSourceChangelist:3556651

Change 3580568 by [REDACTED_PERSON]

	#UE4 Docs: First Composure doc check in.

Change 3580600 by [REDACTED_PERSON]

	#UE4 Docs:Apply editoral feedback.

Change 3580774 by [REDACTED_PERSON]

	#UE4 Docs:Applying more peer feedback.

Change 3580824 by [REDACTED_PERSON]

	#UE4 Docs:Applying Editorial feedback.

Change 3580887 by [REDACTED_PERSON]

	#UE4 Docs:Applying Editoral feedback.

Change 3581141 by [REDACTED_PERSON]

	Adding mac version of ProjectHierarchy screenshot to Glossary_P
	#UEDoc

Change 3581822 by [REDACTED_PERSON]

	#loc UE4DocCHN moved or deleted

Change 3581871 by [REDACTED_PERSON]

	#loc UE4DocKOR

Change 3583752 by [REDACTED_PERSON]

	#loc UE4DocJPN moved or deleted

Change 3583774 by [REDACTED_PERSON]

	#loc UE4DocKOR moved or deleted

Change 3583828 by [REDACTED_PERSON]

	Updated File against INTSourceChangelist:3576484

Change 3583988 by [REDACTED_PERSON]

	#loc UE4DocCHN. do translation.

Change 3584026 by [REDACTED_PERSON]

	#loc UE4DocKOR

Change 3584118 by [REDACTED_PERSON]

	#ue4 docs

	Applying peer review feedback.

Change 3587160 by [REDACTED_PERSON]

	Updating documentation at Team Online's request.
	#UE4Doc

Change 3587264 by [REDACTED_PERSON]

	#loc UE4DocJpn Initial translation against INT#3580568

Change 3587278 by [REDACTED_PERSON]

	#loc UE4DocJpn Initial translation against INT#:3580568

Change 3587304 by [REDACTED_PERSON]

	#loc UE4DocJpn Initial translation against INT#3580568

Change 3587361 by [REDACTED_PERSON]

	#loc UE4DocJpn Initial translation against INT#3580568

Change 3587396 by [REDACTED_PERSON]

	#loc UE4DocJpn Initial translation against INT#3580568

Change 3587410 by [REDACTED_PERSON]

	#loc UE4DocJpn Initial translation against INT#3580568

Change 3587414 by [REDACTED_PERSON]

	#loc UE4DocJpn Initial translation against INT#3580568

Change 3587425 by [REDACTED_PERSON]

	#loc UE4DocJpn Initial translation against INT#3580568

Change 3587477 by [REDACTED_PERSON]

	#loc UE4DocJpn Initial translation against INT#3580568

Change 3587481 by [REDACTED_PERSON]

	#loc UE4DocJpn Initial translation against INT#3580568

Change 3587505 by [REDACTED_PERSON]

	#loc UE4DocJpn File updated against INT#3532405

Change 3587509 by [REDACTED_PERSON]

	#loc UE4DocJpn File updated against INT#3532405

Change 3587510 by [REDACTED_PERSON]

	#loc UE4DocJpn File updated against INT#3532405

Change 3587512 by [REDACTED_PERSON]

	#loc UE4DocJpn File updated against INT#3532405

Change 3587513 by [REDACTED_PERSON]

	#loc UE4DocJpn File updated against INT#3532405

Change 3587524 by [REDACTED_PERSON]

	#loc UE4DocJpn File updated against INT#3584118

Change 3587527 by [REDACTED_PERSON]

	#loc UE4DocJpn File updated against INT#3575019

Change 3587531 by [REDACTED_PERSON]

	#loc UE4DocJpn File updated against INT#3574749

Change 3587677 by [REDACTED_PERSON]

	#UE4 Docs: Applying Peer Feedback.

Change 3587824 by [REDACTED_PERSON]

	#UE4 Docs: Applying SME feedback.

Change 3588288 by [REDACTED_PERSON]

	Implemented editorial review feedback.

Change 3589912 by [REDACTED_PERSON]

	#loc UE4DocJpn File updated against INT#3587677

Change 3589916 by [REDACTED_PERSON]

	#loc UE4DocJpn File updated against INT#3587677

Change 3589918 by [REDACTED_PERSON]

	#loc UE4DocJpn File updated against INT#:3587677

Change 3589921 by [REDACTED_PERSON]

	#loc UE4DocJpn File updated against INT#3587677

Change 3589924 by [REDACTED_PERSON]

	#loc UE4DocJpn File updated against INT#:3587677

Change 3589927 by [REDACTED_PERSON]

	#loc UE4DocJpn File updated against INT#3587677

Change 3589979 by [REDACTED_PERSON]

	#loc UE4DocKOR

Change 3590068 by [REDACTED_PERSON]

	#UE4 Docs:Applying Editorial feedback.

Change 3590438 by [REDACTED_PERSON]

	#UE4 Docs: UEDOC-5397 - Clothing Tool updates for 4.17 Editorial review edits.

Change 3590614 by [REDACTED_PERSON]

	#UE4 Docs: Fixing code comments.

Change 3591289 by [REDACTED_PERSON]

	Fixed a minor error in Replay documentation.

Change 3592139 by [REDACTED_PERSON]

	#loc UE4DocJpn File updated against INT#3591289

Change 3592181 by [REDACTED_PERSON]

	#loc UE4DocJpn File updated against INT#3588288

Change 3592207 by [REDACTED_PERSON]

	#loc UE4DocJpn File updated against INT#3588288

Change 3592290 by [REDACTED_PERSON]

	#loc UE4DocJpn File updated against INT#3588288

Change 3592538 by [REDACTED_PERSON]

	#loc UE4DocJpn File updated against INT#3590438

Change 3593326 by [REDACTED_PERSON]

	Set page to publish

Change 3593328 by [REDACTED_PERSON]

	Added link to Precomputed Lighting Scenarios

Change 3593759 by [REDACTED_PERSON]

	Adding Linux section to the OS specific file descriptions.

	#UE4Doc

Change 3594921 by [REDACTED_PERSON]

	#UEDOC Fixing UEDOC-5526, which illustrates how little I knew about the engine back in April of 2013.

	Removing section on permanent socket attachment and drag/drop, which isn't reliable and doesn't work on sockets.

Change 3595864 by [REDACTED_PERSON]

	#loc UE4DocJpn File updated against INT#3594921

Change 3595974 by [REDACTED_PERSON]

	Updated File against INTSourceChangelist:3593326

Change 3595984 by [REDACTED_PERSON]

	Fixed translation

Change 3595996 by [REDACTED_PERSON]

	Fixed tag error

Change 3595997 by [REDACTED_PERSON]

	Updated File against INTSourceChangelist:3593328

Change 3596043 by [REDACTED_PERSON]

	Updated File against INTSourceChangelist:3590614

Change 3596065 by [REDACTED_PERSON]

	Updated File against INTSourceChangelist:3587824

Change 3596097 by [REDACTED_PERSON]

	Updated File against INTSourceChangelist:3584409

Change 3596109 by [REDACTED_PERSON]

	Updated File against INTSourceChangelist:3580887

Change 3596116 by [REDACTED_PERSON]

	Updated File against INTSourceChangelist:3580824

Change 3596189 by [REDACTED_PERSON]

	#loc UE4DocJpn File updated against INT#3590438

Change 3596201 by [REDACTED_PERSON]

	Fixed translation

Change 3596774 by [REDACTED_PERSON]

	#loc UE4DocCHN. minor improvement for auto-key in Chinese

Change 3596779 by [REDACTED_PERSON]

	Updated a section title to express its intent a bit more clearly.

	#UE4Doc

Change 3599439 by [REDACTED_PERSON]

	#loc UE4DocJpn File updated against INT#3454520

Change 3599465 by [REDACTED_PERSON]

	#loc UE4DocJpn File updated against INT#3454520

Change 3599541 by [REDACTED_PERSON]

	Updated File against INTSourceChangelist:3580774

Change 3599723 by [REDACTED_PERSON]

	#loc UE4DocKOR

Change 3600018 by [REDACTED_PERSON]

	Documentation update for Timers.

Change 3600078 by [REDACTED_PERSON]

	Fixing a small typo in step 7 with the name of the FBX file.

Change 3600525 by [REDACTED_PERSON]

	Added a related page and some helpful links.

Change 3601417 by [REDACTED_PERSON]

	Updated file against INTSourceChangelist:3599948

Change 3601469 by [REDACTED_PERSON]

	Updated File against INTSourceChangelist:3599948

Change 3601483 by [REDACTED_PERSON]

	Updated file against INTSourceChangelist:3599948

Change 3601486 by [REDACTED_PERSON]

	Updated file against INTSourceChangelist:3599948

Change 3601502 by [REDACTED_PERSON]

	Updated file against INTSourceChangelist:3599948

Change 3601518 by [REDACTED_PERSON]

	Updated file against INTSourceChangelist:3599948

Change 3601531 by [REDACTED_PERSON]

	Updated file against INTSourceChangelist:3599948

Change 3601540 by [REDACTED_PERSON]

	Updated file against  INTSourceChangelist:3599948

Change 3601560 by [REDACTED_PERSON]

	Updated file against INTSourceChangelist:3599948

Change 3601599 by [REDACTED_PERSON]

	Updated file against INTSourceChangelist:3571159

Change 3602897 by [REDACTED_PERSON]

	Checking in my changes.

Change 3602938 by [REDACTED_PERSON]

	#ue4 docs

	applying editorial review feedback.

Change 3603049 by [REDACTED_PERSON]

	#ue4 docs

	applying sme review feedback.

Change 3603366 by [REDACTED_PERSON]

	#UE4 Docs: UEDOC-5013 - Color Grading improvements

	* Updating the Post Process Effects for Mobile page which includes numerous image updates and reworking of the properties table to reflect latest changes.
	* Updating the Post Process: Color Grading page to reflect latest changes, new images, and reworking the LUT Creation section into something better.

Change 3603669 by [REDACTED_PERSON]

	Replacing Unicode character with ANSI character.

	#UE4Doc

Change 3604326 by [REDACTED_PERSON]

	Adding r.Mobile.DisableVertexFog=False to multiple projects to enable Vertex Fog on mobile to resolve warnings.
	Saving default spawn location in SunTemple.

Change 3604419 by [REDACTED_PERSON]

	#UE4 Docs: Clarifying phrasing of a sentence based on feedback from our EGK team.

Change 3606232 by [REDACTED_PERSON]

	#loc UE4DocJpn File updated against INT#:3602938

Change 3606246 by [REDACTED_PERSON]

	#loc UE4DocJpn File updated against INT#3603049

Change 3606254 by [REDACTED_PERSON]

	#loc UE4DocJpn File updated against INT#3603049

Change 3606274 by [REDACTED_PERSON]

	#loc UE4DocJpn File updated against INT#3599948

Change 3606291 by [REDACTED_PERSON]

	#loc UE4DocJpn File updated against INT#3599948

Change 3606327 by [REDACTED_PERSON]

	#loc UE4DocJpn File updated against INT#3599948

Change 3606328 by [REDACTED_PERSON]

	#loc UE4DocJpn File updated against INT#3599948

Change 3606363 by [REDACTED_PERSON]

	#loc UE4DocJpn File updated against INT#3599948

Change 3606380 by [REDACTED_PERSON]

	#loc UE4DocJpn File updated against INT#3599948

Change 3606384 by [REDACTED_PERSON]

	#loc UE4DocJpn File updated against INT#3599948

Change 3606390 by [REDACTED_PERSON]

	#loc UE4DocJpn File updated against INT#:3599948

Change 3606391 by [REDACTED_PERSON]

	#loc UE4DocKor

Change 3606398 by [REDACTED_PERSON]

	#loc UE4DocJpn File updated against INT#3599948

Change 3606415 by [REDACTED_PERSON]

	#loc UE4DocJpn File updated against INT#3599948

Change 3606445 by [REDACTED_PERSON]

	#loc UE4DocJpn File updated against INT#3599948

Change 3606448 by [REDACTED_PERSON]

	#loc UE4DocJpn File updated against INT#3599948

Change 3606482 by [REDACTED_PERSON]

	#loc UE4DocJpn File updated against INT#3599948

Change 3606485 by [REDACTED_PERSON]

	#loc UE4DocJpn File updated against INT#3599948

Change 3606552 by [REDACTED_PERSON]

	#loc UE4DocJpn File updated against INT#3454520

Change 3608492 by [REDACTED_PERSON]

	Implemented editorial reviews.

Change 3609596 by [REDACTED_PERSON]

	#loc UE4DocKor

Change 3609599 by [REDACTED_PERSON]

	#loc UE4DocKor

Change 3610206 by [REDACTED_PERSON]

	Updated Plugins document with some new fields that have appeared in .uplugin files.

Change 3610971 by [REDACTED_PERSON]

	Removing destructibles content and levels from Content Examples.
	Disabling Apex Destruction plugin in Content Examples project.
	Fixing any levels that utilized destructibles content.

Change 3611013 by [REDACTED_PERSON]

	#UE4 Docs: UEDOC-5077 - Clipping System updates for 4.17 SME Review edits.

	* Clarified one section
	* Updated image

Change 3611047 by [REDACTED_PERSON]

	Setting T_TD_FlameTower_E2 texture compression quality to highest.

Change 3611622 by [REDACTED_PERSON]

	#loc UE4DocKOR

Change 3612088 by [REDACTED_PERSON]

	#loc UE4DocJpn File updated against INT#3611013

Change 3612112 by [REDACTED_PERSON]

	#loc UE4DocJpn File updated against INT#3608492

Change 3612127 by [REDACTED_PERSON]

	#loc UE4DocJpn File updated against INT#3606403

Change 3612139 by [REDACTED_PERSON]

	#loc UE4DocJpn File updated against INT#3606403

Change 3612169 by [REDACTED_PERSON]

	#loc UE4DocJpn File updated against INT#3599948

Change 3612206 by [REDACTED_PERSON]

	#loc UE4DocCHN. do translation.

Change 3612432 by [REDACTED_PERSON]

	Replacing deprecated BP node in BP_Bomb

Change 3612551 by [REDACTED_PERSON]

	#UE4 Docs: UEDOC-5078 - UMG Clipping Updates -

	* Fixed the UMG User Guide page to now use directory listing vs topic image links.
	* Updated Metadata for all pages under Engine/UMG/UserGuide.

Change 3614123 by [REDACTED_PERSON]

	#loc UE4DocJpn File updated against INT#:3612551

Change 3614127 by [REDACTED_PERSON]

	#loc UE4DocJpn File updated against INT#3612551

Change 3614132 by [REDACTED_PERSON]

	#loc UE4DocJpn File updated against INT#3612551

Change 3614134 by [REDACTED_PERSON]

	#loc UE4DocJpn File updated against INT#3612551

Change 3614140 by [REDACTED_PERSON]

	#loc UE4DocJpn File updated against INT#3612551

Change 3614142 by [REDACTED_PERSON]

	#loc UE4DocJpn File updated against INT#3612551

Change 3614146 by [REDACTED_PERSON]

	#loc UE4DocJpn File updated against INT#:3612551

Change 3614149 by [REDACTED_PERSON]

	#loc UE4DocJpn File updated against INT#3612551

Change 3614151 by [REDACTED_PERSON]

	Updated file against INTSourceChangelist:3603669

Change 3614152 by [REDACTED_PERSON]

	#loc UE4DocJpn File updated against INT#3612551

Change 3614154 by [REDACTED_PERSON]

	#loc UE4DocJpn File updated against INT#:3612551

Change 3614157 by [REDACTED_PERSON]

	#loc UE4DocJpn File updated against INT#3612551

Change 3614160 by [REDACTED_PERSON]

	#loc UE4DocJpn File updated against INT#3612551

Change 3614163 by [REDACTED_PERSON]

	Updated file agasint INTSourceChangelist:3602897

Change 3614164 by [REDACTED_PERSON]

	#loc UE4DocJpn File updated against INT#3612551

Change 3614168 by [REDACTED_PERSON]

	#loc UE4DocJpn File updated against INT#3612551

Change 3614183 by [REDACTED_PERSON]

	Updated file against INTSourceChangelist:3599948

Change 3614188 by [REDACTED_PERSON]

	Updated file against INTSourceChangelist:3599948

Change 3614192 by [REDACTED_PERSON]

	Updated file against INTSourceChangelist:3599948

Change 3614208 by [REDACTED_PERSON]

	Updated file against INTSourceChangelist:3599948

Change 3614457 by [REDACTED_PERSON]

	#loc UE4DocCHN. do translation.

Change 3614682 by [REDACTED_PERSON]

	Added missing topic image

Change 3614688 by [REDACTED_PERSON]

	Adding missing Mac and Windows images for StaticMeshes and ContentBrowser pages.
	#UEDoc

Change 3614707 by [REDACTED_PERSON]

	#UE4 Docs: Pivot Painter doc: Fixed a bad related page link.

Change 3615003 by [REDACTED_PERSON]

	#ue4 docs

	Initial doc for Anim Modifiers

Change 3615724 by [REDACTED_PERSON]

	#ue4 docs

	Added Marker Based Syncing to Sync Group page.

Change 3615845 by [REDACTED_PERSON]

	#ue4 docs

	Applying SME feedback.

Change 3616100 by [REDACTED_PERSON]

	#ue4 docs

	Applying Editorial Review feedback.

Change 3616216 by [REDACTED_PERSON]

	#ue4 docs

	Applying Editorial Review feedback.

Change 3616274 by [REDACTED_PERSON]

	Added note according to SME's feedback.

	#UE4Doc

Change 3616288 by [REDACTED_PERSON]

	#ue4 docs

	Applying Editorial Review feedback.

Change 3616389 by [REDACTED_PERSON]

	#UE4 Docs: UEDOC-5018 - Color Grading Editorial Review

	* Cleaned up the page a little bit.
	* Moved Lookup Tables (LUTs) into its own page since it was cluttering up the Color Grading/Tonemapper page.

Change 3617131 by [REDACTED_PERSON]

	#loc UE4DocJPN moved or deleted

Change 3617142 by [REDACTED_PERSON]

	#loc UE4DocJpn File updated against INT#3614707

Change 3617147 by [REDACTED_PERSON]

	#loc UE4DocJpn File updated against INT#3614688

Change 3617153 by [REDACTED_PERSON]

	#loc UE4DocJpn File updated against INT#3481084

Change 3617155 by [REDACTED_PERSON]

	#loc UE4DocJpn File updated against INT#3429240

Change 3617191 by [REDACTED_PERSON]

	#loc UE4DocJpn File updated against INT#3615003

Change 3617196 by [REDACTED_PERSON]

	#loc UE4DocJpn File updated against INT#3612551

Change 3617213 by [REDACTED_PERSON]

	#loc UE4DocJpn File updated against INT#:3616288

Change 3617218 by [REDACTED_PERSON]

	#loc UE4DocJpn File updated against INT#3616216

Change 3617266 by [REDACTED_PERSON]

	#loc UE4DocJpn File updated against INT#:3615724

Change 3617302 by [REDACTED_PERSON]

	Updated file against INTSourceChangelist:3596779

Change 3617398 by [REDACTED_PERSON]

	#loc UE4DocJpn Initial translation against INT#3615003

Change 3617476 by [REDACTED_PERSON]

	Updated file against INTSourceChangelist:3616389

Change 3617957 by [REDACTED_PERSON]

	#UE4 Docs:  UEDOC-5018 - Color Grading and Using LUTs page Editorial Review edits.

Change 3618359 by [REDACTED_PERSON]

	#UE4 Docs: UEDOC-5078 - Clipping System Editorial Review edits.

	- grammatical fixes
	- updated image
	- clarified some sections to flow better.

Change 3618434 by [REDACTED_PERSON]

	Removing these images as this is not something you can do any more.

Change 3618473 by [REDACTED_PERSON]

	Removing this as it is not needed.

Change 3618474 by [REDACTED_PERSON]

	Removing this as it is not needed.

Change 3618677 by [REDACTED_PERSON]

	#UE4 Docs: Removing this as the user no longer needs to set this up.

Change 3619133 by [REDACTED_PERSON]

	Updated file against INTSourceChangelist:3616274

Change 3619194 by [REDACTED_PERSON]

	#loc UE4DocJpn File updated against INT#3618359

Change 3620333 by [REDACTED_PERSON]

	#UE4 Docs: UEDOC-4647 - Virtual Bones Editorial Review edits.

Change 3620832 by [REDACTED_PERSON]

	#ue4 docs

	Updates made to Overview and Reference pages:
	- Updated images/videos for both page.
	- Updated content to include new features.

Change 3621748 by [REDACTED_PERSON]

	#loc UE4DocJpn File updated against INT#3620832

Change 3621931 by [REDACTED_PERSON]

	#loc UE4DocKor

Change 3622284 by [REDACTED_PERSON]

	#UE4 Docs: Applying peer feedback and a new page referance page.

Change 3622466 by [REDACTED_PERSON]

	#UE4 Docs: VR-Re Org Upate - Updating this page to cover all Motion Controller's UE4 supports.

Change 3622471 by [REDACTED_PERSON]

	#UE4 Docs: VR-Re Org Upate - Updating the information in this page to be inline with UE4.17.

Change 3622501 by [REDACTED_PERSON]

	#UE4 Docs: VR-Re Org Upate - New landing page for desktop VR experiances.

Change 3622503 by [REDACTED_PERSON]

	#UE4 Docs: VR-Re Org Upate - New landing page for mobiel VR experiances.

Change 3622507 by [REDACTED_PERSON]

	#UE4 Docs: VR-Re Org Upate -Updating the information in this doc to be inline with UE4.17. Also applied peer review feedback.

Change 3622510 by [REDACTED_PERSON]

	#UE4 Docs: VR-Re Org Upate - Adding this new page about picking up and dropping items with Motion Controllers.

Change 3622529 by [REDACTED_PERSON]

	#UE4 Docs: VR-Re Org Upate - Updating the information in this document to be inline with UE4.17.
	Note that while this file is called VRCameraRefactor inside the document, it is called Attaching Items To the HMD as that is more inline with what this feature is for.

Change 3622536 by [REDACTED_PERSON]

	#ue4 docs

	Refactored Animation Blending page to Blending Animations & Properties
	- All content / images updated.

Change 3622554 by [REDACTED_PERSON]

	#UE4 Docs: VR-Re Org Upate - Updated the links on this pag to point to the desktop and mobile VR landing pages.

Change 3622561 by [REDACTED_PERSON]

	#UE4 Docs: VR-Re Org Upate - Updating the information in this page to be inline with UE4.17.

Change 3622563 by [REDACTED_PERSON]

	#UE4 Docs: VR-Re Org Upate - Updating the information in this page to be inline with UE4.17.

Change 3622565 by [REDACTED_PERSON]

	#UE4 Docs: VR-Re Org Upate - This is a new Landing Page for the Gear VR How - To section.

Change 3622571 by [REDACTED_PERSON]

	#UE4 Docs: VR-Re Org Upate - Adding new Gear VR How - To for setting up Cameras.

Change 3622576 by [REDACTED_PERSON]

	#UE4 Docs: VR-Re Org Upate - Adding a new page for using the Gear VR HMD touchpad.

Change 3622577 by [REDACTED_PERSON]

	#UE4 Docs: VR-Re Org Upate - Adding new Gear VR topic images.

Change 3622582 by [REDACTED_PERSON]

	#UE4 Docs: VR-Re Org Upate - Adding a new page for using the Gear VR controller touch pad.

Change 3622587 by [REDACTED_PERSON]

	#UE4 Docs: VR-Re Org Upate - Adding a new page for How to use the Gear VR MC touch pad buttons.

Change 3622590 by [REDACTED_PERSON]

	#UE4 Docs: VR-Re Org Upate - Adding a new page that shows how to setup a MC for Gear VR.
	Note that this page might be removed in the future.

Change 3622591 by [REDACTED_PERSON]

	#UE4 Docs: VR-Re Org Upate - Updated this page so that the information is inline with UE4.17.

Change 3622592 by [REDACTED_PERSON]

	#UE4 Docs: VR-Re Org Upate - Updateing the information in this doument to be inline with UE4.17.

Change 3622814 by [REDACTED_PERSON]

	#UE4 Docs: VR-Re Org Upate - Updating this page with the new pages created from the update to this section.

Change 3622815 by [REDACTED_PERSON]

	#UE4 Docs: VR-Re Org Upate - Adding this new document for Google VR.

Change 3622817 by [REDACTED_PERSON]

	#UE4 Docs: VR-Re Org Upate - Updating this document to be inline with UE4.17.

Change 3622818 by [REDACTED_PERSON]

	#UE4 Docs: VR-Re Org Upate - Adding this new document.

Change 3622819 by [REDACTED_PERSON]

	#UE4 Docs: VR-Re Org Upate - Adding this new How To

Change 3622820 by [REDACTED_PERSON]

	#UE4 Docs: VR-Re Org Upate - Updating the sectoin images.

Change 3622821 by [REDACTED_PERSON]

	#UE4 Docs: VR-Re Org Upate - Adding this new How To.

Change 3622822 by [REDACTED_PERSON]

	#UE4 Docs: VR-Re Org Upate - Adding this new How To

Change 3622823 by [REDACTED_PERSON]

	#UE4 Docs: VR-Re Org Upate - Adding this new How To

Change 3622825 by [REDACTED_PERSON]

	#UE4 Docs: VR-Re Org Upate - Checking in the new How To section for Google VR.

Change 3622826 by [REDACTED_PERSON]

	#UE4 Docs: VR-Re Org Upate - Checking in changes to the Google VR Quick Start.

Change 3622828 by [REDACTED_PERSON]

	#UE4 Docs: VR-Re Org Upate - Updating the Google VR Landing page with new doc page.

Change 3622829 by [REDACTED_PERSON]

	#UE4 Docs: VR-Re Org Upate - Updating this document to be inline with UE4.17.

Change 3622830 by [REDACTED_PERSON]

	#UE4 Docs: VR-Re Org Upate - Adding this new How To

Change 3622831 by [REDACTED_PERSON]

	#UE4 Docs: VR-Re Org Upate - Adding this new How To

Change 3622832 by [REDACTED_PERSON]

	#UE4 Docs: VR-Re Org Upate - Adding this new How To

Change 3622834 by [REDACTED_PERSON]

	#UE4 Docs: VR-Re Org Upate - Adding these new topic images.

Change 3622836 by [REDACTED_PERSON]

	#UE4 Docs: VR-Re Org Upate - Adding this new How To

Change 3622841 by [REDACTED_PERSON]

	#UE4 Docs: VR-Re Org Upate - Updating this doucment and applying peer review feedback.

Change 3622842 by [REDACTED_PERSON]

	#UE4 Docs: VR-Re Org Upate - Adding this new How To

Change 3622848 by [REDACTED_PERSON]

	#UE4 Docs: VR-Re Org Upate - Updating the landing page with the new document links.

Change 3622850 by [REDACTED_PERSON]

	#UE4 Docs: VR-Re Org Upate - Adding these new topic iamges.

Change 3622854 by [REDACTED_PERSON]

	#UE4 Docs: VR-Re Org Upate - Updating this section to match the other ones with the same name.

Change 3622856 by [REDACTED_PERSON]

	#UE4 Docs: VR-Re Org Upate - Updating this information in this document to be inline with UE4.17.

Change 3622857 by [REDACTED_PERSON]

	#UE4 Docs: VR-Re Org Upate - Updating this Quick Start to be inline with UE4.17.

Change 3622862 by [REDACTED_PERSON]

	#UE4 Docs: VR-Re Org Upate - Adding this new page

Change 3622864 by [REDACTED_PERSON]

	#UE4 Docs: VR-Re Org Upate - Updating this page with to point to the new VR pages.


diff --git a/Engine/Documentation/Source/API/APIReference.INT.udn b/Engine/Documentation/Source/API/APIReference.INT.udn
--- a/Engine/Documentation/Source/API/APIReference.INT.udn
+++ b/Engine/Documentation/Source/API/APIReference.INT.udn
@@ -5,6 +5,4 @@ Availability:Public
 
 [PUBLISH:Licensee]
 API reference can be found at `//depot/UE4/Engine/Documentation/Builds/CodeAPI-HTML.tgz`. You can extract the contents of that into your documentation folder to have it integrated with the engine documentation.
-
-It is also available in CHM format in that same directory.
 [/PUBLISH]
diff --git a/Engine/Documentation/Source/DocumentationGuidelines/BrandGlossary/A/A.INT.udn b/Engine/Documentation/Source/DocumentationGuidelines/BrandGlossary/A/A.INT.udn
new file mode 100644
--- /dev/null
+++ b/Engine/Documentation/Source/DocumentationGuidelines/BrandGlossary/A/A.INT.udn
@@ -0,0 +1,58 @@
+Availability: Docs
+Crumbs: %ROOT%
+Title: Brand Glossary
+Description: A glossary of key terms, which are used to describe the systems, features, and tools that belong to Unreal Engine 4 (UE4).
+Type: Reference
+Version: 4.16
+Parent:BrandGlossary
+Order:1
+Tags:Brand Glossary
+
+Welcome to the Unreal Engine 4 (UE4) **Brand Glossary**. If you want to create user-facing resources for the UE4 community, this guide 
+helps you gain a better understanding of the terms being used to describe systems, features, and tools belonging to UE4. After reading 
+through this glossary, you'll also know how to properly spell and capitalize the terms we use to describe Unreal Engine.  
+
+<hr>
+  
+<font size = "6" color=#00327F>A</font>&nbsp;&nbsp;
+[<font size = "5" color=#001940>B</font>](DocumentationGuidelines/BrandGlossary/B)&nbsp;&nbsp;
+[<font size = "5" color=#001940>C</font>](DocumentationGuidelines/BrandGlossary/C)&nbsp;&nbsp;
+[<font size = "5" color=#001940>D</font>](DocumentationGuidelines/BrandGlossary/D)&nbsp;&nbsp;
+[<font size = "5" color=#001940>E</font>](DocumentationGuidelines/BrandGlossary/E)&nbsp;&nbsp;
+[<font size = "5" color=#001940>F</font>](DocumentationGuidelines/BrandGlossary/F)&nbsp;&nbsp;
+[<font size = "5" color=#001940>G</font>](DocumentationGuidelines/BrandGlossary/G)&nbsp;&nbsp;
+[<font size = "5" color=#001940>H</font>](DocumentationGuidelines/BrandGlossary/H)&nbsp;&nbsp;
+<font size = "5" color=#001940>I</font>&nbsp;&nbsp;
+<font size = "5" color=#001940>J</font>&nbsp;&nbsp;
+<font size = "5" color=#001940>K</font>&nbsp;&nbsp;
+[<font size = "5" color=#001940>L</font>](DocumentationGuidelines/BrandGlossary/L)&nbsp;&nbsp;
+[<font size = "5" color=#001940>M</font>](DocumentationGuidelines/BrandGlossary/M)&nbsp;&nbsp;
+<font size = "5" color=#001940>N</font>&nbsp;&nbsp;
+[<font size = "5" color=#001940>O</font>](DocumentationGuidelines/BrandGlossary/O)&nbsp;&nbsp;
+[<font size = "5" color=#001940>P</font>](DocumentationGuidelines/BrandGlossary/P)&nbsp;&nbsp;
+<font size = "5" color=#001940>Q</font>&nbsp;&nbsp;
+[<font size = "5" color=#001940>R</font>](DocumentationGuidelines/BrandGlossary/R)&nbsp;&nbsp;
+[<font size = "5" color=#001940>S</font>](DocumentationGuidelines/BrandGlossary/S)&nbsp;&nbsp;
+[<font size = "5" color=#001940>T</font>](DocumentationGuidelines/BrandGlossary/T)&nbsp;&nbsp;
+[<font size = "5" color=#001940>U</font>](DocumentationGuidelines/BrandGlossary/U)&nbsp;&nbsp;
+[<font size = "5" color=#001940>V</font>](DocumentationGuidelines/BrandGlossary/V)&nbsp;&nbsp;
+[<font size = "5" color=#001940>W</font>](DocumentationGuidelines/BrandGlossary/W)&nbsp;&nbsp;
+<font size = "5" color=#001940>X</font>&nbsp;&nbsp;
+<font size = "5" color=#001940>Y</font>&nbsp;&nbsp;
+<font size = "5" color=#001940>Z</font>&nbsp;&nbsp;
+
+##Actor
+
+An Actor is any object that can be placed into a level. Actors are a generic Class that support 3D transformations, such as translation, rotation, and scale. Actors can be created (spawned) and destroyed through gameplay code (C++ or Blueprints). In C++, AActor is the base class of all Actors.
+
+##AI Controller
+
+An AIController possesses a Pawn to represent a non-player character (NPC) in a game.
+
+##Animation Blueprint
+
+An Animation Blueprint is a specialized Blueprint with graphs that control the animation of a SkeletalMesh. It can perform animation blending, directly controlling the bones of the skeleton, and outputting a final pose for a SkeletalMesh per frame.
+
+##Automation System
+
+The Automation System provides the ability to perform Unit Testing, Feature Testing, and Content Stress Testing, using the power of the Unreal Message Bus to increase stability. The Automation system works by performing one or more automation tests. Automation tests can be broken down into different categories depending on their purpose or function.
\ No newline at end of file
diff --git a/Engine/Documentation/Source/DocumentationGuidelines/BrandGlossary/B/B.INT.udn b/Engine/Documentation/Source/DocumentationGuidelines/BrandGlossary/B/B.INT.udn
new file mode 100644
--- /dev/null
+++ b/Engine/Documentation/Source/DocumentationGuidelines/BrandGlossary/B/B.INT.udn
@@ -0,0 +1,62 @@
+Availability: Docs
+Crumbs: %ROOT%
+Title: Brand Glossary
+Description: A glossary of key terms, which are used to describe the systems, features, and tools that belong to Unreal Engine 4 (UE4).
+Type: Reference
+Version: 4.16
+Parent:BrandGlossary
+Order:2
+Tags:Brand Glossary
+
+Welcome to the Unreal Engine 4 (UE4) **Brand Glossary**. If you want to create user-facing resources for the UE4 community, this guide 
+helps you gain a better understanding of the terms being used to describe systems, features, and tools belonging to UE4. After reading 
+through this glossary, you'll also know how to properly spell and capitalize the terms we use to describe Unreal Engine.  
+
+<hr>
+  
+[<font size = "5" color=#001940>A</font>](DocumentationGuidelines/BrandGlossary/A)&nbsp;&nbsp;
+<font size = "6" color=#00327F>B</font>&nbsp;&nbsp;
+[<font size = "5" color=#001940>C</font>](DocumentationGuidelines/BrandGlossary/C)&nbsp;&nbsp;
+[<font size = "5" color=#001940>D</font>](DocumentationGuidelines/BrandGlossary/D)&nbsp;&nbsp;
+[<font size = "5" color=#001940>E</font>](DocumentationGuidelines/BrandGlossary/E)&nbsp;&nbsp;
+[<font size = "5" color=#001940>F</font>](DocumentationGuidelines/BrandGlossary/F)&nbsp;&nbsp;
+[<font size = "5" color=#001940>G</font>](DocumentationGuidelines/BrandGlossary/G)&nbsp;&nbsp;
+[<font size = "5" color=#001940>H</font>](DocumentationGuidelines/BrandGlossary/H)&nbsp;&nbsp;
+<font size = "5" color=#001940>I</font>&nbsp;&nbsp;
+<font size = "5" color=#001940>J</font>&nbsp;&nbsp;
+<font size = "5" color=#001940>K</font>&nbsp;&nbsp;
+[<font size = "5" color=#001940>L</font>](DocumentationGuidelines/BrandGlossary/L)&nbsp;&nbsp;
+[<font size = "5" color=#001940>M</font>](DocumentationGuidelines/BrandGlossary/M)&nbsp;&nbsp;
+<font size = "5" color=#001940>N</font>&nbsp;&nbsp;
+[<font size = "5" color=#001940>O</font>](DocumentationGuidelines/BrandGlossary/O)&nbsp;&nbsp;
+[<font size = "5" color=#001940>P</font>](DocumentationGuidelines/BrandGlossary/P)&nbsp;&nbsp;
+<font size = "5" color=#001940>Q</font>&nbsp;&nbsp;
+[<font size = "5" color=#001940>R</font>](DocumentationGuidelines/BrandGlossary/R)&nbsp;&nbsp;
+[<font size = "5" color=#001940>S</font>](DocumentationGuidelines/BrandGlossary/S)&nbsp;&nbsp;
+[<font size = "5" color=#001940>T</font>](DocumentationGuidelines/BrandGlossary/T)&nbsp;&nbsp;
+[<font size = "5" color=#001940>U</font>](DocumentationGuidelines/BrandGlossary/U)&nbsp;&nbsp;
+[<font size = "5" color=#001940>V</font>](DocumentationGuidelines/BrandGlossary/V)&nbsp;&nbsp;
+[<font size = "5" color=#001940>W</font>](DocumentationGuidelines/BrandGlossary/W)&nbsp;&nbsp;
+<font size = "5" color=#001940>X</font>&nbsp;&nbsp;
+<font size = "5" color=#001940>Y</font>&nbsp;&nbsp;
+<font size = "5" color=#001940>Z</font>&nbsp;&nbsp;
+
+##Behavior Tree
+
+The Behavior Tree is an AI’s processor. It makes decisions, and then acts on those decisions. Behavior Trees are a powerful tool for creating Artificial Intelligence in Unreal Engine 4.
+
+##Behavior Tree Editor
+
+The Behavior Tree Editor enables you to script Artificial Intelligence through a visual node-based system (similar to Blueprints) for Actors in your game.
+
+##Blueprint
+
+Blueprints are special assets that provide an intuitive, node-based interface that can be used to create new types of Actors and script level events; giving designers and gameplay programmers the tools to quickly create and iterate gameplay from within Unreal Editor without ever needing to write a line of C++ code.
+
+##Blueprint Editor
+
+The Blueprint Editor is a node-based graph editor that acts as your primary tool for creating and editing the visual scripting node networks that make-up Blueprints.
+
+##Brush
+
+A Brush is an Actor that describes a 3D volume that is placed in a level in order to define level geometry (referred to as BSPs) and gameplay volumes.
\ No newline at end of file
diff --git a/Engine/Documentation/Source/DocumentationGuidelines/BrandGlossary/BrandGlossary.INT.udn b/Engine/Documentation/Source/DocumentationGuidelines/BrandGlossary/BrandGlossary.INT.udn
new file mode 100644
--- /dev/null
+++ b/Engine/Documentation/Source/DocumentationGuidelines/BrandGlossary/BrandGlossary.INT.udn
@@ -0,0 +1,43 @@
+Availability: Docs
+Crumbs: %ROOT%
+Title: Brand Glossary
+Description: A glossary of key terms, which are used to describe the systems, features, and tools that belong to Unreal Engine 4 (UE4).
+Type: Reference
+Version: 4.16
+Parent:
+Order:
+Tags:Brand Glossary
+
+Welcome to the Unreal Engine 4 (UE4) **Brand Glossary**. If you want to create user-facing resources for the UE4 community, this guide 
+helps you gain a better understanding of the terms being used to describe systems, features, and tools belonging to UE4. After reading 
+through this glossary, you'll also know how to properly spell and capitalize the terms we use to describe Unreal Engine.  
+
+<hr>
+  
+[<font size = "5" color=#001940>A</font>](DocumentationGuidelines/BrandGlossary/A)&nbsp;&nbsp;
+[<font size = "5" color=#001940>B</font>](DocumentationGuidelines/BrandGlossary/B)&nbsp;&nbsp;
+[<font size = "5" color=#001940>C</font>](DocumentationGuidelines/BrandGlossary/C)&nbsp;&nbsp;
+[<font size = "5" color=#001940>D</font>](DocumentationGuidelines/BrandGlossary/D)&nbsp;&nbsp;
+[<font size = "5" color=#001940>E</font>](DocumentationGuidelines/BrandGlossary/E)&nbsp;&nbsp;
+[<font size = "5" color=#001940>F</font>](DocumentationGuidelines/BrandGlossary/F)&nbsp;&nbsp;
+[<font size = "5" color=#001940>G</font>](DocumentationGuidelines/BrandGlossary/G)&nbsp;&nbsp;
+[<font size = "5" color=#001940>H</font>](DocumentationGuidelines/BrandGlossary/H)&nbsp;&nbsp;
+<font size = "5" color=#001940>I</font>&nbsp;&nbsp;
+<font size = "5" color=#001940>J</font>&nbsp;&nbsp;
+<font size = "5" color=#001940>K</font>&nbsp;&nbsp;
+[<font size = "5" color=#001940>L</font>](DocumentationGuidelines/BrandGlossary/L)&nbsp;&nbsp;
+[<font size = "5" color=#001940>M</font>](DocumentationGuidelines/BrandGlossary/M)&nbsp;&nbsp;
+<font size = "5" color=#001940>N</font>&nbsp;&nbsp;
+[<font size = "5" color=#001940>O</font>](DocumentationGuidelines/BrandGlossary/O)&nbsp;&nbsp;
+[<font size = "5" color=#001940>P</font>](DocumentationGuidelines/BrandGlossary/P)&nbsp;&nbsp;
+<font size = "5" color=#001940>Q</font>&nbsp;&nbsp;
+[<font size = "5" color=#001940>R</font>](DocumentationGuidelines/BrandGlossary/R)&nbsp;&nbsp;
+[<font size = "5" color=#001940>S</font>](DocumentationGuidelines/BrandGlossary/S)&nbsp;&nbsp;
+[<font size = "5" color=#001940>T</font>](DocumentationGuidelines/BrandGlossary/T)&nbsp;&nbsp;
+[<font size = "5" color=#001940>U</font>](DocumentationGuidelines/BrandGlossary/U)&nbsp;&nbsp;
+[<font size = "5" color=#001940>V</font>](DocumentationGuidelines/BrandGlossary/V)&nbsp;&nbsp;
+[<font size = "5" color=#001940>W</font>](DocumentationGuidelines/BrandGlossary/W)&nbsp;&nbsp;
+<font size = "5" color=#001940>X</font>&nbsp;&nbsp;
+<font size = "5" color=#001940>Y</font>&nbsp;&nbsp;
+<font size = "5" color=#001940>Z</font>&nbsp;&nbsp;
+
diff --git a/Engine/Documentation/Source/DocumentationGuidelines/BrandGlossary/C/C.INT.udn b/Engine/Documentation/Source/DocumentationGuidelines/BrandGlossary/C/C.INT.udn
new file mode 100644
--- /dev/null
+++ b/Engine/Documentation/Source/DocumentationGuidelines/BrandGlossary/C/C.INT.udn
@@ -0,0 +1,86 @@
+Availability: Docs
+Crumbs: %ROOT%
+Title: Brand Glossary
+Description: A glossary of key terms, which are used to describe the systems, features, and tools that belong to Unreal Engine 4 (UE4).
+Type: Reference
+Version: 4.16
+Parent:BrandGlossary
+Order:3
+Tags:Brand Glossary
+
+Welcome to the Unreal Engine 4 (UE4) **Brand Glossary**. If you want to create user-facing resources for the UE4 community, this guide 
+helps you gain a better understanding of the terms being used to describe systems, features, and tools belonging to UE4. After reading 
+through this glossary, you'll also know how to properly spell and capitalize the terms we use to describe Unreal Engine.  
+
+<hr>
+  
+[<font size = "5" color=#001940>A</font>](DocumentationGuidelines/BrandGlossary/A)&nbsp;&nbsp;
+[<font size = "5" color=#001940>B</font>](DocumentationGuidelines/BrandGlossary/B)&nbsp;&nbsp;
+<font size = "6" color=#00327F>C</font>&nbsp;&nbsp;
+[<font size = "5" color=#001940>D</font>](DocumentationGuidelines/BrandGlossary/D)&nbsp;&nbsp;
+[<font size = "5" color=#001940>E</font>](DocumentationGuidelines/BrandGlossary/E)&nbsp;&nbsp;
+[<font size = "5" color=#001940>F</font>](DocumentationGuidelines/BrandGlossary/F)&nbsp;&nbsp;
+[<font size = "5" color=#001940>G</font>](DocumentationGuidelines/BrandGlossary/G)&nbsp;&nbsp;
+[<font size = "5" color=#001940>H</font>](DocumentationGuidelines/BrandGlossary/H)&nbsp;&nbsp;
+<font size = "5" color=#001940>I</font>&nbsp;&nbsp;
+<font size = "5" color=#001940>J</font>&nbsp;&nbsp;
+<font size = "5" color=#001940>K</font>&nbsp;&nbsp;
+[<font size = "5" color=#001940>L</font>](DocumentationGuidelines/BrandGlossary/L)&nbsp;&nbsp;
+[<font size = "5" color=#001940>M</font>](DocumentationGuidelines/BrandGlossary/M)&nbsp;&nbsp;
+<font size = "5" color=#001940>N</font>&nbsp;&nbsp;
+[<font size = "5" color=#001940>O</font>](DocumentationGuidelines/BrandGlossary/O)&nbsp;&nbsp;
+[<font size = "5" color=#001940>P</font>](DocumentationGuidelines/BrandGlossary/P)&nbsp;&nbsp;
+<font size = "5" color=#001940>Q</font>&nbsp;&nbsp;
+[<font size = "5" color=#001940>R</font>](DocumentationGuidelines/BrandGlossary/R)&nbsp;&nbsp;
+[<font size = "5" color=#001940>S</font>](DocumentationGuidelines/BrandGlossary/S)&nbsp;&nbsp;
+[<font size = "5" color=#001940>T</font>](DocumentationGuidelines/BrandGlossary/T)&nbsp;&nbsp;
+[<font size = "5" color=#001940>U</font>](DocumentationGuidelines/BrandGlossary/U)&nbsp;&nbsp;
+[<font size = "5" color=#001940>V</font>](DocumentationGuidelines/BrandGlossary/V)&nbsp;&nbsp;
+[<font size = "5" color=#001940>W</font>](DocumentationGuidelines/BrandGlossary/W)&nbsp;&nbsp;
+<font size = "5" color=#001940>X</font>&nbsp;&nbsp;
+<font size = "5" color=#001940>Y</font>&nbsp;&nbsp;
+<font size = "5" color=#001940>Z</font>&nbsp;&nbsp;
+
+##Camera
+
+The PlayerCameraManager is the player’s “eyeball”, managing how it behaves. Each PlayerController typically has one of these as well. For more, see the camera workflow page.
+
+##Cascade Particle Editor
+
+The Cascade Particle Editor is a fully integrated and modular particle effects editor that offers real-time feedback and modular effects editing, allowing fast and easy creation of even the most complex effects.
+
+##Character
+
+A Character is a subclass of a Pawn Actor that is intended to be used as a player character. The Character subclass includes a collision setup, input bindings for bipedal movement, and additional code for movement controlled by the player.
+
+##Class
+
+A Class defines the behaviors and properties of a particular Actor or Object used in the creation of an Unreal Engine game. Classes are hierarchical, which means that a Class inherits information from its parent Classes (the Classes it was derived, or “sub-classed,” from) and passes that information to its children. Classes can be created in C++ code or in Blueprints.
+
+##Class Viewer
+
+Browser and picker functionality for C++ and Blueprint classes.
+
+##Color Picker
+
+The Color Picker allows you to easily adjust a color value for a color property in the Unreal Editor. You can adjust colors by the following values: RGBA (red, green, blue, alpha), HSV (hue, saturation, value), and Hex (ARGB). Alternatively you can select a color by clicking anywhere in the color wheel, by capturing the color under your mouse, or clicking anywhere on your monitor.
+
+##Component
+
+A Component is a piece of functionality that can be added to an Actor. Components cannot exist by themselves, however when added to an Actor, the Actor will have access to functionality provided by the Component.
+
+##Content Browser
+
+The Content Browser is the primary area of the Unreal Editor for creating, importing, organizing, viewing, and modifying content assets within Unreal Editor.
+
+##Content Cooking
+
+Content Cooking is the process of converting content from the internal formats used by Unreal Editor, such as PNG for texture data or WAV for audio, to the formats used by each individual platform; either because the platform uses a proprietary format, does not support the format Unreal uses to store the asset, or a more memory/performance-effective format exists.
+
+##Controller
+
+Controllers are non-physical Actors that can possess a Pawn or Character to control its actions.
+
+##Curve Editor
+
+The Curve Editor provides interactive editing and fine control of properties that change over time (for example, properties that make use of Distributions).
\ No newline at end of file
diff --git a/Engine/Documentation/Source/DocumentationGuidelines/BrandGlossary/D/D.INT.udn b/Engine/Documentation/Source/DocumentationGuidelines/BrandGlossary/D/D.INT.udn
new file mode 100644
--- /dev/null
+++ b/Engine/Documentation/Source/DocumentationGuidelines/BrandGlossary/D/D.INT.udn
@@ -0,0 +1,54 @@
+Availability: Docs
+Crumbs: %ROOT%
+Title: Brand Glossary
+Description: A glossary of key terms, which are used to describe the systems, features, and tools that belong to Unreal Engine 4 (UE4).
+Type: Reference
+Version: 4.16
+Parent:BrandGlossary
+Order:4
+Tags:Brand Glossary
+
+Welcome to the Unreal Engine 4 (UE4) **Brand Glossary**. If you want to create user-facing resources for the UE4 community, this guide 
+helps you gain a better understanding of the terms being used to describe systems, features, and tools belonging to UE4. After reading 
+through this glossary, you'll also know how to properly spell and capitalize the terms we use to describe Unreal Engine.  
+
+<hr>
+  
+[<font size = "5" color=#001940>A</font>](DocumentationGuidelines/BrandGlossary/A)&nbsp;&nbsp;
+[<font size = "5" color=#001940>B</font>](DocumentationGuidelines/BrandGlossary/B)&nbsp;&nbsp;
+[<font size = "5" color=#001940>C</font>](DocumentationGuidelines/BrandGlossary/C)&nbsp;&nbsp;
+<font size = "6" color=#00327F>D</font>&nbsp;&nbsp;
+[<font size = "5" color=#001940>E</font>](DocumentationGuidelines/BrandGlossary/E)&nbsp;&nbsp;
+[<font size = "5" color=#001940>F</font>](DocumentationGuidelines/BrandGlossary/F)&nbsp;&nbsp;
+[<font size = "5" color=#001940>G</font>](DocumentationGuidelines/BrandGlossary/G)&nbsp;&nbsp;
+[<font size = "5" color=#001940>H</font>](DocumentationGuidelines/BrandGlossary/H)&nbsp;&nbsp;
+<font size = "5" color=#001940>I</font>&nbsp;&nbsp;
+<font size = "5" color=#001940>J</font>&nbsp;&nbsp;
+<font size = "5" color=#001940>K</font>&nbsp;&nbsp;
+[<font size = "5" color=#001940>L</font>](DocumentationGuidelines/BrandGlossary/L)&nbsp;&nbsp;
+[<font size = "5" color=#001940>M</font>](DocumentationGuidelines/BrandGlossary/M)&nbsp;&nbsp;
+<font size = "5" color=#001940>N</font>&nbsp;&nbsp;
+[<font size = "5" color=#001940>O</font>](DocumentationGuidelines/BrandGlossary/O)&nbsp;&nbsp;
+[<font size = "5" color=#001940>P</font>](DocumentationGuidelines/BrandGlossary/P)&nbsp;&nbsp;
+<font size = "5" color=#001940>Q</font>&nbsp;&nbsp;
+[<font size = "5" color=#001940>R</font>](DocumentationGuidelines/BrandGlossary/R)&nbsp;&nbsp;
+[<font size = "5" color=#001940>S</font>](DocumentationGuidelines/BrandGlossary/S)&nbsp;&nbsp;
+[<font size = "5" color=#001940>T</font>](DocumentationGuidelines/BrandGlossary/T)&nbsp;&nbsp;
+[<font size = "5" color=#001940>U</font>](DocumentationGuidelines/BrandGlossary/U)&nbsp;&nbsp;
+[<font size = "5" color=#001940>V</font>](DocumentationGuidelines/BrandGlossary/V)&nbsp;&nbsp;
+[<font size = "5" color=#001940>W</font>](DocumentationGuidelines/BrandGlossary/W)&nbsp;&nbsp;
+<font size = "5" color=#001940>X</font>&nbsp;&nbsp;
+<font size = "5" color=#001940>Y</font>&nbsp;&nbsp;
+<font size = "5" color=#001940>Z</font>&nbsp;&nbsp;
+
+##Derived Data Cache
+
+Manages the cache of engine-optimized data that has been derived from source assets.
+
+##Details Panel
+
+The Details panel contains information, utilities, and functions specific to the current selection in the viewport; such as transform edit boxes for moving, rotating, and scaling Actors, editable properties for the selected Actors, and quick access to additional editing functionality (depending on the types of Actors selected in the viewport).
+
+##Distributions
+
+Distributions are a group of data types that provide flexibility by allowing constant values, random values within a range, values interpolated along a curve, and values driven by parameters. These are generally used in properties of Particle Systems and nodes within SoundCues.
\ No newline at end of file
diff --git a/Engine/Documentation/Source/DocumentationGuidelines/BrandGlossary/E/E.INT.udn b/Engine/Documentation/Source/DocumentationGuidelines/BrandGlossary/E/E.INT.udn
new file mode 100644
--- /dev/null
+++ b/Engine/Documentation/Source/DocumentationGuidelines/BrandGlossary/E/E.INT.udn
@@ -0,0 +1,54 @@
+Availability: Docs
+Crumbs: %ROOT%
+Title: Brand Glossary
+Description: A glossary of key terms, which are used to describe the systems, features, and tools that belong to Unreal Engine 4 (UE4).
+Type: Reference
+Version: 4.16
+Parent:BrandGlossary
+Order:5
+Tags:Brand Glossary
+
+Welcome to the Unreal Engine 4 (UE4) **Brand Glossary**. If you want to create user-facing resources for the UE4 community, this guide 
+helps you gain a better understanding of the terms being used to describe systems, features, and tools belonging to UE4. After reading 
+through this glossary, you'll also know how to properly spell and capitalize the terms we use to describe Unreal Engine.  
+
+<hr>
+  
+[<font size = "5" color=#001940>A</font>](DocumentationGuidelines/BrandGlossary/A)&nbsp;&nbsp;
+[<font size = "5" color=#001940>B</font>](DocumentationGuidelines/BrandGlossary/B)&nbsp;&nbsp;
+[<font size = "5" color=#001940>C</font>](DocumentationGuidelines/BrandGlossary/C)&nbsp;&nbsp;
+[<font size = "5" color=#001940>D</font>](DocumentationGuidelines/BrandGlossary/D)&nbsp;&nbsp;
+<font size = "6" color=#00327F>E</font>&nbsp;&nbsp;
+[<font size = "5" color=#001940>F</font>](DocumentationGuidelines/BrandGlossary/F)&nbsp;&nbsp;
+[<font size = "5" color=#001940>G</font>](DocumentationGuidelines/BrandGlossary/G)&nbsp;&nbsp;
+[<font size = "5" color=#001940>H</font>](DocumentationGuidelines/BrandGlossary/H)&nbsp;&nbsp;
+<font size = "5" color=#001940>I</font>&nbsp;&nbsp;
+<font size = "5" color=#001940>J</font>&nbsp;&nbsp;
+<font size = "5" color=#001940>K</font>&nbsp;&nbsp;
+[<font size = "5" color=#001940>L</font>](DocumentationGuidelines/BrandGlossary/L)&nbsp;&nbsp;
+[<font size = "5" color=#001940>M</font>](DocumentationGuidelines/BrandGlossary/M)&nbsp;&nbsp;
+<font size = "5" color=#001940>N</font>&nbsp;&nbsp;
+[<font size = "5" color=#001940>O</font>](DocumentationGuidelines/BrandGlossary/O)&nbsp;&nbsp;
+[<font size = "5" color=#001940>P</font>](DocumentationGuidelines/BrandGlossary/P)&nbsp;&nbsp;
+<font size = "5" color=#001940>Q</font>&nbsp;&nbsp;
+[<font size = "5" color=#001940>R</font>](DocumentationGuidelines/BrandGlossary/R)&nbsp;&nbsp;
+[<font size = "5" color=#001940>S</font>](DocumentationGuidelines/BrandGlossary/S)&nbsp;&nbsp;
+[<font size = "5" color=#001940>T</font>](DocumentationGuidelines/BrandGlossary/T)&nbsp;&nbsp;
+[<font size = "5" color=#001940>U</font>](DocumentationGuidelines/BrandGlossary/U)&nbsp;&nbsp;
+[<font size = "5" color=#001940>V</font>](DocumentationGuidelines/BrandGlossary/V)&nbsp;&nbsp;
+[<font size = "5" color=#001940>W</font>](DocumentationGuidelines/BrandGlossary/W)&nbsp;&nbsp;
+<font size = "5" color=#001940>X</font>&nbsp;&nbsp;
+<font size = "5" color=#001940>Y</font>&nbsp;&nbsp;
+<font size = "5" color=#001940>Z</font>&nbsp;&nbsp;
+
+##Editor Preferences
+
+The Editor Preferences window is used to modify settings that control the behavior of Unreal Editor with respect to controls, viewports, source control, auto-saving, and more.
+
+##Environment Query Editor
+
+The editor for AI-navigation environment query templates. Allows the design of queries for such things as cover and spawn points.
+
+##Environment Query System (EQS)
+
+The Environment Query System is a feature of the Artificial Intelligence system in Unreal Engine 4 for collecting data on the environment, asking questions of the data through Tests, then returning an Item that best answers the question.
\ No newline at end of file
diff --git a/Engine/Documentation/Source/DocumentationGuidelines/BrandGlossary/F/F.INT.udn b/Engine/Documentation/Source/DocumentationGuidelines/BrandGlossary/F/F.INT.udn
new file mode 100644
--- /dev/null
+++ b/Engine/Documentation/Source/DocumentationGuidelines/BrandGlossary/F/F.INT.udn
@@ -0,0 +1,50 @@
+Availability: Docs
+Crumbs: %ROOT%
+Title: Brand Glossary
+Description: A glossary of key terms, which are used to describe the systems, features, and tools that belong to Unreal Engine 4 (UE4).
+Type: Reference
+Version: 4.16
+Parent:BrandGlossary
+Order:6
+Tags:Brand Glossary
+
+Welcome to the Unreal Engine 4 (UE4) **Brand Glossary**. If you want to create user-facing resources for the UE4 community, this guide 
+helps you gain a better understanding of the terms being used to describe systems, features, and tools belonging to UE4. After reading 
+through this glossary, you'll also know how to properly spell and capitalize the terms we use to describe Unreal Engine.  
+
+<hr>
+  
+[<font size = "5" color=#001940>A</font>](DocumentationGuidelines/BrandGlossary/A)&nbsp;&nbsp;
+[<font size = "5" color=#001940>B</font>](DocumentationGuidelines/BrandGlossary/B)&nbsp;&nbsp;
+[<font size = "5" color=#001940>C</font>](DocumentationGuidelines/BrandGlossary/C)&nbsp;&nbsp;
+[<font size = "5" color=#001940>D</font>](DocumentationGuidelines/BrandGlossary/D)&nbsp;&nbsp;
+[<font size = "5" color=#001940>E</font>](DocumentationGuidelines/BrandGlossary/E)&nbsp;&nbsp;
+<font size = "6" color=#00327F>F</font>&nbsp;&nbsp;
+[<font size = "5" color=#001940>G</font>](DocumentationGuidelines/BrandGlossary/G)&nbsp;&nbsp;
+[<font size = "5" color=#001940>H</font>](DocumentationGuidelines/BrandGlossary/H)&nbsp;&nbsp;
+<font size = "5" color=#001940>I</font>&nbsp;&nbsp;
+<font size = "5" color=#001940>J</font>&nbsp;&nbsp;
+<font size = "5" color=#001940>K</font>&nbsp;&nbsp;
+[<font size = "5" color=#001940>L</font>](DocumentationGuidelines/BrandGlossary/L)&nbsp;&nbsp;
+[<font size = "5" color=#001940>M</font>](DocumentationGuidelines/BrandGlossary/M)&nbsp;&nbsp;
+<font size = "5" color=#001940>N</font>&nbsp;&nbsp;
+[<font size = "5" color=#001940>O</font>](DocumentationGuidelines/BrandGlossary/O)&nbsp;&nbsp;
+[<font size = "5" color=#001940>P</font>](DocumentationGuidelines/BrandGlossary/P)&nbsp;&nbsp;
+<font size = "5" color=#001940>Q</font>&nbsp;&nbsp;
+[<font size = "5" color=#001940>R</font>](DocumentationGuidelines/BrandGlossary/R)&nbsp;&nbsp;
+[<font size = "5" color=#001940>S</font>](DocumentationGuidelines/BrandGlossary/S)&nbsp;&nbsp;
+[<font size = "5" color=#001940>T</font>](DocumentationGuidelines/BrandGlossary/T)&nbsp;&nbsp;
+[<font size = "5" color=#001940>U</font>](DocumentationGuidelines/BrandGlossary/U)&nbsp;&nbsp;
+[<font size = "5" color=#001940>V</font>](DocumentationGuidelines/BrandGlossary/V)&nbsp;&nbsp;
+[<font size = "5" color=#001940>W</font>](DocumentationGuidelines/BrandGlossary/W)&nbsp;&nbsp;
+<font size = "5" color=#001940>X</font>&nbsp;&nbsp;
+<font size = "5" color=#001940>Y</font>&nbsp;&nbsp;
+<font size = "5" color=#001940>Z</font>&nbsp;&nbsp;
+
+##Foliage Instanced Meshes
+
+The Foliage Instanced Mesh system allows you to quickly paint or erase sets of Static Meshes on Landscape Actors, other Static Mesh Actors, and BSP. These meshes are automatically grouped together into batches that are rendered using hardware instancing; meaning, that many instances can be rendered with only a single draw call.
+
+##Foliage Mode
+
+The Foliage Instanced Mesh system enables you to quickly paint or erase batched sets of Static Meshes that use hardware instancing for efficient rendering on Landscape Actors, other Static Mesh Actors, and BSP geometry.
\ No newline at end of file
diff --git a/Engine/Documentation/Source/DocumentationGuidelines/BrandGlossary/G/G.INT.udn b/Engine/Documentation/Source/DocumentationGuidelines/BrandGlossary/G/G.INT.udn
new file mode 100644
--- /dev/null
+++ b/Engine/Documentation/Source/DocumentationGuidelines/BrandGlossary/G/G.INT.udn
@@ -0,0 +1,62 @@
+Availability: Docs
+Crumbs: %ROOT%
+Title: Brand Glossary
+Description: A glossary of key terms, which are used to describe the systems, features, and tools that belong to Unreal Engine 4 (UE4).
+Type: Reference
+Version: 4.16
+Parent:BrandGlossary
+Order:7
+Tags:Brand Glossary
+
+Welcome to the Unreal Engine 4 (UE4) **Brand Glossary**. If you want to create user-facing resources for the UE4 community, this guide 
+helps you gain a better understanding of the terms being used to describe systems, features, and tools belonging to UE4. After reading 
+through this glossary, you'll also know how to properly spell and capitalize the terms we use to describe Unreal Engine.  
+
+<hr>
+  
+[<font size = "5" color=#001940>A</font>](DocumentationGuidelines/BrandGlossary/A)&nbsp;&nbsp;
+[<font size = "5" color=#001940>B</font>](DocumentationGuidelines/BrandGlossary/B)&nbsp;&nbsp;
+[<font size = "5" color=#001940>C</font>](DocumentationGuidelines/BrandGlossary/C)&nbsp;&nbsp;
+[<font size = "5" color=#001940>D</font>](DocumentationGuidelines/BrandGlossary/D)&nbsp;&nbsp;
+[<font size = "5" color=#001940>E</font>](DocumentationGuidelines/BrandGlossary/E)&nbsp;&nbsp;
+[<font size = "5" color=#001940>F</font>](DocumentationGuidelines/BrandGlossary/F)&nbsp;&nbsp;
+<font size = "6" color=#00327F>G</font>&nbsp;&nbsp;
+[<font size = "5" color=#001940>H</font>](DocumentationGuidelines/BrandGlossary/H)&nbsp;&nbsp;
+<font size = "5" color=#001940>I</font>&nbsp;&nbsp;
+<font size = "5" color=#001940>J</font>&nbsp;&nbsp;
+<font size = "5" color=#001940>K</font>&nbsp;&nbsp;
+[<font size = "5" color=#001940>L</font>](DocumentationGuidelines/BrandGlossary/L)&nbsp;&nbsp;
+[<font size = "5" color=#001940>M</font>](DocumentationGuidelines/BrandGlossary/M)&nbsp;&nbsp;
+<font size = "5" color=#001940>N</font>&nbsp;&nbsp;
+[<font size = "5" color=#001940>O</font>](DocumentationGuidelines/BrandGlossary/O)&nbsp;&nbsp;
+[<font size = "5" color=#001940>P</font>](DocumentationGuidelines/BrandGlossary/P)&nbsp;&nbsp;
+<font size = "5" color=#001940>Q</font>&nbsp;&nbsp;
+[<font size = "5" color=#001940>R</font>](DocumentationGuidelines/BrandGlossary/R)&nbsp;&nbsp;
+[<font size = "5" color=#001940>S</font>](DocumentationGuidelines/BrandGlossary/S)&nbsp;&nbsp;
+[<font size = "5" color=#001940>T</font>](DocumentationGuidelines/BrandGlossary/T)&nbsp;&nbsp;
+[<font size = "5" color=#001940>U</font>](DocumentationGuidelines/BrandGlossary/U)&nbsp;&nbsp;
+[<font size = "5" color=#001940>V</font>](DocumentationGuidelines/BrandGlossary/V)&nbsp;&nbsp;
+[<font size = "5" color=#001940>W</font>](DocumentationGuidelines/BrandGlossary/W)&nbsp;&nbsp;
+<font size = "5" color=#001940>X</font>&nbsp;&nbsp;
+<font size = "5" color=#001940>Y</font>&nbsp;&nbsp;
+<font size = "5" color=#001940>Z</font>&nbsp;&nbsp;
+
+##Game Mode
+
+The GameMode Class is responsible for setting the rules of the game that is being played. The rules can include how players join the game, whether or not a game can be paused, level transitions, as well as any game-specific behavior (for example, win conditions).
+
+##Gameplay Framework
+
+The basic gameplay classes include functionality for representing players, allies, and enemies; as well as for controlling these avatars with player input or AI logic. There are also classes for creating heads-up displays and cameras for players. Finally, gameplay classes such as GameMode, GameState, and PlayerState set the rules of the game. Gameplay classes also track how the game and the players are progressing.
+
+##GameState
+
+The GameState contains information that you want replicated to every client in a game. Put simply, it is ‘The State of the Game’ for every connected player.
+
+##Global Asset Picker
+
+The Global Asset Picker (Ctrl-P) is a fast-access means of editing assets or bringing them into your levels. In some ways, it is similar to the Content Browser but without limiting itself to the current folder you have selected in the Asset Tree. Not only can you pull assets specific to your game, but also anything from Engine, like lights or sound emitters. Because it offers all available assets in a list format, it is not intended to be used as means for browsing. Instead, it’s perfect for those moments when you know the asset (or asset type) you need and just want to quickly type part of its name into the Search Line.
+
+##Grass Tool
+
+A tool to densely populate the Landscape Terrain used in a forest with grass, flowers, and small rocks.
\ No newline at end of file
diff --git a/Engine/Documentation/Source/DocumentationGuidelines/BrandGlossary/H/H.INT.udn b/Engine/Documentation/Source/DocumentationGuidelines/BrandGlossary/H/H.INT.udn
new file mode 100644
--- /dev/null
+++ b/Engine/Documentation/Source/DocumentationGuidelines/BrandGlossary/H/H.INT.udn
@@ -0,0 +1,46 @@
+Availability: Docs
+Crumbs: %ROOT%
+Title: Brand Glossary
+Description: A glossary of key terms, which are used to describe the systems, features, and tools that belong to Unreal Engine 4 (UE4).
+Type: Reference
+Version: 4.16
+Parent:BrandGlossary
+Order:8
+Tags:Brand Glossary
+
+Welcome to the Unreal Engine 4 (UE4) **Brand Glossary**. If you want to create user-facing resources for the UE4 community, this guide 
+helps you gain a better understanding of the terms being used to describe systems, features, and tools belonging to UE4. After reading 
+through this glossary, you'll also know how to properly spell and capitalize the terms we use to describe Unreal Engine.  
+
+<hr>
+  
+[<font size = "5" color=#001940>A</font>](DocumentationGuidelines/BrandGlossary/A)&nbsp;&nbsp;
+[<font size = "5" color=#001940>B</font>](DocumentationGuidelines/BrandGlossary/B)&nbsp;&nbsp;
+[<font size = "5" color=#001940>C</font>](DocumentationGuidelines/BrandGlossary/C)&nbsp;&nbsp;
+[<font size = "5" color=#001940>D</font>](DocumentationGuidelines/BrandGlossary/D)&nbsp;&nbsp;
+[<font size = "5" color=#001940>E</font>](DocumentationGuidelines/BrandGlossary/E)&nbsp;&nbsp;
+[<font size = "5" color=#001940>F</font>](DocumentationGuidelines/BrandGlossary/F)&nbsp;&nbsp;
+[<font size = "5" color=#001940>G</font>](DocumentationGuidelines/BrandGlossary/G)&nbsp;&nbsp;
+<font size = "6" color=#00327F>H</font>&nbsp;&nbsp;
+<font size = "5" color=#001940>I</font>&nbsp;&nbsp;
+<font size = "5" color=#001940>J</font>&nbsp;&nbsp;
+<font size = "5" color=#001940>K</font>&nbsp;&nbsp;
+[<font size = "5" color=#001940>L</font>](DocumentationGuidelines/BrandGlossary/L)&nbsp;&nbsp;
+[<font size = "5" color=#001940>M</font>](DocumentationGuidelines/BrandGlossary/M)&nbsp;&nbsp;
+<font size = "5" color=#001940>N</font>&nbsp;&nbsp;
+[<font size = "5" color=#001940>O</font>](DocumentationGuidelines/BrandGlossary/O)&nbsp;&nbsp;
+[<font size = "5" color=#001940>P</font>](DocumentationGuidelines/BrandGlossary/P)&nbsp;&nbsp;
+<font size = "5" color=#001940>Q</font>&nbsp;&nbsp;
+[<font size = "5" color=#001940>R</font>](DocumentationGuidelines/BrandGlossary/R)&nbsp;&nbsp;
+[<font size = "5" color=#001940>S</font>](DocumentationGuidelines/BrandGlossary/S)&nbsp;&nbsp;
+[<font size = "5" color=#001940>T</font>](DocumentationGuidelines/BrandGlossary/T)&nbsp;&nbsp;
+[<font size = "5" color=#001940>U</font>](DocumentationGuidelines/BrandGlossary/U)&nbsp;&nbsp;
+[<font size = "5" color=#001940>V</font>](DocumentationGuidelines/BrandGlossary/V)&nbsp;&nbsp;
+[<font size = "5" color=#001940>W</font>](DocumentationGuidelines/BrandGlossary/W)&nbsp;&nbsp;
+<font size = "5" color=#001940>X</font>&nbsp;&nbsp;
+<font size = "5" color=#001940>Y</font>&nbsp;&nbsp;
+<font size = "5" color=#001940>Z</font>&nbsp;&nbsp;
+
+##HUD
+
+A HUD is a “heads-up display”, or the 2D on-screen display that is common in many games (for example, health bar, ammunition indicator, gun reticle, etc.) Each PlayerController typically has one of these.
\ No newline at end of file
diff --git a/Engine/Documentation/Source/DocumentationGuidelines/BrandGlossary/L/L.INT.udn b/Engine/Documentation/Source/DocumentationGuidelines/BrandGlossary/L/L.INT.udn
new file mode 100644
--- /dev/null
+++ b/Engine/Documentation/Source/DocumentationGuidelines/BrandGlossary/L/L.INT.udn
@@ -0,0 +1,70 @@
+Availability: Docs
+Crumbs: %ROOT%
+Title: Brand Glossary
+Description: A glossary of key terms, which are used to describe the systems, features, and tools that belong to Unreal Engine 4 (UE4).
+Type: Reference
+Version: 4.16
+Parent:BrandGlossary
+Order:9
+Tags:Brand Glossary
+
+Welcome to the Unreal Engine 4 (UE4) **Brand Glossary**. If you want to create user-facing resources for the UE4 community, this guide 
+helps you gain a better understanding of the terms being used to describe systems, features, and tools belonging to UE4. After reading 
+through this glossary, you'll also know how to properly spell and capitalize the terms we use to describe Unreal Engine.  
+
+<hr>
+  
+[<font size = "5" color=#001940>A</font>](DocumentationGuidelines/BrandGlossary/A)&nbsp;&nbsp;
+[<font size = "5" color=#001940>B</font>](DocumentationGuidelines/BrandGlossary/B)&nbsp;&nbsp;
+[<font size = "5" color=#001940>C</font>](DocumentationGuidelines/BrandGlossary/C)&nbsp;&nbsp;
+[<font size = "5" color=#001940>D</font>](DocumentationGuidelines/BrandGlossary/D)&nbsp;&nbsp;
+[<font size = "5" color=#001940>E</font>](DocumentationGuidelines/BrandGlossary/E)&nbsp;&nbsp;
+[<font size = "5" color=#001940>F</font>](DocumentationGuidelines/BrandGlossary/F)&nbsp;&nbsp;
+[<font size = "5" color=#001940>G</font>](DocumentationGuidelines/BrandGlossary/G)&nbsp;&nbsp;
+[<font size = "5" color=#001940>H</font>](DocumentationGuidelines/BrandGlossary/H)&nbsp;&nbsp;
+<font size = "5" color=#001940>I</font>&nbsp;&nbsp;
+<font size = "5" color=#001940>J</font>&nbsp;&nbsp;
+<font size = "5" color=#001940>K</font>&nbsp;&nbsp;
+<font size = "6" color=#00327F>L</font>&nbsp;&nbsp;
+[<font size = "5" color=#001940>M</font>](DocumentationGuidelines/BrandGlossary/M)&nbsp;&nbsp;
+<font size = "5" color=#001940>N</font>&nbsp;&nbsp;
+[<font size = "5" color=#001940>O</font>](DocumentationGuidelines/BrandGlossary/O)&nbsp;&nbsp;
+[<font size = "5" color=#001940>P</font>](DocumentationGuidelines/BrandGlossary/P)&nbsp;&nbsp;
+<font size = "5" color=#001940>Q</font>&nbsp;&nbsp;
+[<font size = "5" color=#001940>R</font>](DocumentationGuidelines/BrandGlossary/R)&nbsp;&nbsp;
+[<font size = "5" color=#001940>S</font>](DocumentationGuidelines/BrandGlossary/S)&nbsp;&nbsp;
+[<font size = "5" color=#001940>T</font>](DocumentationGuidelines/BrandGlossary/T)&nbsp;&nbsp;
+[<font size = "5" color=#001940>U</font>](DocumentationGuidelines/BrandGlossary/U)&nbsp;&nbsp;
+[<font size = "5" color=#001940>V</font>](DocumentationGuidelines/BrandGlossary/V)&nbsp;&nbsp;
+[<font size = "5" color=#001940>W</font>](DocumentationGuidelines/BrandGlossary/W)&nbsp;&nbsp;
+<font size = "5" color=#001940>X</font>&nbsp;&nbsp;
+<font size = "5" color=#001940>Y</font>&nbsp;&nbsp;
+<font size = "5" color=#001940>Z</font>&nbsp;&nbsp;
+
+##Landscape
+
+The Landscape system enables you to create terrain for your world - mountains, valleys, uneven or sloped ground, openings for caves - and easily modify both its shape and its appearance by using a range of tools.
+
+##Landscape Mode
+
+The Landscape system enables you to create terrain for your world - mountains, valleys, uneven or sloped ground, openings for caves - and easily modify both its shape and its appearance by using a range of tools.
+
+##Level
+
+A Level is a user defined area of gameplay. Levels are created, viewed, and modified by placing, transforming, and editing the properties of the Actors it contains. In the Unreal Editor, each Level is saved as a separate .umap file, which is also why you will sometimes see them referred to as Maps.
+
+##Level Editor
+
+Tools for viewing and editing levels. Container for many Editor functions.
+
+##Level Sequence
+
+The Level Sequence is the “container” for your cinematic scenes and must be created in order to begin working inside the Sequencer Editor.
+
+##Level Streaming
+
+The Level Streaming feature makes it possible to load and unload map files into memory as well as toggle their visibility during play. This makes it possible to have worlds broken up into smaller chunks so that only the relevant parts of the world are taking up resources and being rendered at any point. If done properly, this allows for the creation of very large, seamless levels that can make the player feel as if they are playing within a world that dwarfs them in size.
+
+##Lightmass
+
+Lightmass is the static global illumination system which creates lightmaps with complex light interactions, such as area shadowing and diffuse interreflection from stationary and static lights.
\ No newline at end of file
diff --git a/Engine/Documentation/Source/DocumentationGuidelines/BrandGlossary/M/M.INT.udn b/Engine/Documentation/Source/DocumentationGuidelines/BrandGlossary/M/M.INT.udn
new file mode 100644
--- /dev/null
+++ b/Engine/Documentation/Source/DocumentationGuidelines/BrandGlossary/M/M.INT.udn
@@ -0,0 +1,78 @@
+Availability: Docs
+Crumbs: %ROOT%
+Title: Brand Glossary
+Description: A glossary of key terms, which are used to describe the systems, features, and tools that belong to Unreal Engine 4 (UE4).
+Type: Reference
+Version: 4.16
+Parent:BrandGlossary
+Order:10
+Tags:Brand Glossary
+
+Welcome to the Unreal Engine 4 (UE4) **Brand Glossary**. If you want to create user-facing resources for the UE4 community, this guide 
+helps you gain a better understanding of the terms being used to describe systems, features, and tools belonging to UE4. After reading 
+through this glossary, you'll also know how to properly spell and capitalize the terms we use to describe Unreal Engine.  
+
+<hr>
+  
+[<font size = "5" color=#001940>A</font>](DocumentationGuidelines/BrandGlossary/A)&nbsp;&nbsp;
+[<font size = "5" color=#001940>B</font>](DocumentationGuidelines/BrandGlossary/B)&nbsp;&nbsp;
+[<font size = "5" color=#001940>C</font>](DocumentationGuidelines/BrandGlossary/C)&nbsp;&nbsp;
+[<font size = "5" color=#001940>D</font>](DocumentationGuidelines/BrandGlossary/D)&nbsp;&nbsp;
+[<font size = "5" color=#001940>E</font>](DocumentationGuidelines/BrandGlossary/E)&nbsp;&nbsp;
+[<font size = "5" color=#001940>F</font>](DocumentationGuidelines/BrandGlossary/F)&nbsp;&nbsp;
+[<font size = "5" color=#001940>G</font>](DocumentationGuidelines/BrandGlossary/G)&nbsp;&nbsp;
+[<font size = "5" color=#001940>H</font>](DocumentationGuidelines/BrandGlossary/H)&nbsp;&nbsp;
+<font size = "5" color=#001940>I</font>&nbsp;&nbsp;
+<font size = "5" color=#001940>J</font>&nbsp;&nbsp;
+<font size = "5" color=#001940>K</font>&nbsp;&nbsp;
+[<font size = "5" color=#001940>L</font>](DocumentationGuidelines/BrandGlossary/L)&nbsp;&nbsp;
+<font size = "6" color=#00327F>M</font>&nbsp;&nbsp;
+<font size = "5" color=#001940>N</font>&nbsp;&nbsp;
+[<font size = "5" color=#001940>O</font>](DocumentationGuidelines/BrandGlossary/O)&nbsp;&nbsp;
+[<font size = "5" color=#001940>P</font>](DocumentationGuidelines/BrandGlossary/P)&nbsp;&nbsp;
+<font size = "5" color=#001940>Q</font>&nbsp;&nbsp;
+[<font size = "5" color=#001940>R</font>](DocumentationGuidelines/BrandGlossary/R)&nbsp;&nbsp;
+[<font size = "5" color=#001940>S</font>](DocumentationGuidelines/BrandGlossary/S)&nbsp;&nbsp;
+[<font size = "5" color=#001940>T</font>](DocumentationGuidelines/BrandGlossary/T)&nbsp;&nbsp;
+[<font size = "5" color=#001940>U</font>](DocumentationGuidelines/BrandGlossary/U)&nbsp;&nbsp;
+[<font size = "5" color=#001940>V</font>](DocumentationGuidelines/BrandGlossary/V)&nbsp;&nbsp;
+[<font size = "5" color=#001940>W</font>](DocumentationGuidelines/BrandGlossary/W)&nbsp;&nbsp;
+<font size = "5" color=#001940>X</font>&nbsp;&nbsp;
+<font size = "5" color=#001940>Y</font>&nbsp;&nbsp;
+<font size = "5" color=#001940>Z</font>&nbsp;&nbsp;
+
+##Material
+
+A Material is an asset that can be applied to a mesh to control the visual look of the scene. At a high level, it is easiest to think of a Material as the “paint” that’s applied to an object. But even that can be a bit misleading, since a Material literally defines the type of surface from which your object appears to be made. You can define its color, its shininess, its transparency, and much more.
+
+##Material Editor
+
+The Material Editor provides the ability to create shaders for geometry, using a node-based graph interface.
+
+##Material Instance
+
+Material Instances are child Materials that use parameters to change the appearance of a Material without incurring an expensive recompilation of the Material.
+
+##Matinee Editor
+
+The Matinee animation tool provides the ability to animate the properties of Actors over time, to create either dynamic gameplay or cinematic in-game sequences. The system is based on the use of specialized animation tracks in which you can place keyframes to set the values of certain properties of the Actors in the level. The Matinee Editor is similar to the non-linear editors used for video editing, making it familiar to video professionals.
+
+##Media Framework
+
+The Media Framework enables playback of videos inside Unreal Engine.
+
+##Media Player Window
+
+The Media Player Window provides playback and preview functionality for any Media Source asset.
+
+##Modes
+
+Modes are tools that change the primary behavior of the Level Editor for a specialized task, such as placing new assets into the world, creating geometry brushes and volumes, painting on meshes, generating foliage, and sculpting landscapes.
+
+##Modes Panel
+
+The Modes panel provides access to the Editor Modes’s settings and tools.
+
+##Module
+
+Modules are containers for a collection of related classes that are compiled into a Dynamic-Link Library (DLL).
\ No newline at end of file
diff --git a/Engine/Documentation/Source/DocumentationGuidelines/BrandGlossary/O/O.INT.udn b/Engine/Documentation/Source/DocumentationGuidelines/BrandGlossary/O/O.INT.udn
new file mode 100644
--- /dev/null
+++ b/Engine/Documentation/Source/DocumentationGuidelines/BrandGlossary/O/O.INT.udn
@@ -0,0 +1,50 @@
+Availability: Docs
+Crumbs: %ROOT%
+Title: Brand Glossary
+Description: A glossary of key terms, which are used to describe the systems, features, and tools that belong to Unreal Engine 4 (UE4).
+Type: Reference
+Version: 4.16
+Parent:BrandGlossary
+Order:11
+Tags:Brand Glossary
+
+Welcome to the Unreal Engine 4 (UE4) **Brand Glossary**. If you want to create user-facing resources for the UE4 community, this guide 
+helps you gain a better understanding of the terms being used to describe systems, features, and tools belonging to UE4. After reading 
+through this glossary, you'll also know how to properly spell and capitalize the terms we use to describe Unreal Engine.  
+
+<hr>
+  
+[<font size = "5" color=#001940>A</font>](DocumentationGuidelines/BrandGlossary/A)&nbsp;&nbsp;
+[<font size = "5" color=#001940>B</font>](DocumentationGuidelines/BrandGlossary/B)&nbsp;&nbsp;
+[<font size = "5" color=#001940>C</font>](DocumentationGuidelines/BrandGlossary/C)&nbsp;&nbsp;
+[<font size = "5" color=#001940>D</font>](DocumentationGuidelines/BrandGlossary/D)&nbsp;&nbsp;
+[<font size = "5" color=#001940>E</font>](DocumentationGuidelines/BrandGlossary/E)&nbsp;&nbsp;
+[<font size = "5" color=#001940>F</font>](DocumentationGuidelines/BrandGlossary/F)&nbsp;&nbsp;
+[<font size = "5" color=#001940>G</font>](DocumentationGuidelines/BrandGlossary/G)&nbsp;&nbsp;
+[<font size = "5" color=#001940>H</font>](DocumentationGuidelines/BrandGlossary/H)&nbsp;&nbsp;
+<font size = "5" color=#001940>I</font>&nbsp;&nbsp;
+<font size = "5" color=#001940>J</font>&nbsp;&nbsp;
+<font size = "5" color=#001940>K</font>&nbsp;&nbsp;
+[<font size = "5" color=#001940>L</font>](DocumentationGuidelines/BrandGlossary/L)&nbsp;&nbsp;
+[<font size = "5" color=#001940>M</font>](DocumentationGuidelines/BrandGlossary/M)&nbsp;&nbsp;
+<font size = "5" color=#001940>N</font>&nbsp;&nbsp;
+<font size = "6" color=#00327F>O</font>&nbsp;&nbsp;
+[<font size = "5" color=#001940>P</font>](DocumentationGuidelines/BrandGlossary/P)&nbsp;&nbsp;
+<font size = "5" color=#001940>Q</font>&nbsp;&nbsp;
+[<font size = "5" color=#001940>R</font>](DocumentationGuidelines/BrandGlossary/R)&nbsp;&nbsp;
+[<font size = "5" color=#001940>S</font>](DocumentationGuidelines/BrandGlossary/S)&nbsp;&nbsp;
+[<font size = "5" color=#001940>T</font>](DocumentationGuidelines/BrandGlossary/T)&nbsp;&nbsp;
+[<font size = "5" color=#001940>U</font>](DocumentationGuidelines/BrandGlossary/U)&nbsp;&nbsp;
+[<font size = "5" color=#001940>V</font>](DocumentationGuidelines/BrandGlossary/V)&nbsp;&nbsp;
+[<font size = "5" color=#001940>W</font>](DocumentationGuidelines/BrandGlossary/W)&nbsp;&nbsp;
+<font size = "5" color=#001940>X</font>&nbsp;&nbsp;
+<font size = "5" color=#001940>Y</font>&nbsp;&nbsp;
+<font size = "5" color=#001940>Z</font>&nbsp;&nbsp;
+
+##Object
+
+The base building blocks in the Unreal Engine are called Objects, containing a lot of the essential “under the hood” functionality for your game assets. Just about everything in Unreal Engine 4 inherits (or gets some functionality) from an Object. In C++, UObject is the base class of all objects; it implements features such as garbage collection, metadata (UProperty) support for exposing variables to the Unreal Editor, and serialization for loading and saving.
+
+##Online Subsystem
+
+The Online subsystem and its interfaces exist to provide a clear abstraction to common online functionality across the available set of platforms in a given environment. Platforms in this context would refer to Steam, Xbox Live, Facebook, etc. Portability is one of the main goals of this subsystem.
\ No newline at end of file
diff --git a/Engine/Documentation/Source/DocumentationGuidelines/BrandGlossary/P/P.INT.udn b/Engine/Documentation/Source/DocumentationGuidelines/BrandGlossary/P/P.INT.udn
new file mode 100644
--- /dev/null
+++ b/Engine/Documentation/Source/DocumentationGuidelines/BrandGlossary/P/P.INT.udn
@@ -0,0 +1,118 @@
+Availability: Docs
+Crumbs: %ROOT%
+Title: Brand Glossary
+Description: A glossary of key terms, which are used to describe the systems, features, and tools that belong to Unreal Engine 4 (UE4).
+Type: Reference
+Version: 4.16
+Parent:BrandGlossary
+Order:12
+Tags:Brand Glossary
+
+Welcome to the Unreal Engine 4 (UE4) **Brand Glossary**. If you want to create user-facing resources for the UE4 community, this guide 
+helps you gain a better understanding of the terms being used to describe systems, features, and tools belonging to UE4. After reading 
+through this glossary, you'll also know how to properly spell and capitalize the terms we use to describe Unreal Engine.  
+
+<hr>
+  
+[<font size = "5" color=#001940>A</font>](DocumentationGuidelines/BrandGlossary/A)&nbsp;&nbsp;
+[<font size = "5" color=#001940>B</font>](DocumentationGuidelines/BrandGlossary/B)&nbsp;&nbsp;
+[<font size = "5" color=#001940>C</font>](DocumentationGuidelines/BrandGlossary/C)&nbsp;&nbsp;
+[<font size = "5" color=#001940>D</font>](DocumentationGuidelines/BrandGlossary/D)&nbsp;&nbsp;
+[<font size = "5" color=#001940>E</font>](DocumentationGuidelines/BrandGlossary/E)&nbsp;&nbsp;
+[<font size = "5" color=#001940>F</font>](DocumentationGuidelines/BrandGlossary/F)&nbsp;&nbsp;
+[<font size = "5" color=#001940>G</font>](DocumentationGuidelines/BrandGlossary/G)&nbsp;&nbsp;
+[<font size = "5" color=#001940>H</font>](DocumentationGuidelines/BrandGlossary/H)&nbsp;&nbsp;
+<font size = "5" color=#001940>I</font>&nbsp;&nbsp;
+<font size = "5" color=#001940>J</font>&nbsp;&nbsp;
+<font size = "5" color=#001940>K</font>&nbsp;&nbsp;
+[<font size = "5" color=#001940>L</font>](DocumentationGuidelines/BrandGlossary/L)&nbsp;&nbsp;
+[<font size = "5" color=#001940>M</font>](DocumentationGuidelines/BrandGlossary/M)&nbsp;&nbsp;
+<font size = "5" color=#001940>N</font>&nbsp;&nbsp;
+[<font size = "5" color=#001940>O</font>](DocumentationGuidelines/BrandGlossary/O)&nbsp;&nbsp;
+<font size = "6" color=#00327F>P</font>&nbsp;&nbsp;
+<font size = "5" color=#001940>Q</font>&nbsp;&nbsp;
+[<font size = "5" color=#001940>R</font>](DocumentationGuidelines/BrandGlossary/R)&nbsp;&nbsp;
+[<font size = "5" color=#001940>S</font>](DocumentationGuidelines/BrandGlossary/S)&nbsp;&nbsp;
+[<font size = "5" color=#001940>T</font>](DocumentationGuidelines/BrandGlossary/T)&nbsp;&nbsp;
+[<font size = "5" color=#001940>U</font>](DocumentationGuidelines/BrandGlossary/U)&nbsp;&nbsp;
+[<font size = "5" color=#001940>V</font>](DocumentationGuidelines/BrandGlossary/V)&nbsp;&nbsp;
+[<font size = "5" color=#001940>W</font>](DocumentationGuidelines/BrandGlossary/W)&nbsp;&nbsp;
+<font size = "5" color=#001940>X</font>&nbsp;&nbsp;
+<font size = "5" color=#001940>Y</font>&nbsp;&nbsp;
+<font size = "5" color=#001940>Z</font>&nbsp;&nbsp;
+
+##Packaging
+
+The purpose of packaging is to test your full game instead of a single map, or to prepare your game for submission/distribution.
+
+##Paper 2D
+
+Paper 2D is a sprite-based system for creating 2D and 2D/3D hybrid games entirely within Unreal Engine.
+
+##Paper 2D Flipbook Editor
+
+The Flipbook Editor enables setup and editing of individual Flipbook assets, where you can edit the sequential playability of your sprites, tune or tweak the duration in which a sprite is displayed, or determine the overall playback speed of the Flipbook animation through the Frames Per Second option.
+
+##Paper 2D Sprite Editor
+
+The Paper 2D Sprite Editor enables setup and editing of individual Sprite assets.
+
+##Particle System Component
+
+A ParticleSystemComponent allows you to add a particle emitter as a sub-object to another object. 
+
+##Pawn
+
+A Pawn is an Actor that can be an “agent” within the world. Pawns can be possessed by a Controller, they are set up to easily accept input, and they can do a variety of player-like things. Note that a Pawn is not assumed to be humanoid.
+
+##Persona Animation Editor
+
+The Persona Animation Editor is the animation editing toolset which provides a robust system for editing Skeletons, Skeletal Meshes, Animation Blueprints, and various other animation assets.
+
+##Physics Asset Tool (PhAT)
+
+The Unreal Physics Asset Tool (PhAT) is an integrated editor in the Unreal Editor that is designed specifically to manipulate Physics Assets for Skeletal Meshes.
+
+##Play in Editor (PIE)
+
+Play In Editor allows you to play the current level directly from the editor, so that you can test gameplay functionality, including player controls and level events triggered by players’ actions.
+
+##PlayerController
+
+A PlayerController is the interface between the Pawn and the human player controlling it. The PlayerController essentially represents the human player’s will.
+
+##PlayerState
+
+A PlayerState is the state of a participant in the game, such as a human player or a bot that is simulating a player.
+
+##Plugin
+
+Plugins enable you to add entirely new features and to modify built-in functionality without modifying the engine code directly, such as adding new menu items and tool bar commands to the editor, or even adding entirely new features and editor sub-modes.
+
+##Post Process Effect
+
+Post Process Effects enable artists and designers to tweak the overall look and feel of the scene by applying effects including bloom (HDR blooming effect on bright objects), ambient occlusion, and tone mapping to the rendered scene before displaying it in the viewport.
+
+##Procedural Foliage Tool
+
+The Procedural Foliage Tool is used for simulating huge forests that are filled with lots of different kinds of trees and bushes. This tool works by simulating how a forest grows over the years using Steps to define years of growth. At each Step in the simulation, newly created virtual seeds are cast into the level to act as spawn locations for other Foliage Type Actors. By tweaking the various Foliage Type options like Spread, Size, and Priority Settings, you will affect how those Foliage Types spread and compete for spawn locations through the level that the Procedural Foliage Spawner is placed in. 
+
+##Project
+
+A Project is a self-contained unit that holds all of the content and code that make-up an individual game and coincides with a set of directories on your disk.
+
+##Project Browser
+
+The Project Browser provides an interface for creating new projects, opening your existing projects, or opening sample content like sample games and Showcases.
+
+##Project Launcher
+
+The Project Launcher tab provides a graphical front-end interface to build, cook, deploy, and launch your game.
+
+##Project Settings
+
+The Project Settings editor provides access to configuration options that specify information about your project, as well as defining how the engine behaves when running the project.
+
+##Property Matrix
+
+The Property Matrix enables easy bulk editing and value comparison for large numbers of Objects or Actors by displaying a configurable set of properties for a collection of objects as columns in a table view that can be sorted on any column.
\ No newline at end of file
diff --git a/Engine/Documentation/Source/DocumentationGuidelines/BrandGlossary/R/R.INT.udn b/Engine/Documentation/Source/DocumentationGuidelines/BrandGlossary/R/R.INT.udn
new file mode 100644
--- /dev/null
+++ b/Engine/Documentation/Source/DocumentationGuidelines/BrandGlossary/R/R.INT.udn
@@ -0,0 +1,50 @@
+Availability: Docs
+Crumbs: %ROOT%
+Title: Brand Glossary
+Description: A glossary of key terms, which are used to describe the systems, features, and tools that belong to Unreal Engine 4 (UE4).
+Type: Reference
+Version: 4.16
+Parent:BrandGlossary
+Order:13
+Tags:Brand Glossary
+
+Welcome to the Unreal Engine 4 (UE4) **Brand Glossary**. If you want to create user-facing resources for the UE4 community, this guide 
+helps you gain a better understanding of the terms being used to describe systems, features, and tools belonging to UE4. After reading 
+through this glossary, you'll also know how to properly spell and capitalize the terms we use to describe Unreal Engine.  
+
+<hr>
+  
+[<font size = "5" color=#001940>A</font>](DocumentationGuidelines/BrandGlossary/A)&nbsp;&nbsp;
+[<font size = "5" color=#001940>B</font>](DocumentationGuidelines/BrandGlossary/B)&nbsp;&nbsp;
+[<font size = "5" color=#001940>C</font>](DocumentationGuidelines/BrandGlossary/C)&nbsp;&nbsp;
+[<font size = "5" color=#001940>D</font>](DocumentationGuidelines/BrandGlossary/D)&nbsp;&nbsp;
+[<font size = "5" color=#001940>E</font>](DocumentationGuidelines/BrandGlossary/E)&nbsp;&nbsp;
+[<font size = "5" color=#001940>F</font>](DocumentationGuidelines/BrandGlossary/F)&nbsp;&nbsp;
+[<font size = "5" color=#001940>G</font>](DocumentationGuidelines/BrandGlossary/G)&nbsp;&nbsp;
+[<font size = "5" color=#001940>H</font>](DocumentationGuidelines/BrandGlossary/H)&nbsp;&nbsp;
+<font size = "5" color=#001940>I</font>&nbsp;&nbsp;
+<font size = "5" color=#001940>J</font>&nbsp;&nbsp;
+<font size = "5" color=#001940>K</font>&nbsp;&nbsp;
+[<font size = "5" color=#001940>L</font>](DocumentationGuidelines/BrandGlossary/L)&nbsp;&nbsp;
+[<font size = "5" color=#001940>M</font>](DocumentationGuidelines/BrandGlossary/M)&nbsp;&nbsp;
+<font size = "5" color=#001940>N</font>&nbsp;&nbsp;
+[<font size = "5" color=#001940>O</font>](DocumentationGuidelines/BrandGlossary/O)&nbsp;&nbsp;
+[<font size = "5" color=#001940>P</font>](DocumentationGuidelines/BrandGlossary/P)&nbsp;&nbsp;
+<font size = "5" color=#001940>Q</font>&nbsp;&nbsp;
+<font size = "6" color=#00327F>R</font>&nbsp;&nbsp;
+[<font size = "5" color=#001940>S</font>](DocumentationGuidelines/BrandGlossary/S)&nbsp;&nbsp;
+[<font size = "5" color=#001940>T</font>](DocumentationGuidelines/BrandGlossary/T)&nbsp;&nbsp;
+[<font size = "5" color=#001940>U</font>](DocumentationGuidelines/BrandGlossary/U)&nbsp;&nbsp;
+[<font size = "5" color=#001940>V</font>](DocumentationGuidelines/BrandGlossary/V)&nbsp;&nbsp;
+[<font size = "5" color=#001940>W</font>](DocumentationGuidelines/BrandGlossary/W)&nbsp;&nbsp;
+<font size = "5" color=#001940>X</font>&nbsp;&nbsp;
+<font size = "5" color=#001940>Y</font>&nbsp;&nbsp;
+<font size = "5" color=#001940>Z</font>&nbsp;&nbsp;
+
+##Redirector
+
+A Redirector provides a reference to the new location of an Asset when it is moved or renamed so that content not loaded at the time of the operation can find the asset.
+
+##Replication
+
+Replication is used to keep the Actors within the world in sync when dealing with networked multiplayer games. Property values and function calls can both be replicated, allowing for complete control over the state of the game on all clients.
\ No newline at end of file
diff --git a/Engine/Documentation/Source/DocumentationGuidelines/BrandGlossary/S/S.INT.udn b/Engine/Documentation/Source/DocumentationGuidelines/BrandGlossary/S/S.INT.udn
new file mode 100644
--- /dev/null
+++ b/Engine/Documentation/Source/DocumentationGuidelines/BrandGlossary/S/S.INT.udn
@@ -0,0 +1,90 @@
+Availability: Docs
+Crumbs: %ROOT%
+Title: Brand Glossary
+Description: A glossary of key terms, which are used to describe the systems, features, and tools that belong to Unreal Engine 4 (UE4).
+Type: Reference
+Version: 4.16
+Parent:BrandGlossary
+Order:14
+Tags:Brand Glossary
+
+Welcome to the Unreal Engine 4 (UE4) **Brand Glossary**. If you want to create user-facing resources for the UE4 community, this guide 
+helps you gain a better understanding of the terms being used to describe systems, features, and tools belonging to UE4. After reading 
+through this glossary, you'll also know how to properly spell and capitalize the terms we use to describe Unreal Engine.  
+
+<hr>
+  
+[<font size = "5" color=#001940>A</font>](DocumentationGuidelines/BrandGlossary/A)&nbsp;&nbsp;
+[<font size = "5" color=#001940>B</font>](DocumentationGuidelines/BrandGlossary/B)&nbsp;&nbsp;
+[<font size = "5" color=#001940>C</font>](DocumentationGuidelines/BrandGlossary/C)&nbsp;&nbsp;
+[<font size = "5" color=#001940>D</font>](DocumentationGuidelines/BrandGlossary/D)&nbsp;&nbsp;
+[<font size = "5" color=#001940>E</font>](DocumentationGuidelines/BrandGlossary/E)&nbsp;&nbsp;
+[<font size = "5" color=#001940>F</font>](DocumentationGuidelines/BrandGlossary/F)&nbsp;&nbsp;
+[<font size = "5" color=#001940>G</font>](DocumentationGuidelines/BrandGlossary/G)&nbsp;&nbsp;
+[<font size = "5" color=#001940>H</font>](DocumentationGuidelines/BrandGlossary/H)&nbsp;&nbsp;
+<font size = "5" color=#001940>I</font>&nbsp;&nbsp;
+<font size = "5" color=#001940>J</font>&nbsp;&nbsp;
+<font size = "5" color=#001940>K</font>&nbsp;&nbsp;
+[<font size = "5" color=#001940>L</font>](DocumentationGuidelines/BrandGlossary/L)&nbsp;&nbsp;
+[<font size = "5" color=#001940>M</font>](DocumentationGuidelines/BrandGlossary/M)&nbsp;&nbsp;
+<font size = "5" color=#001940>N</font>&nbsp;&nbsp;
+[<font size = "5" color=#001940>O</font>](DocumentationGuidelines/BrandGlossary/O)&nbsp;&nbsp;
+[<font size = "5" color=#001940>P</font>](DocumentationGuidelines/BrandGlossary/P)&nbsp;&nbsp;
+<font size = "5" color=#001940>Q</font>&nbsp;&nbsp;
+[<font size = "5" color=#001940>R</font>](DocumentationGuidelines/BrandGlossary/R)&nbsp;&nbsp;
+<font size = "6" color=#00327F>S</font>&nbsp;&nbsp;
+[<font size = "5" color=#001940>T</font>](DocumentationGuidelines/BrandGlossary/T)&nbsp;&nbsp;
+[<font size = "5" color=#001940>U</font>](DocumentationGuidelines/BrandGlossary/U)&nbsp;&nbsp;
+[<font size = "5" color=#001940>V</font>](DocumentationGuidelines/BrandGlossary/V)&nbsp;&nbsp;
+[<font size = "5" color=#001940>W</font>](DocumentationGuidelines/BrandGlossary/W)&nbsp;&nbsp;
+<font size = "5" color=#001940>X</font>&nbsp;&nbsp;
+<font size = "5" color=#001940>Y</font>&nbsp;&nbsp;
+<font size = "5" color=#001940>Z</font>&nbsp;&nbsp;
+
+##Sequencer
+
+Unreal Engine’s cinematic toolset that provides director-level control over cut scenes, dynamic gameplay sequences, and movies.
+
+##Simulate In Editor
+
+Simulate In Editor enables you to run game logic in the editor viewport so you can inspect, edit, and interact with Actors in the world as the game runs.
+
+##Skeletal Mesh
+
+Skeletal Meshes are built up of two parts; a set of polygons composed to make up the surface of the Skeletal Mesh, and a hierarchical set of interconnected bones that can be used to animate the polygons
+
+##Skeleton
+
+A Skeleton is an asset that holds bone and hierarchy information for a specific type of character or Skeletal Mesh.
+
+##Slate UI Framework
+
+Slate is a completely custom and platform agnostic user interface framework that is designed to make building the user interfaces for tools and applications, such as Unreal Editor, or in-game user interfaces fun and efficient. It combines a declarative syntax with the ability to design, layout, and style components that allows for easily creating and iterating on UIs.
+
+##Sound Cue
+
+Asset that defines the playback behavior of a sound.
+
+##Sound Cue Editor
+
+Node-based tool used to edit Sound Cues.
+
+##Spawn
+
+Creating a new instance of an Actor, either programmatically at run-time, or manually in the Level Editor at design-time.
+
+##Static Mesh
+
+A Static Mesh is a piece of geometry that consists of a static set of polygons and are the basic unit used to create world geometry for levels in Unreal Engine 4.
+
+##Static Mesh Actor
+
+An Actor whose visual representation is a Static Mesh.
+
+##Static Mesh Components
+
+A Component that adds a Static Mesh to an Actor.
+
+##Static Mesh Editor
+
+The Static Mesh Editor provides tools for previewing the geometry, collision, and levels of detail of a Static Mesh asset, as well as editing properties, applying Materials, and setting up collision geometry.
\ No newline at end of file
diff --git a/Engine/Documentation/Source/DocumentationGuidelines/BrandGlossary/T/T.INT.udn b/Engine/Documentation/Source/DocumentationGuidelines/BrandGlossary/T/T.INT.udn
new file mode 100644
--- /dev/null
+++ b/Engine/Documentation/Source/DocumentationGuidelines/BrandGlossary/T/T.INT.udn
@@ -0,0 +1,46 @@
+Availability: Docs
+Crumbs: %ROOT%
+Title: Brand Glossary
+Description: A glossary of key terms, which are used to describe the systems, features, and tools that belong to Unreal Engine 4 (UE4).
+Type: Reference
+Version: 4.16
+Parent:BrandGlossary
+Order:15
+Tags:Brand Glossary
+
+Welcome to the Unreal Engine 4 (UE4) **Brand Glossary**. If you want to create user-facing resources for the UE4 community, this guide 
+helps you gain a better understanding of the terms being used to describe systems, features, and tools belonging to UE4. After reading 
+through this glossary, you'll also know how to properly spell and capitalize the terms we use to describe Unreal Engine.  
+
+<hr>
+  
+[<font size = "5" color=#001940>A</font>](DocumentationGuidelines/BrandGlossary/A)&nbsp;&nbsp;
+[<font size = "5" color=#001940>B</font>](DocumentationGuidelines/BrandGlossary/B)&nbsp;&nbsp;
+[<font size = "5" color=#001940>C</font>](DocumentationGuidelines/BrandGlossary/C)&nbsp;&nbsp;
+[<font size = "5" color=#001940>D</font>](DocumentationGuidelines/BrandGlossary/D)&nbsp;&nbsp;
+[<font size = "5" color=#001940>E</font>](DocumentationGuidelines/BrandGlossary/E)&nbsp;&nbsp;
+[<font size = "5" color=#001940>F</font>](DocumentationGuidelines/BrandGlossary/F)&nbsp;&nbsp;
+[<font size = "5" color=#001940>G</font>](DocumentationGuidelines/BrandGlossary/G)&nbsp;&nbsp;
+[<font size = "5" color=#001940>H</font>](DocumentationGuidelines/BrandGlossary/H)&nbsp;&nbsp;
+<font size = "5" color=#001940>I</font>&nbsp;&nbsp;
+<font size = "5" color=#001940>J</font>&nbsp;&nbsp;
+<font size = "5" color=#001940>K</font>&nbsp;&nbsp;
+[<font size = "5" color=#001940>L</font>](DocumentationGuidelines/BrandGlossary/L)&nbsp;&nbsp;
+[<font size = "5" color=#001940>M</font>](DocumentationGuidelines/BrandGlossary/M)&nbsp;&nbsp;
+<font size = "5" color=#001940>N</font>&nbsp;&nbsp;
+[<font size = "5" color=#001940>O</font>](DocumentationGuidelines/BrandGlossary/O)&nbsp;&nbsp;
+[<font size = "5" color=#001940>P</font>](DocumentationGuidelines/BrandGlossary/P)&nbsp;&nbsp;
+<font size = "5" color=#001940>Q</font>&nbsp;&nbsp;
+[<font size = "5" color=#001940>R</font>](DocumentationGuidelines/BrandGlossary/R)&nbsp;&nbsp;
+[<font size = "5" color=#001940>S</font>](DocumentationGuidelines/BrandGlossary/S)&nbsp;&nbsp;
+<font size = "6" color=#00327F>T</font>&nbsp;&nbsp;
+[<font size = "5" color=#001940>U</font>](DocumentationGuidelines/BrandGlossary/U)&nbsp;&nbsp;
+[<font size = "5" color=#001940>V</font>](DocumentationGuidelines/BrandGlossary/V)&nbsp;&nbsp;
+[<font size = "5" color=#001940>W</font>](DocumentationGuidelines/BrandGlossary/W)&nbsp;&nbsp;
+<font size = "5" color=#001940>X</font>&nbsp;&nbsp;
+<font size = "5" color=#001940>Y</font>&nbsp;&nbsp;
+<font size = "5" color=#001940>Z</font>&nbsp;&nbsp;
+
+##Texture
+
+Textures are images that are mapped to the surfaces that a Material is applied to.
\ No newline at end of file
diff --git a/Engine/Documentation/Source/DocumentationGuidelines/BrandGlossary/U/U.INT.udn b/Engine/Documentation/Source/DocumentationGuidelines/BrandGlossary/U/U.INT.udn
new file mode 100644
--- /dev/null
+++ b/Engine/Documentation/Source/DocumentationGuidelines/BrandGlossary/U/U.INT.udn
@@ -0,0 +1,62 @@
+Availability: Docs
+Crumbs: %ROOT%
+Title: Brand Glossary
+Description: A glossary of key terms, which are used to describe the systems, features, and tools that belong to Unreal Engine 4 (UE4).
+Type: Reference
+Version: 4.16
+Parent:BrandGlossary
+Order:16
+Tags:Brand Glossary
+
+Welcome to the Unreal Engine 4 (UE4) **Brand Glossary**. If you want to create user-facing resources for the UE4 community, this guide 
+helps you gain a better understanding of the terms being used to describe systems, features, and tools belonging to UE4. After reading 
+through this glossary, you'll also know how to properly spell and capitalize the terms we use to describe Unreal Engine.  
+
+<hr>
+  
+[<font size = "5" color=#001940>A</font>](DocumentationGuidelines/BrandGlossary/A)&nbsp;&nbsp;
+[<font size = "5" color=#001940>B</font>](DocumentationGuidelines/BrandGlossary/B)&nbsp;&nbsp;
+[<font size = "5" color=#001940>C</font>](DocumentationGuidelines/BrandGlossary/C)&nbsp;&nbsp;
+[<font size = "5" color=#001940>D</font>](DocumentationGuidelines/BrandGlossary/D)&nbsp;&nbsp;
+[<font size = "5" color=#001940>E</font>](DocumentationGuidelines/BrandGlossary/E)&nbsp;&nbsp;
+[<font size = "5" color=#001940>F</font>](DocumentationGuidelines/BrandGlossary/F)&nbsp;&nbsp;
+[<font size = "5" color=#001940>G</font>](DocumentationGuidelines/BrandGlossary/G)&nbsp;&nbsp;
+[<font size = "5" color=#001940>H</font>](DocumentationGuidelines/BrandGlossary/H)&nbsp;&nbsp;
+<font size = "5" color=#001940>I</font>&nbsp;&nbsp;
+<font size = "5" color=#001940>J</font>&nbsp;&nbsp;
+<font size = "5" color=#001940>K</font>&nbsp;&nbsp;
+[<font size = "5" color=#001940>L</font>](DocumentationGuidelines/BrandGlossary/L)&nbsp;&nbsp;
+[<font size = "5" color=#001940>M</font>](DocumentationGuidelines/BrandGlossary/M)&nbsp;&nbsp;
+<font size = "5" color=#001940>N</font>&nbsp;&nbsp;
+[<font size = "5" color=#001940>O</font>](DocumentationGuidelines/BrandGlossary/O)&nbsp;&nbsp;
+[<font size = "5" color=#001940>P</font>](DocumentationGuidelines/BrandGlossary/P)&nbsp;&nbsp;
+<font size = "5" color=#001940>Q</font>&nbsp;&nbsp;
+[<font size = "5" color=#001940>R</font>](DocumentationGuidelines/BrandGlossary/R)&nbsp;&nbsp;
+[<font size = "5" color=#001940>S</font>](DocumentationGuidelines/BrandGlossary/S)&nbsp;&nbsp;
+[<font size = "5" color=#001940>T</font>](DocumentationGuidelines/BrandGlossary/T)&nbsp;&nbsp;
+<font size = "6" color=#00327F>U</font>&nbsp;&nbsp;
+[<font size = "5" color=#001940>V</font>](DocumentationGuidelines/BrandGlossary/V)&nbsp;&nbsp;
+[<font size = "5" color=#001940>W</font>](DocumentationGuidelines/BrandGlossary/W)&nbsp;&nbsp;
+<font size = "5" color=#001940>X</font>&nbsp;&nbsp;
+<font size = "5" color=#001940>Y</font>&nbsp;&nbsp;
+<font size = "5" color=#001940>Z</font>&nbsp;&nbsp;
+
+##Unreal Build Tool (UBT)
+
+UnrealBuildTool (UBT) is a custom tool that manages the process of building the Unreal Engine source code across the numerous possible build configurations.
+
+##Unreal Editor
+
+Unreal Editor is the complete suite of tools provided with Unreal Engine that enables developers to build levels, import content, create effects, and much more.
+
+##Unreal Header Tool (UHT)
+
+UnrealHeaderTool (UHT) is a custom parsing and code-generation tool that supports the UObject system.
+
+##UMG UI Designer
+
+Unreal Motion Graphics UI Designer (UMG) is a visual UI authoring tool that can be used to create UI elements, such as in-game HUDs, menus or other interface related graphics you wish to present to your users.
+
+##Unreal Frontend (UFE)
+
+Unreal Frontend (UFE) is a tool intended to simplify and speed up daily video game development and testing tasks, such as preparing game builds, deploying them to a device, and launching them. UFE is designed to be the central interface for all game deployment, profiling, and testing tasks.
\ No newline at end of file
diff --git a/Engine/Documentation/Source/DocumentationGuidelines/BrandGlossary/V/V.INT.udn b/Engine/Documentation/Source/DocumentationGuidelines/BrandGlossary/V/V.INT.udn
new file mode 100644
--- /dev/null
+++ b/Engine/Documentation/Source/DocumentationGuidelines/BrandGlossary/V/V.INT.udn
@@ -0,0 +1,54 @@
+Availability: Docs
+Crumbs: %ROOT%
+Title: Brand Glossary
+Description: A glossary of key terms, which are used to describe the systems, features, and tools that belong to Unreal Engine 4 (UE4).
+Type: Reference
+Version: 4.16
+Parent:BrandGlossary
+Order:17
+Tags:Brand Glossary
+
+Welcome to the Unreal Engine 4 (UE4) **Brand Glossary**. If you want to create user-facing resources for the UE4 community, this guide 
+helps you gain a better understanding of the terms being used to describe systems, features, and tools belonging to UE4. After reading 
+through this glossary, you'll also know how to properly spell and capitalize the terms we use to describe Unreal Engine.  
+
+<hr>
+  
+[<font size = "5" color=#001940>A</font>](DocumentationGuidelines/BrandGlossary/A)&nbsp;&nbsp;
+[<font size = "5" color=#001940>B</font>](DocumentationGuidelines/BrandGlossary/B)&nbsp;&nbsp;
+[<font size = "5" color=#001940>C</font>](DocumentationGuidelines/BrandGlossary/C)&nbsp;&nbsp;
+[<font size = "5" color=#001940>D</font>](DocumentationGuidelines/BrandGlossary/D)&nbsp;&nbsp;
+[<font size = "5" color=#001940>E</font>](DocumentationGuidelines/BrandGlossary/E)&nbsp;&nbsp;
+[<font size = "5" color=#001940>F</font>](DocumentationGuidelines/BrandGlossary/F)&nbsp;&nbsp;
+[<font size = "5" color=#001940>G</font>](DocumentationGuidelines/BrandGlossary/G)&nbsp;&nbsp;
+[<font size = "5" color=#001940>H</font>](DocumentationGuidelines/BrandGlossary/H)&nbsp;&nbsp;
+<font size = "5" color=#001940>I</font>&nbsp;&nbsp;
+<font size = "5" color=#001940>J</font>&nbsp;&nbsp;
+<font size = "5" color=#001940>K</font>&nbsp;&nbsp;
+[<font size = "5" color=#001940>L</font>](DocumentationGuidelines/BrandGlossary/L)&nbsp;&nbsp;
+[<font size = "5" color=#001940>M</font>](DocumentationGuidelines/BrandGlossary/M)&nbsp;&nbsp;
+<font size = "5" color=#001940>N</font>&nbsp;&nbsp;
+[<font size = "5" color=#001940>O</font>](DocumentationGuidelines/BrandGlossary/O)&nbsp;&nbsp;
+[<font size = "5" color=#001940>P</font>](DocumentationGuidelines/BrandGlossary/P)&nbsp;&nbsp;
+<font size = "5" color=#001940>Q</font>&nbsp;&nbsp;
+[<font size = "5" color=#001940>R</font>](DocumentationGuidelines/BrandGlossary/R)&nbsp;&nbsp;
+[<font size = "5" color=#001940>S</font>](DocumentationGuidelines/BrandGlossary/S)&nbsp;&nbsp;
+[<font size = "5" color=#001940>T</font>](DocumentationGuidelines/BrandGlossary/T)&nbsp;&nbsp;
+[<font size = "5" color=#001940>U</font>](DocumentationGuidelines/BrandGlossary/U)&nbsp;&nbsp;
+<font size = "6" color=#00327F>V</font>&nbsp;&nbsp;
+[<font size = "5" color=#001940>W</font>](DocumentationGuidelines/BrandGlossary/W)&nbsp;&nbsp;
+<font size = "5" color=#001940>X</font>&nbsp;&nbsp;
+<font size = "5" color=#001940>Y</font>&nbsp;&nbsp;
+<font size = "5" color=#001940>Z</font>&nbsp;&nbsp;
+
+##Viewport
+
+Viewports are your window into the worlds you create in Unreal. They can be navigated just as you would in a game, or can be used in a more schematic design sense as you would for an architectural blueprint
+
+##VR Editor
+
+The VR Editor enables you to design and build worlds in a virtual reality environment using the full capabilities of the editor toolset combined with interaction models designed specifically for VR world building.
+
+##Volume
+
+Volumes are 3-dimensional areas within your levels, each of which serve a specific purpose.
\ No newline at end of file
diff --git a/Engine/Documentation/Source/DocumentationGuidelines/BrandGlossary/W/W.INT.udn b/Engine/Documentation/Source/DocumentationGuidelines/BrandGlossary/W/W.INT.udn
new file mode 100644
--- /dev/null
+++ b/Engine/Documentation/Source/DocumentationGuidelines/BrandGlossary/W/W.INT.udn
@@ -0,0 +1,58 @@
+Availability: Docs
+Crumbs: %ROOT%
+Title: Brand Glossary
+Description: A glossary of key terms, which are used to describe the systems, features, and tools that belong to Unreal Engine 4 (UE4).
+Type: Reference
+Version: 4.16
+Parent:BrandGlossary
+Order:18
+Tags:Brand Glossary
+
+Welcome to the Unreal Engine 4 (UE4) **Brand Glossary**. If you want to create user-facing resources for the UE4 community, this guide 
+helps you gain a better understanding of the terms being used to describe systems, features, and tools belonging to UE4. After reading 
+through this glossary, you'll also know how to properly spell and capitalize the terms we use to describe Unreal Engine.  
+
+<hr>
+  
+[<font size = "5" color=#001940>A</font>](DocumentationGuidelines/BrandGlossary/A)&nbsp;&nbsp;
+[<font size = "5" color=#001940>B</font>](DocumentationGuidelines/BrandGlossary/B)&nbsp;&nbsp;
+[<font size = "5" color=#001940>C</font>](DocumentationGuidelines/BrandGlossary/C)&nbsp;&nbsp;
+[<font size = "5" color=#001940>D</font>](DocumentationGuidelines/BrandGlossary/D)&nbsp;&nbsp;
+[<font size = "5" color=#001940>E</font>](DocumentationGuidelines/BrandGlossary/E)&nbsp;&nbsp;
+[<font size = "5" color=#001940>F</font>](DocumentationGuidelines/BrandGlossary/F)&nbsp;&nbsp;
+[<font size = "5" color=#001940>G</font>](DocumentationGuidelines/BrandGlossary/G)&nbsp;&nbsp;
+[<font size = "5" color=#001940>H</font>](DocumentationGuidelines/BrandGlossary/H)&nbsp;&nbsp;
+<font size = "5" color=#001940>I</font>&nbsp;&nbsp;
+<font size = "5" color=#001940>J</font>&nbsp;&nbsp;
+<font size = "5" color=#001940>K</font>&nbsp;&nbsp;
+[<font size = "5" color=#001940>L</font>](DocumentationGuidelines/BrandGlossary/L)&nbsp;&nbsp;
+[<font size = "5" color=#001940>M</font>](DocumentationGuidelines/BrandGlossary/M)&nbsp;&nbsp;
+<font size = "5" color=#001940>N</font>&nbsp;&nbsp;
+[<font size = "5" color=#001940>O</font>](DocumentationGuidelines/BrandGlossary/O)&nbsp;&nbsp;
+[<font size = "5" color=#001940>P</font>](DocumentationGuidelines/BrandGlossary/P)&nbsp;&nbsp;
+<font size = "5" color=#001940>Q</font>&nbsp;&nbsp;
+[<font size = "5" color=#001940>R</font>](DocumentationGuidelines/BrandGlossary/R)&nbsp;&nbsp;
+[<font size = "5" color=#001940>S</font>](DocumentationGuidelines/BrandGlossary/S)&nbsp;&nbsp;
+[<font size = "5" color=#001940>T</font>](DocumentationGuidelines/BrandGlossary/T)&nbsp;&nbsp;
+[<font size = "5" color=#001940>U</font>](DocumentationGuidelines/BrandGlossary/U)&nbsp;&nbsp;
+[<font size = "5" color=#001940>V</font>](DocumentationGuidelines/BrandGlossary/V)&nbsp;&nbsp;
+<font size = "6" color=#00327F>W</font>&nbsp;&nbsp;
+<font size = "5" color=#001940>X</font>&nbsp;&nbsp;
+<font size = "5" color=#001940>Y</font>&nbsp;&nbsp;
+<font size = "5" color=#001940>Z</font>&nbsp;&nbsp;
+
+##Widget Blueprint
+
+A Widget Blueprint is a specialized Blueprint used by Unreal Motion Graphics (UMG) to provide visual layout and logic editing functionality for widgets.
+
+##World
+
+A World contains a list of Levels that are loaded. It handles the streaming of Levels and the spawning (creation) of dynamic Actors.
+
+##World Composition
+
+World Composition enables you to manage large worlds by streaming in levels using world origin shifting to avoid the floating point errors generally associated with large worlds.
+
+##World Outliner
+
+The World Outliner panel displays all of the Actors within the scene in a hierarchical tree view that can be filtered and enables selection and modification of Actors directly in the tree.
\ No newline at end of file
