Copying //UE4/Dev-Documentation to //UE4/Samples-Main (Source: //UE4/Dev-Documentation @ 3623665)

diff --git a/Engine/Documentation/Source/Engine/Actors/Merging/ActorMerging.INT.udn b/Engine/Documentation/Source/Engine/Actors/Merging/ActorMerging.INT.udn
--- a/Engine/Documentation/Source/Engine/Actors/Merging/ActorMerging.INT.udn
+++ b/Engine/Documentation/Source/Engine/Actors/Merging/ActorMerging.INT.udn
@@ -1,129 +1,211 @@
 Availability: Public
-Title:Actor Merging
-Crumbs: %ROOT%, Engine, Engine/Actors, Engine/Actors/Merging
-Description:How to use the Actor Merge tool to combine multiple Static Meshes into one in UE4.
-Related: Engine/Actors/StaticMeshActor
-Version: 4.12
-parent:Engine/Actors
-type:overview
-order:6
+Crumbs: %ROOT%
+Title: Actor Merging
+Description: An overview of the Actor Merge tool that is used to combine multiple Static Meshes into a single mesh inside UE4.
+Type: Overview
+SkillLevel:beginner
+Version:4.17
+Parent:Engine/Actors
+Order:6
+Tags:Static Mesh
+Tags:HLOD
 
-[TOC(start:2 end:2)]
+[VAR:spacer]![](%ROOT%/spacer.png)(w:150 convert:false)[/VAR]
 
-The Merge Actors tool enables the user to combine multiple Static Meshes into a single new Actor. Materials can also be baked together into a single Material with a series of custom textures using a new UV layout, and the original Static Meshes used can even be completely replaced by the newly merged asset.
+[VAR:spacer2]![](%ROOT%/spacer.png)(w:40 convert:false)[/VAR]
 
+
+[TOC(start:2 end:3)]
+
+The **Merge Actors** tool enables the user to combine multiple Static Meshes into a single new Actor. Materials can also be baked together into a single Material with a series of custom textures 
+using a new UV layout, and the original Static Meshes used can even be completely replaced by the newly merged asset.
+<!---
 [REGION:warning]
 The Merge Actor feature is officially supported in 4.12 and later. In 4.8 to 4.11, it must be activated from the Editor Settings under the Experimental tab. If you are using a version of the engine earlier than 4.12, this feature is considered experimental and is not officially supported.
 [/REGION]
+--->
+[REGION:raw]
+![](MergingActors.png)(w:950)
+[/REGION]
 
-[OBJECT:ToggleButtonContent]
-[PARAMLITERAL:category]
-OS
-[/PARAMLITERAL]
-[PARAMLITERAL:id]
-windows
-[/PARAMLITERAL]
-[PARAMLITERAL:active]
-[/PARAMLITERAL]
-[PARAM:content]
-![](MergeActors_Windows.png)
-[/PARAM]
-[/OBJECT]
-
-[OBJECT:ToggleButtonContent]
-[PARAMLITERAL:category]
-OS
-[/PARAMLITERAL]
-[PARAMLITERAL:id]
-mac
-[/PARAMLITERAL]
-[PARAMLITERAL:active]
-[/PARAMLITERAL]
-[PARAM:content]
-![](MergeActors_Mac.png)
-[/PARAM]
-[/OBJECT]
-
-## Actor Merging Worklow
-
-<!--
-[OBJECT:EmbeddedVideo]
-[PARAMLITERAL:videoid]
-_fTSGN8HGu8
-[/PARAMLITERAL]
-[PARAMLITERAL:width]
-930
-[/PARAMLITERAL]
-[PARAMLITERAL:height]
-568
-[/PARAMLITERAL]
-[/OBJECT]
--->
+
+## Actor Merging Workflow
 
 Merging Static Meshes in your scene is handled like so:
 
+1. From the main menu choose **Window > Developer Tools > Merge Actors.**
+	
+	[REGION:raw]
+	![](MergeMenu.png)(w:350)
+	[/REGION]
 1. Select a number of Static Meshes in your level that your would like to merge. 
+	
+	![](SelectedMeshes.png)(w:500)
 
-2. From the main menubar choose **Window > Developer Tools > Merge Actors.**
+1. Choose the appropriate options for merging the selected Actors.
 
-3. Choose the appropriate options (explained below).
+	![](ChooseSettings.png)(w:500)
 
-4. Click the Merge Actors button.
+	[REGION:note] Use the [reference tables](#mergeactorstooloptions) below to learn about each of the settings in this window and the actions they perform. [/REGION]
+	
+1. Click the **Merge Actors** button and give your new merged Actor a name, then click **Save**.
+	
+	[REGION:raw]
+	![](MergingActors.png)(w:950)
+	[/REGION]
 
-[REGION:tip]If you select fewer than **two** Static Meshes, you will get a warning saying there are insufficient meshes to merge.[/REGION]
+	You can then open the merged Actor as it's own asset from your **Content Browser**.
+	
 
 ## Merge Actors Tool Options
 
 At the top of the Merge Actors window is a list of all the Static Meshes selected in the viewport. Here you can check which meshes will be included in the final merge. 
 
+[REGION:note]
 If you have one or fewer meshes selected, you will see a yellow warning near the bottom of the window stating that you have insufficient mesh components for merging.
+[/REGION]
+
+### Mesh Components List
+
+In the **Mesh Components List** you'll find the selected actors that will be used when merging them together. 
+
+![](MeshComponentList.png)(w:500)
+
+[REGION:raw]
+| Option | %spacer2% | Description |
+|--------|-------------| ---- |
+| **Mesh Components List** || This area contains a checklist where you can choose which of the selected meshes will be included in the merge.|
+[/REGION]
+
+### Mesh Settings
+
+In the **Mesh Settings** section, you'll find the mesh properties that pertain to the meshes that will be merged.
+
+![](MeshSettings.png)(w:500)
+
+[REGION:raw]
+| Option | %spacer2% | Description |
+|--------|----------|-------------|
+|**Pivot Point at Zero**||If set to true, the pivot point of the final merged actor will be set to the scene's origin [0,0,0]. If set to false, the pivot will be located at the pivot location of the first Static Mesh Actor selected for the merge.
+|**Merge Physics Data**||If true, this will merge together all of the physics primitives into a single physics object.|
+|**Bake Vertex Data**||Bakes vertex data including vertex color and normals into the new resulting mesh.|
+|**LOD Selection Type**|| Use this to control which LODs are added to the merged object.  [INCLUDE:#lodselectiontype]  |
+|[REGION:tablesection]Advanced Options[/REGION]|||
+|**Generate Lightmap UVs**||Creates a new lightmap UV layout for the merged mesh, using the existing UV islands. This new layout will be placed in the Target Light Map UVChannel.|
+|**Target Light Map Resolution**||Specifies the resolution for the new lightmap UVs, which controls the distance calculated between each island within the layout. Higher resolutions will yield tighter spacing between each island.|
+[/REGION]
+
+<!--- LOD Selection Type Excerpts ---
+
+[EXCERPT:lodselectiontype]
+[REGION:raw]
+| Selection | %spacer2% | Description |
+| --- | --- | --- |
+| **Use All LODs** || If selected, this use all available LOD levels when merging the actors. |
+| **Use Specific LOD Level** || If selected, this will use only the selected LOD level when merging the actors. |
+| **Calculate Correct LOD level** || If selected, this will calculate the appropriate LOD model for the given screen size. |
+[/REGION]
+[/EXCERPT:lodselectiontype]
+
+--- LOD Selection Type Excerpts --->
+
+
+### Material Settings
+
+In the **Material Settings** section, you'll find the texture and material properties that pertain to the meshes that will be merged.  
+
+![](MaterialSettings.png)(w:500)
+
+[REGION:raw]
+| Option | %spacer2% | Description |
+|--------| --- |-------------|
+| **Merge Materials** || If checked, the Materials of all the merged assets will be combined into a single Material with new UV layout. The following subsettings apply only if Merge Materials is checked. |
+| **Texture Size** || This provides access to the X and Y pixel dimensions (resolution) for the final merged Material's textures.|
+| **Texture Sizing Type** || The selection options available when choosing how you want the size of your texture to be calculated.  [INCLUDE:#texturesizing]|
+| **Normal Map** || If checked, a Normal map will be baked out for the resulting merged Actor.|
+| ** Metallic Map** || If checked, a Metallic map will be baked out for the resulting merged Actor. This nullifies the Metallic Constant property.|
+| **Metallic Constant** || Only enabled if Metallic Map is unchecked, this sets a constant value for the Metallic property of the resulting merged Actor's Material.|
+| **Roughness Map** || If checked, a Roughness map will be baked out for the resulting merged Actor. This nullifies the Roughness Constant property.|
+| **Roughness Constant** || Only enabled if Roughness Map is unchecked, this sets a constant value for the Roughness property of the resulting merged Actor's Material.| 
+| **Specular Map** || If checked, a Specular map will be baked out for the resulting merged Actor. This nullifies the Specular Constant property.|
+| **Specular Constant** || Only enabled if Specular Map is unchecked, this sets a constant value for the Specular property of the resulting merged Actor's Material.|
+| **Emissive Map** || If checked, an Emissive map will be baked out for the resulting merged Actor.|
+| **Opacity Map** || If checked, an Opacity map will be baked out for the resulting merged Actor. This nullifies the Opacity Constant property.  Must also use the **Blend Mode** for **Translucent** to show correctly in the generated instances. |
+| **Opacity Constant** || Only enabled if Opacity Map is unchecked, this sets a constant value for the Opacity property of the resulting merged Actor's Material. |
+| **Opacity Mask Map** || If checked, an Opacity Mask map will be baked out for the resulting merged Actor. This nullifies the Opacity Mask Constant property. Must also use the **Blend Mode** for **Masked** to show correctly in the generated instances. |
+| **Opacity Mask Constant** || Only enabled if Opacity Mask Map is unchecked, this sets a constant value for the Opacity Mask property of the resulting merged Actor's Material. |
+| **Diffuse Texture Size** || Override for the Diffuse Texture Size. Enabled when **Texture Sizing Type** is set to **Use per property manually overriden texture sizes**. |
+| **Normal Texture Size** || Override for the Normal Texture Size. Enabled when **Texture Sizing Type** is set to **Use per property manually overriden texture sizes**. |
+| **Metallic Texture Size** || Override for the Metallic Texture Size. Enabled when **Texture Sizing Type** is set to **Use per property manually overriden texture sizes**. |
+| **Roughness Texture Size** || Override for the Roughness Texture Size. Enabled when **Texture Sizing Type** is set to **Use per property manually overriden texture sizes**. |
+| **Specular Texture Size** || Override for the Specular Texture Size. Enabled when **Texture Sizing Type** is set to **Use per property manually overriden texture sizes**. |
+| **Emissive Texture Size** || Override for the Emissive Texture Size. Enabled when **Texture Sizing Type** is set to **Use per property manually overriden texture sizes**. |
+| **Opacity Texture Size** || Override for the Opacity Texture Size. Enabled when **Texture Sizing Type** is set to **Use per property manually overriden texture sizes**. |
+| **Opacity Mask Texture Size** || Override for the Opacity Mask Texture Size. Enabled when **Texture Sizing Type** is set to **Use per property manually overriden texture sizes**. |
+| **Blend Mode** || This determines how the material's color is blended with background colors. [INCLUDE:#blendmode] |
+| **Use Vertex Data for Baking Material** || Whether or not to use vertex data such as vertex color to bake Material result. Useful if you have a Material that is being controlled in some way by vertex color.|
+| **Use Texture Binning** || This will calculate the varying output texture sizes according to their importance when packing the final atlas texture. |
+[/REGION]
+
+
+<!-- Material Settings Excerpts ----
 
 
-| Option | Description |
-|--------|-------------|
-|[REGION:tablesection]Mesh Components List[/REGION]|This area contains a checklist where you can chose which of the selected meshes will be included in the merge.|
-|[REGION:tablesection]Mesh Settings[/REGION]||
-|**Pivot Point at Zero**|If set to true, the pivot point of the final merged actor will be set to the scene's origin <0,0,0>. If set to false, the pivot will be located at the pivot location of the first Static Mesh Actor selected for the merge.
-|**Merge Physics Data**|If true, this will merge together all of the physics primitives into a single physics object.|
-|**Bake Vertex Data**|Bakes vertex data including vertex color and normals into the new resulting mesh.|
-|**LOD Selection Type**| Use this to control which LODs are added to the merged object.|
-|[REGION:tablesection]Advanced Options[/REGION]|These options are found by expanding the small white triangle icon in the middle of the options list|
-|**Generate Lightmap UVs**|Creates a new lightmap UV layout for the merged mesh, using the existing UV islands. This new layout will be placed in the Target Light Map UVChannel.|
-|**Target Light Map UVChannel**|Specifies the UV channel into which the new lightmap UVs will be saved.|
-|**Target Light Map Resolution**|Specifies the resolution for the new lightmap UVs, which controls the distance calculated between each island within the layout. Higher resolutions will yield tighter spacing between each island.|
-|[REGION:tablesection]Material Settings[/REGION]||
-|**Merge Materials**|If checked, the Materials of all the merged assets will be combined into a single Material with new UV layout. The following subsettings apply only if Merge Materials is checked. |
-| |[INCLUDE:#materialsettings]|
-|[REGION:tablesection]Landscape Culling[/REGION]||
-|**Use Landscape Culling**|If checked, existing landscape geometry will be used to cull fully occluded triangles. For example, if you had a tree mesh that was planted such that the bottommost triangles were well below the surface of a landscape, then those occluded triangles would be removed upon merge.|
-|[REGION:tablesection]Standard Options[/REGION]||
-|**Replace Source Actors**|If checked, the Actors selected in the viewport will be removed and replaced with the new merged version. This will not afffect the the original objects in the Content Browser.| 
-
-
-
-<!--
-[EXCERPT:MaterialSettings]
-[REGION:simpletable]
-|----|----|
-|**Texture Size**|This provides access to the X and Y pixel dimensions (resolution) for the final merged Material's textures.|
-|**Texture Sizing Type**| The settings for this property are explained in the table below.|
-| **Normal Map**| If checked, a Normal map will be baked out for the resulting merged Actor.|
-|** Metallic Map** |If checked, a Metallic map will be baked out for the resulting merged Actor. This nullifies the Metallic Constant property.|
-| **Metallic Constant**|Only enabled if Metallic Map is unchecked, this sets a constant value for the Metallic property of the resulting merged Actor's Material.|
-| **Roughness Map** |If checked, a Roughness map will be baked out for the resulting merged Actor. This nullifies the Roughness Constant property.|
-| **Roughness Constant**|Only enabled if Roughness Map is unchecked, this sets a constant value for the Roughness property of the resulting merged Actor's Material.| 
-| **Specular Map** |If checked, a Specular map will be baked out for the resulting merged Actor. This nullifies the Specular Constant property.|
-| **Specular Constant**|Only enabled if Specular Map is unchecked, this sets a constant value for the Specular property of the resulting merged Actor's Material.|
-| **Emissive Map** |If checked, an Emissive map will be baked out for the resulting merged Actor.|
-|**Use Vertex Data for Baking Material**|Whether or not to use vertex data such as vertex color to bake Material result. Useful if you have a Material that is being controlled in some way by vertex color.|
+[EXCERPT:texturesizing]
+[REGION:raw]
+| Type | %spacer2% | Description |
+| --- | :---: | :--- |	
+| **Use TextureSize for all material properties** || This will use the value entered into Texture Size for sizing. |
+| **Use automatically biased texture sizes based on TextureSize** || This creates a biased set of {normal:diffuse:otherproperties} textures using the ratio of {1:0.5:0.25} so if the texture size is set to 512 it will output a 512 normal map, 256 diffuse map and 128 map for roughness specular etc (diffuse is clamped to 32x32 and roughness specular etc to 16x16) |
+| **Use per property manually overridden texture sizes** || Enables the per texture overrides, and uses those values instead of Texture Size. |
+| **Use Simplygon’s automatic texture sizing** || Uses Simplygon's algorithm for texture sizing based on its decimation settings. This should only be used if you are already using Simplygon. |
 [/REGION]
-[/EXCERPT:MaterialSettings]
--->
-
-|[REGION:tablesection]Texture Sizing Type Settings[/REGION]||
-|-----|-----|
-|**Use TextureSize for all material properties**|This uses the resolution set in the Texture Size property for all baked textures.|
-|**Use automatically biased texture sizes based on TextureSize**|This creates a 1:0.5:0.25 bias for all textures such that Normal gets 100% of the resolution set in Texture Size, Diffuse (Base Color) gets 50% of Texture Size, and all other properties get 25% of TextureSize|
-|**Use per property manually overriden texture sizes**|This provides individual number entry fields for explicitly providing the resolution for each map being baked.|
-|**Use Simplygon's automatic texture sizing**|This uses Simplygon's algorithm for texture sizing based on its decimation settings. This should only be used if you are already using Simplygon.|
+[/EXCERPT:texturesizing]
+
+
+
+
+
+
+
+[EXCERPT:blendmode]
+[REGION:raw]
+| Selection | %spacer2% | Description|
+| --- | --- | --- |
+| **Opaque** || This blend mode defines a surface through which light neither passes nor penetrates. |
+| **Masked** || This blend mode defines the visibility in a binary (on/off) fashion by marking whether areas of the Material should be visible or not. |
+| **Translucent** || This blend mode defines the visibility of a Material by using an Opacity value. |
+| **Additive** || This blend mode takes the pixels of the Material and adds them to the pixels of the background. |
+| **Modulate** || This blend mode multiplies the value of the Material against the pixels of the background. |
+| **Alpha Composite (Premultiplied Alpha)** ||  This blend mode is used with textures that already have a pre-multiplied alpha. |
+[/REGION]
+[/EXCERPT:blendmode]
+
+--- /End Material Setting Excerpts -->
+
+
+### Landscape Culling Settings
+
+In the **Landscape Culling** section, you'll find properties that use the Landscape to cull parts of the mesh while merging them.
+
+![](LandscapeCullingSettings.png)(w:500)
+
+[REGION:raw]
+| Option | %spacer% | Description |
+|--------|-------------| --- |
+| **Use Landscape Culling** || If checked, existing landscape geometry will be used to cull fully occluded triangles. For example, if you had a tree mesh that was planted such that the bottommost triangles were well below the surface of a landscape, then those occluded triangles would be removed upon merge.|
+[/REGION]
+
+### Standard Options
+
+In this section, you'll find properties that do not fit categorically in other sections, like being able to replace the actors in the viewport with the merged one.
+
+![](StandardOptions.png)(w:500)
+
+[REGION:raw]
+| Option | %spacer% | Description |
+|--------|-------------| --- |
+| **Replace Source Actors** || If checked, the Actors selected in the viewport will be removed and replaced with the new merged version. This will not affect the the original objects in the Content Browser.| 
+[/REGION]
+
 
diff --git a/Engine/Documentation/Source/Engine/Actors/Volumes/Volumes.INT.udn b/Engine/Documentation/Source/Engine/Actors/Volumes/Volumes.INT.udn
--- a/Engine/Documentation/Source/Engine/Actors/Volumes/Volumes.INT.udn
+++ b/Engine/Documentation/Source/Engine/Actors/Volumes/Volumes.INT.udn
@@ -1,31 +1,37 @@
 Availability: Public
-Title:Volumes Overview
-Crumbs:%ROOT%, Engine, Engine/Actors
+Crumbs:%ROOT%
+Title:Volumes Reference
 Description:Guide to placing and working with Volumes to augment the behavior and setup of your levels.
+Type:Reference
+SkillLevel:Beginner
 version:4.16
-type:overview
 Parent:Engine/Actors/Types
 order:8
-type:overview
+tags:Volumes
+tags:Collision
 
+[TOC(start:2 end:4)]
 
-[TOC(start:2 end:3)]
+**Volumes** are three-dimensional **Actors** used to alter the behavior of areas within levels. Volumes can be used for behaviors like:
 
-**Volumes** are 3-dimensional areas within your levels that can serve variety of purposes. Some of these purposes include:
-
-* Causing damage to the Player.
-* Blocking Actors from entering the volume, acting as a collision surface.
-* Detecting and reacting to Actors entering or leaving the volume.
+* Causing damage to the player or other Actors inside the Volume.
+* Blocking certain Actors from entering the Volume, acting as a collision surface.
+* Opening a door when an Actor enters the Volume.
 * Changing the way a level calculates its lighting or visibility.
 
-## Volume Placement Considerations
 
-There are two important points to keep in mind about using volumes:
+## Creating and Placing Volumes
+
+Volumes can be created by selecting a Volume from the **Volumes** tab of the **Modes** menu and dragging and dropping it into the level viewport. 
+
+![](creatingVolumes.png)
+
+Once the Volume has been placed in the level, you can resize or reposition it as you see fit. There are two important points to keep in mind when placing Volumes:
 
-1. Volumes are support actors, generally used to detect when certain Actor types have entered a specific area and trigger an effect in response.
-1. Volumes sometimes have built-in effects of their own (either in code or Blueprints), but most often they'll provide cues to other Actors. They should usually be thought of as components of a greater effect, with some other level element acting as a visual cue.
+1. Volumes are support Actors, generally used to detect when certain Actor types have entered a specific area and trigger an effect in response.
+1. Volumes sometimes have built-in effects of their own (either in code or Blueprints), but most often they'll provide cues to other Actors. They should usually be thought of as parts of a greater system, with some other  element acting as a visual cue.
 
-As an example, a **Pain-Causing Volume**, which causes damage to anything inside the volume, could be placed into a level like this:
+As an example, a **Pain-Causing Volume**, which causes damage to anything inside the Volume, could be placed into a level like this:
 
 [VAR:VolumeDemoSize]
 h:301
@@ -37,7 +43,7 @@ h:301
 | In Editor | In Game |
 [/REGION]
 
-In this case, the player will not know to avoid the volume, or understand why they take damage when they enter it. The level designer should provide a hint to the player that explains the Pain Causing Volume:
+In this case, the player will not know to avoid the Volume, or understand why they take damage when they enter it. The level designer should provide a hint to the player that explains the Pain Causing Volume:
 
 
 [REGION:imagetable]
@@ -46,137 +52,121 @@ In this case, the player will not know to avoid the volume, or understand why th
 | In Editor | In Game |
 [/REGION]
 
-In this case, the fire particle effect provides a visual explanation of why the area is dangerous, while the Pain Causing Volume supports the visuals with a gameplay effect of decreasing the player's health during exposure to the fire.
-
-
-## Creating Volumes
-
-Volumes can be created by selecting a volume from the **Volumes** tab of the **Modes** menu and dragging and dropping it into the level viewport. 
-
-![](creatingVolumes.png)
+In this case, the fire particle effect provides a visual explanation of why the area is dangerous, while the Pain Causing Volume supports the visuals with the gameplay effect of decreasing the player's health during exposure to the fire.
 
-Once the volume has been placed in the level, you can then resize and (or) reposition the volume as you see fit. 
 
 ## Volume Types
 
 
-### Audio Volume
-**Audio Volumes** allow you to control and apply various sounds in your level as well as provide an avenue to create compartmentilized audio zones where you can control what is heard inside and outside of the volume. Through the use of Audio Volumes, you can generate audio transitions via fading and control the Ambient Zone settings within the volumes themselves.
-[REGION:note]See the [details page](Engine/Actors/Volumes/AudioVolumeDetails) for more information on the specific properties of Audio Volumes.[/REGION]
-
+### Collision and Overlap Volumes
 
-### Blocking Volume
+#### Blocking Volume
 
-A **Blocking Volume** serves as a collision shape through which objects are not intended to pass. By adjusting the collision channels on the volume (pictured below), you can control which types of objects will be allowed to pass, and if any objects should also generate overlap events when they pass through the volume.
+A **Blocking Volume** serves as a collision shape through which Actors are not intended to pass. By adjusting the collision channels on the Volume (pictured below), you can control which types of Actors will be allowed to pass, and if any Actors should also generate overlap events when they pass through the Volume.
 [REGION:note]See the [Collision Documentation](Engine\Physics\Collision) for more details.[/REGION]
 
 ![](blocking_volume_details.png)
 
 Blocking Volumes can be used in the place of collision surfaces on Static Meshes, particularly in the case of walls in structures. This can cause scenes to operate more predictably, since physical objects won't interact with small details like bumps on the floor and walls. It can also improve performance by reducing the cost of physics simulation.
 
+#### Camera Blocking Volume
 
-### Camera Blocking Volume
-
-**Camera Blocking Volumes** are volumes with their collision settings pre-configured to block only cameras and ignore everything else. They are intended to define invisible barriers that keep the camera out of bad locations.  An example of this might be in a third-person game, when the walls of the play area have decorative coverings like leafy vines. In these cases, a thin Camera Blocking Volume can be placed against the wall so that the camera doesn't bump into the vines or go behind the leaves, leaving it to slide smoothly and provide an unobstructed view of the action in the room.
+**Camera Blocking Volumes** are Volumes with their collision settings pre-configured to block only cameras and ignore everything else. They are intended to define invisible barriers that keep the camera out of bad locations.  An example of this might be in a third-person game, when the walls of the play area have decorative coverings like leafy vines. In these cases, a thin Camera Blocking Volume can be placed against the wall so that the camera doesn't bump into the vines or go behind the leaves, leaving it to slide smoothly and provide an unobstructed view of the action in the room.
 
-
-### Cull Distance Volume
-
-**Cull Distance Volumes** are optimization tools that allow for objects to be culled (not drawn to the screen) based on that object's distance from the camera and its size. This can help optimize your scene by not drawing objects when they are small enough to be considered unimportant. Size is calculated by the bounding box along its longest dimension, and the cull distance chosen is the one closest to that size.
-
-Cull Distance Volume setup is dependent on the **Cull Distances** property, shown in the **Details** panel below.
-
-![](CullDistancesProperty.png)
-
-In the image above, the properties define the following behavior:
-
-* Objects within the volume that have a size closest to 50 units (distance equal to or less than 85 units) will be culled (disappear) when they are 500 units away from the camera or farther.
-* Objects within the volume that have a size closest to 120 units (distance between 85 units and 210 units) will be culled (disappear) when they are 1000 units away from the camera or farther.
-* Objects within the volume that have a size closest to 300 units (distance equal to or greater than 210 units) will never be culled, since 0 in this case is treated as infinity, meaning that the camera can never get far enough away.
-
-Setup starts with adding a new entry into the Cull Distances array, which is done by clicking the ![](button_Plus.png) button. Next, simply fill in the size of a given object and the desired distance at which you want objects of that size or smaller to be culled. These properties do not have to be created in a specific order.
-
-
-### Hierarchical LOD Volume
-
-**Hierarchical LOD Volumes** have no special properties, but are used by the [HLOD](Engine/HLOD) system to group Actors into a single HLOD cluster. When generating clusters, the Unreal Engine will override its normal generation process in deference to manually-placed volumes.
-
-
-### Kill Z Volume
+#### Kill Z Volume
 
 The purpose of **Kill Z Volume** is to prevent objects from going out of bounds in certain types of games, such as falling off of a cliff or into a pit in a platformer game, or leaving a spaceship without a suit in a science-fiction setting. The Kill Z Volume will call the `FellOutOfWorld` function on any Actor that enters it, and by default, Actors will go through a quick cleanup procedure and then destroy themselves. You can override this behavior for any of your Actor types if your game requires something different. For example, if a key or other item that the player is required to collect in order to continue the game falls into a lava pit, your game might want to teleport the item back up to an area the player can reach rather than destroying it, or at least inform the player that the item was lost and reload the last checkpoint, rather than leaving the game in an unwinnable state.
 
+#### Pain-Causing Volume
 
-### Level Streaming Volumes
+**Pain-Causing Volumes** are Physics Volumes with the ability to inflict damage on players. These are useful when you have obvious areas of a level that the player should not go, such as a lava pit, a cloud of toxic gas, and so on.
+[REGION:note]See the [details page](Engine/Actors/Volumes/PainCausingVolumeDetails) for more information on the specific properties of Pain-Causing Volumes.[/REGION]
 
-[EXCERPT:LSV]
-**Level Streaming Volumes** are used to aid in the [level streaming](Engine/LevelStreaming) process. They provide a simple way to encapsulate a level, as well as control when it streams in and out of memory, based on when a Player enters or exits the volume.
-[/EXCERPT]
+#### Physics Volume
 
+**Physics Volumes** are Volumes in which the physical setup that affects characters and other physics objects can be manipulated. A common use for them is for the creation of watery environments in which the player needs to swim. The effects of Physics Volumes are visible and can be interpreted as desired. The Character Movement Component class uses the current fields to adjust how their owning `Character` moves through the environment. If your game has custom physics, deriving your own child class from `APhysicsVolume` can help you to provide it.
+[REGION:note]See the [details page](Engine/Actors/Volumes/PhysicsVolumeDetails) for more information on the specific properties of Physics Volumes.[/REGION]
 
-### Lightmass Character Indirect Detail Volume
+#### Trigger Volume
 
-The **Lightmass Character Indirect Detail Volume** is similar to the Lightmass Importance Volume and generates indirect light samples, not just at Player height above ground, but generates them inside the entire Volume. An example usage of this type of Volume would be an elevator shaft to ensure the indirect lighting is correct extending up/down the shaft.
+**Trigger Volumes** are used as a way to cause events when a Player or other object enters or exits them. They can be used quickly with the **Level Blueprint** to test out events and gameplay scenarios or functionality without the need for additional Blueprints.
 
+For example, you could place a TriggerVolume in your level, then create an overlap event for the Volume in your **Level Blueprint** which can play a sound, open a door, or start a cinematic scene.
 
-### Lightmass Importance Volume
+[REGION:note]
+Remember to check the collision settings to ensure that your trigger will react to the intended Actors with the Overlap collision response setting.
+[/REGION]
 
-**Lightmass Importance Volumes** are used to focus Lightmass calculations. Lightmass is the precalculated lighting and global illumination system used in Unreal Engine 4. In order to minimize - and thereby optimize - the amount of processing needed to generate precalculated lights in a level, Lightmass Importance Volumes are use to enclose the areas of the level geometry. Lightmass will only calculate within that area, skipping anything outside of it.
-[REGION:note]For more information on Lightmass and Lightmass Importance Volumes, please see the [Lightmass documentation](Engine/Rendering/LightingAndShadows/Lightmass).[/REGION]
 
+### Graphical and Audio Volumes
 
-### Mesh Merge Culling Volume
+#### Audio Volume
+**Audio Volumes** allow you to control and apply various sounds in your level as well as provide an avenue to create compartmentilized audio zones where you can control what is heard inside and outside of the Volume. Through the use of Audio Volumes, you can generate audio transitions via fading and control the Ambient Zone settings within the Volumes themselves.
+[REGION:note]See the [details page](Engine/Actors/Volumes/AudioVolumeDetails) for more information on the specific properties of Audio Volumes.[/REGION]
 
-**Mesh Merge Culling Volumes** mark the mesh objects they contain so that these objects will be combined into a single mesh. This can improve performance by making a collection of small meshes in a contained area all cull together as one mesh, or by causing [HLOD](Engine/HLOD) generation to reduce geometry in a more optimal way.
+#### Cull Distance Volume
 
+**Cull Distance Volumes** are optimization tools that allow for objects to be culled (not drawn to the screen) based on that object's distance from the camera and its size. This can help optimize your scene by not drawing objects when they are small enough to be considered unimportant. Size is calculated by the bounding box along its longest dimension, and the cull distance chosen is the one closest to that size.
 
-### Nav Mesh Bounds Volume
+Cull Distance Volume setup is dependent on the **Cull Distances** property, shown in the **Details** panel below.
 
-**Nav Mesh Bounds Volumes** are used as a way to control where Nav Meshes will be built in a level. Nav Meshes are used to calculate navigation paths throughout the areas of a level.
+![](CullDistancesProperty.png)
 
-Within the volume, a Nav Mesh is constructed on all surfaces with an appropriate angle to be walked upon. You may overlap as many of these as you need to generate the desired Nav Mesh.
+In the image above, the properties define the following behavior:
 
-To use the Nav Mesh Bounds Volume, simply create one (or more) that encloses the navigable areas of your level. The Nav Mesh will be built automatically.
+* Objects within the Volume that have a size closest to 50 units (distance equal to or less than 85 units) will be culled (disappear) when they are 500 units away from the camera or farther.
+* Objects within the Volume that have a size closest to 120 units (distance between 85 units and 210 units) will be culled (disappear) when they are 1000 units away from the camera or farther.
+* Objects within the Volume that have a size closest to 300 units (distance equal to or greater than 210 units) will never be culled, since 0 in this case is treated as infinity, meaning that the camera can never get far enough away.
 
-[REGION:tip]
-You can press **P** at any time in the viewport to visualize the Nav Mesh.
-[/REGION]
+Setup starts with adding a new entry into the Cull Distances array, which is done by clicking the ![](button_Plus.png) button. Next, simply fill in the size of a given object and the desired distance at which you want objects of that size or smaller to be culled. These properties do not have to be created in a specific order.
 
-[REGION:note]Please see the Content Examples and [Nav Mesh Documentation](Resources\ContentExamples\NavMesh) for more details and examples.[/REGION]
+#### Hierarchical LOD Volume
 
+**Hierarchical LOD Volumes** have no special properties, but are used by the [HLOD](Engine/HLOD) system to group Actors into a single HLOD cluster. When generating clusters, the Unreal Engine will override its normal generation process in deference to manually-placed Volumes.
 
-### Pain-Causing Volume
+#### Lightmass Character Indirect Detail Volume
 
-**Pain-Causing Volumes** are Physics Volumes with the ability to inflict damage on players. These are useful when you have obvious areas of a level that the player should not go, such as a lava pit, a cloud of toxic gas, and so on.
-[REGION:note]See the [details page](Engine/Actors/Volumes/PainCausingVolumeDetails) for more information on the specific properties of Pain-Causing Volumes.[/REGION]
+The **Lightmass Character Indirect Detail Volume** is similar to the Lightmass Importance Volume and generates indirect light samples, not just at Player height above ground, but inside the entire Volume. An example usage of this type of Volume would be to place the Volume in an elevator shaft, ensuring that indirect lighting will be correct all the way along the shaft, instead of just at the bottom.
 
+#### Lightmass Importance Volume
 
-### Physics Volume
+**Lightmass Importance Volumes** are used to focus Lightmass calculations. Lightmass is the precalculated lighting and global illumination system used in Unreal Engine 4. In order to minimize the amount of processing needed to generate precalculated lights in a level, Lightmass Importance Volumes are used to enclose areas of the level geometry. Lightmass will only calculate within those areas, skipping anything outside of them.
+[REGION:note]For more information on Lightmass and Lightmass Importance Volumes, please see the [Lightmass documentation](Engine/Rendering/LightingAndShadows/Lightmass).[/REGION]
 
-**Physics Volumes** are volumes in which the physical setup that affects characters and other physics objects can be manipulated. A common use for them is for the creation of watery environments in which the player needs to swim. The effects of Physics Volumes are visible and can be interpreted as desired. The Character Movement Component class uses the current fields to adjust how their owning `Character` moves through the environment. If your game has custom physics, deriving your own child class from `APhysicsVolume` can help you to provide it.
-[REGION:note]See the [details page](Engine/Actors/Volumes/PhysicsVolumeDetails) for more information on the specific properties of Physics Volumes.[/REGION]
+#### Mesh Merge Culling Volume
 
+**Mesh Merge Culling Volumes** mark the mesh objects they contain so that these objects will be combined into a single mesh. This can improve performance by making a collection of small meshes in a contained area all cull together as one mesh, or by causing [HLOD](Engine/HLOD) generation to reduce geometry in a more optimal way.
 
-### Post Process Volume
+#### Post Process Volume
 
 A **Post Process Volume** is a space in which the Post Process settings applied to the camera can be overridden by adjusting the properties in the **Details** panel. 
 [REGION:note]For more information on Post Processing, please see the [Post Processing Documentation](Engine/Rendering/PostProcessEffects).[/REGION]
  
+#### Precomputed Visibility Volume
+**Precomputed Visibility Volumes** are used primiarly for performance optimization. These Volumes store the visibility of Actors for their location in the world. These should be placed only in areas the Player can access.
 
-### Precomputed Visibility Volume
-**Precomputed Visibility Volumes** are used primiarly for performance optimization. These volumes store the visibility of Actors for their location in the world. These should be placed only in areas the Player can access.
+#### Precomputed Visibility Override Volume
+**Precomputed Visibility Override Volumes** let you manually override the visibility of Actors for their location in the world if the auto generated result of a Precomputed Visibilty Volume is undesired. These are also used for performance optimization and should only be placed in areas where the Player can access.
 
 
-### Precomputed Visibility Override Volume
-**Precomputed Visibility Override Volumes** let you manually override the visibility of Actors for their location in the world if the auto generated result of a Precomputed Visibilty Volume is undesired. These are also used for performance optimization and should only be placed in areas where the Player can access.
+### Level and AI Volumes
 
+#### Level Streaming Volumes
 
-### Trigger Volume
+[EXCERPT:LSV]
+**Level Streaming Volumes** are used to aid in the [level streaming](Engine/LevelStreaming) process. They provide a simple way to encapsulate a level, as well as control when it streams in and out of memory, based on when a Player enters or exits the Volume.
+[/EXCERPT]
 
-**Trigger Volumes** are used as a way to cause events when a Player or other object enters or exits them. They can be used quickly with the **Level Blueprint** to test out events and gameplay scenarios or functionality without the need for additional Blueprints.
+#### Nav Mesh Bounds Volume
+
+**Nav Mesh Bounds Volumes** are used as a way to control where Nav Meshes will be built in a level. Nav Meshes are used to calculate navigation paths throughout the areas of a level.
 
-For example, you could place a TriggerVolume in your level, then create an overlap event for the volume in your **Level Blueprint** which can play a sound, open a door, or start a cinematic scene.
+Within the Volume, a Nav Mesh is constructed on all surfaces with an appropriate angle to be walked upon. You may overlap as many of these as you need to generate the desired Nav Mesh.
 
-[REGION:note]
-Remember to check the collision settings to ensure that your trigger will react to the intended Actors with the Overlap collision response setting.
+To use the Nav Mesh Bounds Volume, simply create one (or more) that encloses the navigable areas of your level. The Nav Mesh will be built automatically.
+
+[REGION:tip]
+You can press **P** at any time in the viewport to visualize the Nav Mesh.
 [/REGION]
+
+[REGION:note]Please see the Content Examples and [Nav Mesh Documentation](Resources\ContentExamples\NavMesh) for more details and examples.[/REGION]
diff --git a/Engine/Documentation/Source/Engine/Animation/AnimModifiers/AnimModifiers.INT.udn b/Engine/Documentation/Source/Engine/Animation/AnimModifiers/AnimModifiers.INT.udn
new file mode 100644
--- /dev/null
+++ b/Engine/Documentation/Source/Engine/Animation/AnimModifiers/AnimModifiers.INT.udn
@@ -0,0 +1,143 @@
+Availability: Docs
+Crumbs: %ROOT%
+Title: Animation Modifiers
+Description:Animation Modifiers enable users to define a sequence of actions for a given animation sequence or skeleton.
+Type: Overview
+SkillLevel: Advanced
+Version: 4.17
+Parent: Engine/Animation
+Order: 1
+Tags: Animation Notifies
+Tags: Skeleton
+Tags: Animation Asset Editor
+topic-image:Engine/Physics/physics_topic.png
+
+[REGION:warning]
+This feature is currently experimental and is in development, therefore some features may not work as expected. 
+[/REGION]
+
+[TOC (start:2 end:2)]
+
+**Animation Modifiers (Anim Modifier)** are a type of native or [](Engine/Blueprints/UserGuide/Types/ClassBlueprint) that enable users to apply a sequence of actions to an [Animation Sequence](Engine/Animation/Sequences) or [Skeleton](Engine/Animation/Skeleton) asset.
+This includes but is not limited to things like creating automatic foot sync markers by pin-pointing on which frames(s) the right or left foot is placed on the ground and adding **Animation Sync Markers** to frames where a foot bone is at its lowest point (touching the floor). 
+
+Accessing and applying Anim Modifiers is performed in either the **Animation Editor** or **Skeleton Editor** under the **Animation Data** tab. 
+When applying an Anim Modifier to a Skeleton, the modifier is applied to all Animation Sequences that are based on the Skeleton. 
+When applying the modifier to an Animation Sequence, it is only applied to the sequence itself and no other sequences. 
+
+## Creating an Anim Modifier
+
+To get started, you will need to create an Anim Modifier Blueprint class: 
+
+1. In your project's **Content Browser**, click the **Add New** button and select **Blueprint Class**.
+
+	![](Create_01.png)
+
+1. In the **Pick Parent Class** window, expand **All Classes** and search for and select **Animation Modifier**, click **Select** then give it a name.
+
+	![](Create_02.png)
+
+1. Double-click and open the new Animation Modifier Blueprint to open in the Blueprint Editor. 
+
+	![](Create_03.png)
+
+You can now use Blueprint Script and functions contained with the **Animation Blueprint Library** to access and manipulate animation data. 
+
+### Animation Blueprint Library
+
+Right-click in the graph of the Anim Modifier Blueprint to see the context menu and list of functions available, particularly under the Animation Blueprint Library. 
+
+![](AnimBPLibrary.png)
+
+[REGION:caption]
+In the image above, functions related to **Marker Syncing** are expanded which enable you to synchronize animations using marker data. 
+[/REGION]
+
+Before working with the various functions available for accessing data, you will want to implement the **OnApply** and **OnRevert** events. 
+The OnApply Event enables the user to change, add or remove data from the animation while the OnRevert enables the user to remove previously applied user changes (or return the sequence to its original state). 
+Each event returns the Animation Sequence in which to feed the Animation Blueprint Library operations. 
+
+![](EventNodes.png)
+
+[REGION:caption]
+In the image above, when the Anim Modifier is applied, a new Notify track is created with the defined name while reverting the Anim Modifier, the track is removed. 
+[/REGION]
+
+### Exposing Properties
+
+Inside your Anim Modifier Blueprint, you will want to expose parameters using **Instance Editable** that you can manipulate within the **Animation Tab** in the Animation Tools
+
+![](InstanceEditable.png)
+
+[REGION:caption]
+In the image above, we use a structure variable containing information that we can set for the name of a bone as well as provide an offset. 
+[/REGION]
+
+Inside our Animation Sequence when we implement our Anim Modifier and apply it, we can see our exposed parameters. 
+
+![](AnimDataVariables.png)
+
+[REGION:caption]
+In the image above, we enter the names of the bones that we want to include in our Anim Modifier in which to gather animation data (such as Bone Transform) from. 
+[/REGION]
+
+
+## Implementing an Anim Modifier
+
+Implementing an Anim Modifier can be done inside a **Skeleton** asset (to add to the Anim Modifier to all animations associated with the Skeleton) or in a single Animation Sequence. 
+
+1. Inside the **Skeleton Editor** or **Animation Editor**, go to **Window** menu option and select **Animation Data Modifiers**. 
+
+	![](Implement_01.png)
+
+1. In the **Animation Data Modifiers** tab, select **Add Modifier** and select your desired **Anim Modifier Blueprint**. 
+
+	![](Implement_02.png)
+
+1. Right-click on your Anim Modifier Blueprint then click **Apply Modifier** to apply the Anim Modifier and any changes (or **Revert Modifier** to remove changes). 
+
+	![](Implement_03.png)
+
+	Your Anim Modifier will be listed as **Out of Date** until you apply it. 
+
+Any properties that you have set as **Instance Editable** will be available to edit in the **Animation Data Modifiers** tab. 
+
+![](Implement_04.png)
+
+[REGION:caption]
+In the image above, we define properties that can be used to drive automatic foot syncing. 
+[/REGION]
+
+
+## Animation Blueprint Library Reference
+
+While there are several different nodes available within the Animation Blueprint Library, this section lists some of the more commonly used types within Anim Modifiers.
+
+### Add/Remove Notify and Curve Tracks
+
+![](NotifiesAndCurves.png)
+
+These nodes enable you to dynamically add Notify or Curve Tracks to the specified Animation Sequence. 
+Once Notify or Curve Tracks have been added, you can add various types of keys or events to those tracks. 
+For example, you may want to **Add Float Curve Keys** to a Curve Track, **Add Animation Notify Events** or **Add Animation Sync Markers** to a Notify Track. 
+
+![](AddMarkers.png)
+
+### Get Bone Pose
+
+![](GetPose.png)
+
+To get Bone Transform data you can use the **Get Bone Pose for Frame** or **Get Bone Pose for Time** nodes which will return the Bone Transform for the provided Animation Sequence with the specified Bone Name. 
+Alternatively, you can use the **Get Bone Poses for Frame** or **Get Bone Poses for Time** to collect Bone Transforms for a specified array of Bone Names. 
+
+[REGION:note]
+When getting Bone Transforms the transform data will be in local space. If you need it to be in component space, you will need to manually convert the transforms. 
+[/REGION]
+
+### Helper Nodes
+
+![](HelperNodes.png)
+
+The Animation Blueprint Library nodes include several helper nodes including nodes for converting frame and time information (**Get Time at Frame** or **Get Frame at Time**). 
+Another node that is useful for getting information is the **Find Bone Path to Root** which will take an Animation Sequence and a Bone Name (typically the Root Bone) and output a list of Bone Names in a chain. 
+
diff --git a/Engine/Documentation/Source/Engine/Animation/AnimMontage/AnimMontage.INT.udn b/Engine/Documentation/Source/Engine/Animation/AnimMontage/AnimMontage.INT.udn
--- a/Engine/Documentation/Source/Engine/Animation/AnimMontage/AnimMontage.INT.udn
+++ b/Engine/Documentation/Source/Engine/Animation/AnimMontage/AnimMontage.INT.udn
@@ -1,409 +1,29 @@
-Availability:Public
-Title:AnimMontage
-Crumbs:
-Description:Info detailing what AnimMontages are and how they are used in UE4's animation system.
-Related: Engine/Content/FBX/Animations
-Related: Engine/Content/Types/SkeletalMeshes
-Related: Engine/Animation/PhysicallyDrivenAnimation
-Related: Engine/Animation/StateMachines
-Related: Gameplay/AI
-Related: Engine/Content/Tools/MayaRiggingTool
-Related: Engine/Animation/RootMotion
-Related: Gameplay/Input
-version: 4.14
-Parent:Engine/Animation
-Order:
-tags:Animation Montage
-tags:Animation
-type:overview
-tags:Assets
+Availability: Docs
+Crumbs: %ROOT%
+Title:Animation Montage
+Description:Animation Montages are a flexible tool that enables you to combine and selectively play animations that are contained in a single asset. 
+Type: Landing
+Version: 4.16
+Parent: Engine/Animation
+Order: 1
+Tags: Animation
+Tags: Animation Montage
+Tags: Animation Blending
 topic-image:Engine/Animation/Overview/MontageTopic.png
 
-[TOC(start:2 end:3)]
 
-[EXCERPT:Intro]
-**Animation Montages** (or, "Montages", for short) are a multipurpose tool that allows for a wide variety of animation effects, primarily related to exposing animation controls within code, or [Blueprint Visual Scripting](Engine/Blueprints). It can also be used to create a wide variety of animation effects, including intelligent loops of animation, logic-based animation switching, and much more.
-[/EXCERPT:Intro]
-
-Montages are animation assets that can be created, and exist, within the **Content Browser**. These assets can then be plugged directly into the AnimGraph in run-time, and you can modify any state of it. For example, you can jump to different sections, or you can re-link different sections. This is mostly for code driven animations (or one-off animations) such as melee attacks, allowing you to control triggers, stop when you want, or change between states (looping or non-looping). 
-
-Montages also provide a way to replicate [](Engine\Animation\RootMotion\) across networked games if your animation has Root Motion enabled. 
-
-## Montage Example
-
-![](Montage_Screen2.png)
-
-The image above is a melee attack with 3 sections, labeled _Start_, _Loop_, and _End_; the divisions between the sections are marked in yellow. When the player presses the attack button, the animation is played by the game, and it begins at the _Start_ section. The looping middle section, called _Loop_, immediately follows the _Start_ section and will be programmed to repeat indefinitely. If the player releases the attack button, the animation will stop looping and eventually reach the _End_ section to transition out of the animation.
-
-Some additional uses for Montages include:
-* The ability to play an animation from within an AnimBlueprint's **EventGraph**.
-* Chaining together a complex sequence of animations that you want to think of as a single animation.
-* Looping only a specific portion of animation (or animations) based on code, or a Blueprint script.
-* Handling event-based switching of multiple animations based on code, or a Blueprint script.
-* The ability to assign complex animation sequences to named slots that can be switched between in code, or with Blueprints.
-* Precise switching between various AnimSequences based on code, or with a Blueprint script.
-
-[REGION:note]
-A **Montage** applied to a character can also be seen on the [Animation Content Examples](Resources\ContentExamples\Animation) page.
-[/REGION]
-
-## Montage Properties
-
-The following is a breakdown of the Montage asset properties. These are available in **Persona** when looking at a Montage, and can also be accessed by right-clicking on a Montage within the **Content Browser**, choosing **Properties** from the context menu.
-
-![](MontageProperties.png)
-
-|Montage Properties||
-|---|---|
-|[REGION:tablesection]Montage[/REGION]||
-|**Blend In Time**| The amount of time at the beginning of Montage playback, during which the character will blend in from its current pose.|
-|**Blend Out Time**|The amount of time at the end of Montage playback, during which the character will blend back to its original pose.|
-|**Blend Out Trigger Time**| Time from the Sequence End to trigger blend out. < 0 uses BlendOutTime, so BlendOut finishes as the Montage Ends. >= 0 uses SequenceEnd-BlendOutTriggerTime to trigger blend out. |
-|[REGION:tablesection]Additive Settings[/REGION]||
-|**Preview Base Pose**|Sets a base preview pose used for additive Blend Spaces.|
-|[REGION:tablesection]Animation[/REGION]||
-|**Rate Scale**|A multiplier value for how fast the Montage will play back. Default is 1.0.|
-|**Skeleton**|Contains the skeleton associated with this Montage. Cannot be changed in the editor.|
-|[REGION:tablesection]Meta Data[/REGION]||
-|**Meta Data**| This is Meta Data that can be saved with the asset. The meta data is a Blueprintable class derived from the Anim Meta Data class. This allows you to add custom meta data to animation assets (Anim Sequence, Anim Montage, Anim Composite, and Blendspace are all supported). You can query the data from an animation asset in C++ by calling the `GetMetaData()` method, or by using `GetSectionMetaData()` on Anim Montages. |
-|[REGION:tablesection]Thumbnail[/REGION]||
-|**Orbit Pitch**| The pitch of the orbit camera around the asset. |
-|**Orbit Yaw**| The yaw of the orbit camera around the asset. |
-|**Orbit Zoom**| The offset from the bounds sphere distance from the asset.|
-
-## Montage UI
-
-When looking at a Montage in Persona, it is useful to know what each area is, and what it does:
-
-[REGION:raw]
-![](MontageUI.png)
-[/REGION]
-
-1. [Montage Area](#MontageArea)
-1. [Sections Area](#SectionsArea)
-1. [Element Timing Area](#elementtimingarea)
-1. [Notifies Area](#NotifiesArea)
-1. [Curves Area](#CurvesArea)
-
-## Montage Area
-
-The Montage Area breaks down as follows:
-
-[REGION:raw]
-![](MontageArea.png)
-[/REGION]
-
-1. **Section Track** - Shows any [Sections](#Sections) that have been defined for this Montage. Sections can be dragged to different positions along the timeline with the left mouse button.
-1. **Slot Track** - Shows the current [Slot](#Slots), along with the Slot name on the right. This Slot may be filled with as many animations as you desire; they will be played back in order. Notice that multiple animations will have alternating positions in the Slot Track - top, then bottom, then top again, and so on. This is to help you differentiate between different animations.
-
-[REGION:note]
-You may have as many **Slot Tracks** as you like within a single **Montage**, each with their own name, containing their own unique animations. However, you may only have one **Section Track** for each Montage.
-[/REGION]
-
-### Sections
-
-Montage Sections provide a way to break a Slot into multiple portions of animation. Each Section has a name and location in the Slot's timeline. Using the name, you can either jump directly to a particular Section, or queue it to play next, when the current segment is complete. In Blueprint, you can query the current Section, jump to a Section, or set the next Section that will play.
-
-It may help to think of Sections like songs on a music playlist, with Slots being the albums. Just like with many modern media players, you can choose which song will play next, when the current one finishes, or just jump to the one you want to hear right now.
-
-Sections are created by right-clicking on the Section track and choosing **New Montage Section**. 
-
-
-### Slots
-
-[EXCERPT:Slots]
-Within a Montage, a **Slot** is just a single track that can hold any number of animations. you can pick any name for a Slot and then use that name to blend to those specific animations. A great example is having a character with a weapon reload animation. You may have different versions of the reload for when the player is standing, crouching, and lying prone. As long as all 3 of the animations use the same timing, you could place each one within a separate Slot in your Montage; the Slots could be named Standing, Crouching, and Prone. In your Animation Blueprint's **AnimGraph**, you can use the Slot node to determine which one you want to play, based on your character's current state. When they are standing, you can use the result of the animation in the Standing Slot. When they are prone, you can see the result of the Prone Slot.
-
-It is an important point to keep in mind that although much of your Montage control will take place in the Animation Blueprint's **Event Graph**, Slots are actually handled within the **Anim Graph**. This is done by way of the Slot node, which takes in the name for a Slot. By positioning this node at a strategic point along your AnimGraph's execution, you can have multiple Montages that utilize the same Slot name.
-[/EXCERPT:Slots]
-
-<!---------------- Branch moved to Notify page.
-### Branch Points
-
-Branch Points allow you to create events that coincide with animation playback. These Branch Point Events can be used in code or Blueprint to cause other things to happen, but specifically, Branch Points are useful for switching to other animation Sections within a Montage. 
-
-If you are already familiar with [Animation Notifies](Engine/Animation/Sequences\Notifies), you may notice a strong similarity between the two systems, as they both expose events that can be leveraged in script. The key difference is that Notifies are asynchronous, while Branch Points are synchronous. What this means to the end user is that Branch Points come with a much higher degree of precision for where they will take place along the animation timeline.
-
-High precision is important when you need to jump to a specific animation at a very precise moment in time. While you could use a Notify to do the same job, the asynchronous nature of Notifies means that the Notify Event could be fired at the incorrect animation frame, which can lead to hitches and jumps in your motion.
-
-[REGION:warning]
-Due to their synchronous nature and the precision resulting from it, Branch Points are more performance expensive than Notifies. You should only use them when an event must be fired at a precise moment along the animation timeline, such as jumping directly to another animation that matches up frame-to-frame. If being off by a frame (or some percentage of one) is not important, you should use Notifies instead.
-[/REGION]
-
--------------------->
-
-## Sections Area
-
-The Sections area is where you can establish relationships between the Sections you define in the Montage Area. For instance, you may want a certain Section of animation (or group of Sections) to play in a specific sequence, or even loop. 
-
-[REGION:raw]
-![](SectionsArea.png)
-[/REGION]
-
-
-1. **Create Default and Clear Buttons** - **Create Default** creates default associations between all Sections, stringing them together one after the other. **Clear** wipes out all associations.
-1. **Section Buttons** - In this area, you will see a button for each of the Sections you define in the Montage area. By choosing an existing Section, then clicking one of these buttons, you associate the Section for that button with the selected track. For instance, in the image above, we have associated Swing2 with Swing1. The order actually goes Swing1, Swing2, and then Swing1 again, which causes a loop. See the [Looping](#Looping) section below for details.
-1. **Section Association Tracks** - This is where you can visualize and preview the relationships between animation Sections. By clicking on the **Preview** buttons, you can see the result of each individual track, or click the **Preview All Sections** button, and see all Sections playing back, one after the other.
-
-### Looping
-
-Sections can be set up to loop indefinitely, which is extremely useful for any action that you need to repeat. By associating the same Section more than once in a Section Association Track, you cause that association to run in a loop. The section will turn blue to show this. As an example, consider an animation in which a character is reloading a shotgun, one shell at a time. You could take just the section in which the character inserts a shell, and loop it. Then, by using Notifies, you could create Notify Events in Blueprint that increment the ammo count each time the animation plays through. Once that count reaches a set number (full ammo), you could then switch to an animation of the character closing the receiver and returning to idle.
-
-## Element Timing Area
-
-The Element Timing area takes information from the Montage area and the Notifies area to help with the timing of different sections.
-
-![](ElementTiming1.png)
-
-Each node in the track is given a number denoting the trigger order of that object in the entire montage. It is colored to match what they describe. The default colors are:
-* Red - Notifies, Notify states and notify state ending markers
-* Light Blue - Branching points
-* Green - Montage sections
-
-[REGION:note]
-These colors can be edited in the **Animation Editor** preferences, which can be found in **Editor Preferences**.
-[/REGION]
-
-Hovering over any of the nodes in the track will cause a tooltip to appear that explains what that node represents.
-
-![](ElementTiming2.png)
-
-There is a dropdown menu on the right side of the track to toggle the visibility of different elements.
-
-![](ElementTiming3.png)
-
-## Notifies Area
-
-![](NotifiesArea.png)
-
-[INCLUDE:Engine/Animation/Sequences\Notifies#Intro]
-
-[REGION:note]
-Please see [](Engine/Animation/Sequences\Notifies) for more information.
-[/REGION]
-
-## Curves Area
-
-![](CurvesArea.png)
-
-[INCLUDE:Engine/Animation/Sequences\Curves#Intro]
-
-[REGION:note]
-Please see [](Engine/Animation/Sequences\Curves) for more information.
-[/REGION]
-
-
-## Creating Variant Montages
-**Child Anim Montages** provide a convenient way to create animation variants without affecting gameplay, and without needing to recreate and maintain notifications (or other details) across multiple assets. An example use case for this might be varying the auto-attack animations on units in a real-time strategy (RTS) game, or changing a character's idle animation in a role-playing game (RPG), or an action game, based on what kind of weapons or armor they are holding or wearing. A Child Animation Montage can be created from an existing Animation Montage via the "Create Child Montage" context-menu option.
-
-![](CreateChildMontage.png)
-
-The new montage will have the same name as the original, but with "_Child" added to the end.
-
-![](ChildAnimMontageCB.png)
-
-[REGION:note]
-Although Child Anim Montages are not a separate class from standard Anim Montages (both are `UAnimMontage`), Child Anim Montages cannot be used as parents for other Child Anim Montages.
-[/REGION]
-
-When editing a Child Anim Montage, all editable parts of the interface will be dimmed and deactivated (except for the Slot Track). The Slot Track is in a special mode which enables animations to be overridden. New animations cannot be added, and the details of an animation (such as play rate) will not be editable. This helps to create varied animations with minimal effort, without affecting gameplay. Since you can randomly pick between the original montage and its children freely, knowing that the timing of the animation (and all notification) events, will be identical. When editing a Child Anim Montage, animations that have not been changed from the parent will be shown in green, and animations that have been overridden will be shown in yellow.
-
-![](AnimMontageSlot.png)
-
-![](ChildAnimMontageSlot.png)
-
-The animation substitution in the picture above is indicated by the yellow box in the Child Anim Montage's editor window. 
-
-It is possible to substitute any animation for any other, even if they are of different lengths. The engine will automatically adjust the play rate, start time, and end time, of the replacement animation so that it uses the equivalent portion and runs in the same time as the original. For example, if a one-second animation with a start time of 0.5, and an end time of 1.0, is replaced with a two-second animation, the two-second animation will play with a start time of 1.0, with an end time of 2.0, and will play at double the speed of the one-second animation. The replacement animation will always fit perfectly into the original slot.
-
-## Playing Montages
-
-Montages can be played in a number of different ways: 
-
-* You can assign them as the **Anim to Play** for a **Skeletal Mesh**.
-* You can call them directly through a **Blueprint** script, or C++.
-* You can incorporate them into the **Animation Graph** as part of your **Animation Blueprint**.
-
-[REGION:note]
-Please see [](Engine/Animation/AnimHowTo/AdditiveAnimations) for an example of how a Montage can be used inside of your Animation Blueprints. This example illustrates a character that can perform two animations at the same time independently of one another (in this case, the character can move and shoot at the same time).
-[/REGION]
-
-### Playing Multiple Montages
-
-Prior to Unreal Engine version 4.6, you could only play a single Montage at a time. However, as of 4.6 or later, it is now possible to play as many Montages at you would like through the use of **Anim Slots** and **Anim Groups**.
-
-Here's a quick rundown of how it works: 
-
-1. Animations are organized into Anim Groups. You can play a single Montage per group. If you want to play another Montage at the same time, you can create a new Anim Group for it.
-1. Anim Groups and Anim Slots are centralized into the **Anim Slot Manager** tab inside **Persona** (see [Anim Slot Manager](#animslotmanager) below).
-1. You can use the drop-down User Interface (UI) to select between valid slots.
-1. When creating a new Montage, Slot nodes start out with a single **Default Slot** so that they are immediately functional.
-
-### Anim Slot Manager
-
-The **Anim Slot Manager** is where you will manage your Anim Slot and Anim Group names. It displays a hierarchical view and lets you make changes. You can open the Anim Slot Manager inside of a Montage in **Persona** by clicking the magnifying glass icon (pictured below).
-
-![](OpenAnimSlotManger.png)
-
-This will open the Anim Slot Manager.
-
-![](AnimSlotManager.png)
-
-Groups and Slots are saved in the **Skeleton** asset so that all Animation Blueprints and Montages from the same family can use the same settings. The **Save** button inside of the Anim Slot Manager is there so that you can easily save the changes to the Skeleton asset whenever you make any changes.
-
-When right-clicking on a Slot name, a context menu is presented, which will also allow you to move it to a different group (if you desire). This makes it easier to manage the Slots you have created and place and move them into different groups and arrange them as you see fit.
-
-![](SlotClicking.png)
-
-### Editing Slots on Montage Tracks
-
-Below is an example of the **Montage Editor** inside **Persona**, the drop-down box (highlighted) lists the available Slots for that track.
-
-![](SlotExample1.png)
-
-If you need to create a new Slot or Group, clicking the magnifying glass icon will take you to the **Anim Slot Manager**.
-
-The status bar across the top will tell you which Group the Montage belongs to. This is determined by the first Track's Slot.
-
-![](SlotExample2.png)
-
-If you have multiple Slots, the status bar will also tell you if there are any issues with the setup. For example, if you mix and match Slots that belong to different groups, a warning sign will appear, allowing you to see more information about the problem.
-
-![](SlotExample3.png)
-
-### Editing Slots in Animation Blueprints
-
-Inside the AnimGraph of an Animation Blueprint, you can add a **Slot** node that, when selected, allows you to see the same options as shown in the Montage Track view. 
-
-![](SlotNode.png)
-
-From here, you can use the drop-down box to pick an existing Slot name, and the magnifying glass icon, to open the Anim Slot Manager. 
-
-###Final Notes
-
-* This is backwards compatible with previous setups. You should not notice any changes... other than your existing Slots being part of a "Default Group". 
-* The Montage Blueprint and C++ API have been enhanced to take Montage asset references. This is how you can control different Montages playing at the same time. For example, see the **Montage Set Play Rate** function. 
-
-
-<!-------------------- Example removed. This should be a HowTo.
-
-## Montage Practical Example
-
-
-In this example, we want to have a character that can freely run in all directions, with an attack animation that only plays on the upper body. This attack will have multiple animations that could take place during its course. This is a perfect way to show the assembly of a Montage, as well as how to control it in the Event Graph and blend into it within the AnimGraph.
-
-However, there are few things we have set up in advance:
-
-* We already have a State Machine defining locomotion. This is just like the one used in the Third Person Project Template.
-* We have several animations provided by an artist that we would like to string together to make the Montage.
-* We have created a [Character Blueprint Class](Engine\Blueprints\UserGuide\Types\ClassBlueprint\Creation) that we can get input information from. 
- * For this example: an **IsAttacking** Boolean that is set to TRUE when the **Left Mouse Button** is pressed and set to FALSE when it is released. 
-  ![](CharacterBlueprint.png)
-
-### Creating the Montage
-
-Making a Montage is as easy as **Right-clicking** in the **Content Browser** and choosing **Animation > Animation Montage**. You can also **Right-click** on an existing Animation Sequence and choose **Create Montage** from the context menu. This will automatically create a new Montage with the selected AnimSequence already set up in the default Slot.
-
-![](MakeMontage.png)
-
-### Montage Setup
-
-Our first order of business was to name our Slot. We only need one in this Montage, and since we want our attack to only affect the upper body, the name Upper Body seemed perfect. We then drag/dropped the animations we would need into this slot. The animations we have do the following:
-
-* Right-to-left hammer swing
-* Left-to-right hammer swing
-* Go from the end of the right-to-left swing back to Idle
-* Go from the end of the left-to-right swing back to Idle
-
-[REGION:note]
-The two swinging animations both end on the same pose as the other begins. This means that the two animations can play back to back in a loop and the character will seamlessly swing the hammer back and forth.
+[REGION:banner]
+![](MontageBanner.png)
 [/REGION]
 
-While order is not extremely important, having the first two animations at the beginning and back to back will simplify things later.
-
-![](MontageSetup.png)
-
-
-### Section Creation
-
-Our next step is to section out the animations in our Montage so that we can query them and call them up when needed in our Blueprint code. This is as easy as **Right-clicking** and choosing **Add New Montage Section** while clicking on the Section Track at the top of the Montage Area. 
-
-We used fairly straightforward names for each Section. Note that we replaced the Default section that came along with the Montage (made a new one and deleted the Default by **Right-clicking** and choosing **Delete Montage Section**). You can drag these Sections along the Section Track if you need to, and you will notice that they snap a bit at the boundary between two animation segments when you release the mouse. Use this to your advantage.
-
-![](AddNewSections.png)
-
-
-
-### Defining Section Relationships
-
-
-Now that our Sections are created, we can now define any special relationships between them using the Sections Area. For instance, we can define a relationship between the Swing1 and Swing2 sections so that they play back to back in a loop. This is very useful for us. We start by clicking the **Clear** button to wipe out any default relationships. Then it is just a matter of selecting the track with Swing1 in it, and clicking the green button labeled Swing2 near the top of the Sections Area. This will remove the Swing2 track and add Swing2 to the Swing1 track.
-
-![](Swing2Track.png)
-
-If we repeat the process, clicking the new Swing2 segment and then clicking the Swing1 button, the system detects that you are creating a loop and the track turns blue. This means that the Swing1 and Swing2 segments are considered a looping Section. They will play back to back, repeating indefinitely.
-
-![](LoopingTracks.png)
-
-### Setting Up Branch Points
-
-We will now set up some Branch Points to test whether to continue with the loop or jump to one of two possible endings for our attack animation. All we have to do is **Right-click** in the Branch Point Track and choose **New Branch Point**. We chose the names Swing_1_End and Swing_2_End for our Branch Points. We also zoomed very closely in with the mouse wheel to make sure that each one fires right at the last possible moment of its corresponding section. This means it is placed very slightly to the left of the border between the two.
-
-![](BranchPointsSetUp.png)
-
-
-
-[REGION:note]
-A similar system could be set up in which you use Notifies instead of Branch Points, but they would need to fire a little earlier along the timeline, and your Animation Blueprint Event Graph would be queuing the appropriate ending animation using the **Montage Set Next Section** node, rather than doing a direct switch with the **Montage Jump to Section** node, as we will be doing. Doing it with Notifies would be slightly cheaper on performance, as Notifies are asynchronous. Our example is academic, but we wanted to make sure to mention the other method, as well.
-[/REGION]
-
-
-### Setting Up the AnimGraph
-
-At this point, our Montage is all set up. We now need to establish our AnimGraph so that it can read in the result of our Montage. This is a pretty simple process, but it does cause us to have to think carefully about how we proceed. Our AnimGraph starts off looking something like this, where we are only seeing the result of our State Machine:
-
-![](StateMachineResult.png)
-
-
-We only want the Montage to play back from the Spine_01 bone (the waist) up, so we will be using a **Layered Blend per Bone** node. We add a Blend Pose to it set the weight to 1. We also associate this Blend Pose with Spine_01 in the properties for the node. We then bring in a **Slot** node and set it to UpperBody, the name of our slot. But now we run into a problem:
-
-![](SlotAndBlend.png)
-
-The Slot node needs a Source connection to fall back on once the Montage is done playing. Without it, the character will return to the T-pose from the waist up after the Montage completes. However, we cannot connect the State Machine to both the Base Pose of the **Layered Blend per Bone** as well as the Source for the **Slot** node. The solution? Use a **Cache** node! We can store the result of the State Machine into a Cache node, and then create **Cached Pose** nodes to connect to both our required inputs. This is somewhat similar to storing the result of the State Machine into a variable so you can use it in multiple places. For this example, we named the Cache LocoCache.
-
-![](CachedLocomotion.png)
-
-Our AnimGraph is done. As soon as the UpperBody Slot node receives data from any AnimMontage (you can use any Montage, so long as it has a Slot named UpperBody), it will blend it in. As soon as it is no longer receiving data, it will fall back on the result of the State Machine.
-
-
-### Setting Up the Event Graph
-
-Our Event Graph setup is very basic. Using the **Event Blueprint Update Animation** node and dragging off the out pin of a **Get Player Character** node, we can **Cast To** our **Character Blueprint** (MyCharacter in this example) to access the variables and functions from that Blueprint (see [Blueprint Communications](Engine\Blueprints\UserGuide\BlueprintComms) for more info on communicating between Blueprints). 
+[EXCERPT:Intro]
+**Animation Montages** (or, "Montages", for short) enable a wide variety of animation effects, primarily related to exposing animation controls within [Blueprint Visual Scripting](Engine/Blueprints) or through code. They can also be used to create a wide variety of animation effects, including intelligent loops of animation, logic-based animation switching, and much more.
+[/EXCERPT:Intro]
 
-_Click for full view_
-[REGION:lightbox]
-[![](PlayMontageEventGraph.png)(w:1020)](PlayMontageEventGraph.png)
+[REGION:topics half]
+%Engine\Animation\AnimMontage\Overview:topic%
+%Engine\Animation\AnimMontage\Creation:topic%
+%Engine\Animation\AnimMontage\Editor:topic%
+%Engine\Animation\AnimMontage\UserGuide:topic%
 [/REGION]
 
-The first thing we do is check if the **IsAttacking** variable from our Character Blueprint is TRUE and if it is, we then check if the Montage is already playing. If the Montage is playing, we do not want to play it again; however if it is not playing, then we play the Montage. This prevents the system from restarting the animation in the middle of playback, which would look bad. 
-
-Dragging off the **IsAttacking** node which extends from the **Cast To** node, we see if the mouse button is still down, and if not, we jump to the appropriate ending animation, depending on which half of the loop is playing. Creating the Branch Point Events is done by **right-clicking** and choosing the appropriate event under **Add Montage Branching Point Event**.
-
-![](EndLoopEventGraph.png)
-
-And that is it! If we compile, we will now see that the character continues to swing as long as the mouse button is down, and performs an intelligent end animation when the mouse button is released!
-
-
-[OBJECT:EmbeddedVideo]
-[PARAMLITERAL:width]
-640
-[/PARAMLITERAL]
-[PARAMLITERAL:height]
-360
-[/PARAMLITERAL]
-[PARAMLITERAL:videoid]
-6eMxjbkn6hs
-[/PARAMLITERAL]
-[/OBJECT]
-
--------------------->
diff --git a/Engine/Documentation/Source/Engine/Animation/AnimMontage/Creation/MontageCreation.INT.udn b/Engine/Documentation/Source/Engine/Animation/AnimMontage/Creation/MontageCreation.INT.udn
new file mode 100644
--- /dev/null
+++ b/Engine/Documentation/Source/Engine/Animation/AnimMontage/Creation/MontageCreation.INT.udn
@@ -0,0 +1,63 @@
+Availability: Docs
+Crumbs: %ROOT%, Engine, Engine/Animation, Engine\Animation\AnimMontage
+Title: Creating an Animation Montage
+Description:Overview of creating an Animation Montage for use within Blueprints and Animation Blueprints. 
+Type: how-to
+SkillLevel: Beginner
+Version: 4.16
+Parent: Engine\Animation\AnimMontage
+Order: 1
+Tags: Animation
+Tags: Animation Montage
+Tags: Animation Assets
+topic-image:CreateTopic.png
+Related: Engine/Content/Types/SkeletalMeshes
+Related: Engine/Animation/Sequences
+Related: Engine/Animation/Persona/Modes/Animation
+Related: Engine/Animation/StateMachines
+Related: Engine/Animation/AnimBlueprints
+
+This page covers the different ways in which you can create an **Animation Montage** using animation assets. 
+
+## Steps
+
+1. In the **Content Browser** of your project, click the **Add New** button (or right-click in the **Content Browser**) and choose **Animation** > **Animation Montage**.
+
+	![](AddNewButton.png)
+
+1. Select the **Skeleton** asset that you want to use for the **Animation Montage**. 
+
+	![](PickSkeleton.png)
+
+1. Enter a name for your newly created **Animation Montage**. 
+
+	![](NamedMontage.png)
+
+1. Double-click to open your new Animation Montage to begin editing. 	
+	
+
+## End Result
+
+After creating your Animation Montage, you can open it and it will appear empty (below). 
+
+![](NewMontageOpened.png)
+
+From here you can begin [editing your Animation Montage](Engine\Animation\AnimMontage\Editor) asset to reach your desired effect. 
+
+You can also right-click any **Animation Sequence** asset and create an Animation Montage from the asset. 
+
+![](RightClickExample.png)
+
+When opening a Montage you have created from the right-click method, the animation will automatically be added to the Montage. 
+
+![](RightClickExampleOpen.png)
+
+
+
+
+
+
+
+
+
+
diff --git a/Engine/Documentation/Source/Engine/Animation/AnimMontage/Editor/MontageEditor.INT.udn b/Engine/Documentation/Source/Engine/Animation/AnimMontage/Editor/MontageEditor.INT.udn
new file mode 100644
--- /dev/null
+++ b/Engine/Documentation/Source/Engine/Animation/AnimMontage/Editor/MontageEditor.INT.udn
@@ -0,0 +1,277 @@
+Availability: Docs
+Crumbs: %ROOT%, Engine, Engine/Animation, Engine\Animation\AnimMontage
+Title: Editing an Animation Montage
+Description:Outlines how you can edit and set up the parameters that drive your Animation Montage asset.
+Type: Reference
+SkillLevel: Intermediate
+Version: 4.16
+Parent: Engine\Animation\AnimMontage
+Order: 1
+Tags: Animation
+Tags: Animation Montage
+Tags: Animation Assets
+topic-image:EditTopic.png
+Related: Engine/Content/Types/SkeletalMeshes
+Related: Engine/Animation/Sequences
+Related: Engine/Animation/Persona/Modes/Animation
+Related: Engine/Animation/StateMachines
+Related: Engine/Animation/AnimBlueprints
+
+[REGION:prereq]
+If you are new to **Animation Montages**, it is recommended that you read the [](Engine\Animation\AnimMontage\Overview) before proceeding. 
+[/REGION]
+
+[TOC (start:2 end:2)]
+
+This page covers how to set up an **Animation Montage** and the options within the Montage Editor. 
+
+## Adding Animations to a Montage
+
+To get started, the first thing you will need to do is add the **Animation Sequences** that will comprise the Montage inside the **Montage** area. 
+
+To do this, inside the **Asset Browser** window drag-and-drop an Animation Sequence into the Montage **Slot** track. 
+
+![](DragAndDropAnim.png)
+
+Once you drop the animation on the Slot track, the animation will be added. 
+
+![](AnimationAdded.png)
+
+You can drag-and-drop additional animations that you want to add to the Montage onto the Slot track and they will be added sequentially. 
+
+![](MontageAnimationsSet.png)
+
+Animations will automatically alternate between the first and second Slot tracks to help distinguish where each animation starts and stops. 
+
+You can also left-click and drag-and-drop the animations on the Slot track to re-order them. 
+
+![](ReOrderAnimations.png)
+
+Right-clicking on an animation in the Montage area provides the following context menu:
+
+![](MontageAreaContextMenu.png)
+
+|---|---|
+|**Delete Segment**| Deletes the selected animation. |
+|**Open Asset**| Opens the selected animation. |
+|**New Montage Section**| Adds a new Montage Section. |
+|**New Slot**| Adds a new Slot.  |
+|**Delete Slot**| Deletes selected Slot. |
+|**Duplicate Slot**| Duplicates the Slot and creates a new Slot. |
+|**Set Elements to...**| Sets all Montage elements (Sections, Notifies) to **Absolute**, **Relative** or **Proportional** values.  |
+
+## Creating Montage Sections
+
+In order to define relationships between Sections in the Sections area, you must first create the Sections that you want to work with inside the Montage area. 
+
+To create a Section, right-click on the Slot track or the Section track and select **New Montage Section**.
+
+![](SelectNewMontageSection.png)
+
+When prompted, enter your desired Section name.
+
+![](EnterSectionName.png)
+
+When adding a new Section, it will be added to the Sections track as well as inside the Section area.
+
+![](SectionAdded.png)
+
+[REGION:note]
+By default, all Animation Montages include a **Default** Section that, when used by itself, will play the entire Montage. 
+[/REGION]
+
+You can move Sections around by left-clicking and dragging the Section to the desired location. 
+
+![](ExampleSectionMovement.png)
+
+[REGION:caption]
+Above we have moved our **Start** Section to the start of the Montage and moved the **Default** Section over or we can delete it.
+[/REGION]
+
+Below we add two more Sections to our Montage for a reloading Montage with a **Start**, **Loop** and **End** to correspond with our animations. 
+
+![](MontageSectionComplete.png)
+
+## Creating Montage Slots
+
+With Slots, you can take a group of related animations (such as reloading while standing, reloading while crouched and reloading while prone) and assign them to different Slots within the Montage. 
+Through Blueprint or C++, you can then tell one of the Slots to play the animations it contains based on whether a character is standing, crouched or prone, as they all exist within the same Reload Montage. 
+
+To add a Slot in the Montage area, you can use the right-click context menu and select **New Slot**. 
+
+![](NewSlot.png)
+
+This will create an empty Slot track to which you can add new animations to. 
+
+![](EmptySlot.png)
+
+Selecting **Duplicate Slot** will copy all animation segments into a new Slot track. 
+
+![](DuplicatedSlot.png)
+
+When adding or duplicating a Slot, you will see a warning at the top of the Montage indicating that the current Slot name is already in use and that all Slots must be unique. 
+This because the new Slot track will be set by default to use the **Default Group** and **Default Slot** that are created when the Montage was created. 
+You can change which **Anim Group** you are in and **Slot** you are using from the drop-down menu to the right of the Slot track. 
+
+![](ChangeSlotButton.png)
+
+You may need to create a new Slot, which you can do from the **Anim Slot Manager** by clicking the ![](AnimSlotButton.png) button. This button will open the Anim Slot Manager panel.
+
+![](AnimSlotManager.png)
+
+Click the **Add Slot** button (or **Add Group** button based on your needs) and assign a new unique name. 
+
+![](ProneSlotAdded.png)
+
+[REGION:caption]
+Above, we created a new Slot called **Prone** which we can add our prone reloading animations to.
+[/REGION]
+
+You can right-click a Slot or Group in the Anim Slot Manager to access a context menu with additional options. 
+
+![](SlotManagerContextMenu.png)
+
+|---|---|
+|**Set Slot**| Enables you to move the selected Slot to a different Anim Group.  |
+|**Rename Slot**| Enables you to rename the currently selected Slot. |
+|**Delete Slot**| Deletes the currently selected Slot. |
+|**Add Slot**| Adds a new Slot.  |
+|**Add Group**| Adds a new Anim Group. |
+
+Below we have renamed the **Default Slot** to **Standing** and assigned both the Standing and Prone Slots in the Montage area. 
+
+![](AssignedNewSlots.png)
+
+
+## Using Multiple Slots
+
+To get the best results when using multiple Slots, ensure that the timing of each of your animations is the same. 
+In the image below, the arrows indicate that the timing of our Prone Reload animations does not match up with our Standing Reloading animations which will cause problems when we play back the Prone Slot of our Montage.
+
+![](BadTiming.png)
+
+If we were to call the Loop section of this Montage, while it would play fine in our Standing Reload, in our Prone Reload it would start midway through the animation and would also start to loop reloading before actually reaching the Loop Section. 
+Ideally, when creating our reloading animations, we would use the same timing and frame range so that when we call the Loop Section in either the Standing or Prone Slots, it starts at the start of the animation without any issues.
+
+![](BetterTiming.png)
+
+[REGION:caption]
+Above, our timing of our Standing and Prone Reloading animations are the same which will give us better results when we call either Slot to play. 
+[/REGION]
+
+
+## Section Relationships
+
+After you have created the Sections for your Montage in the Montage area, you can define the Sections' relationships to each other in the Sections area.
+
+![](SectionsDefault.png)
+
+By default, all Sections are set to play sequentially, indicated by the single Section track containing each of the Sections. 
+
+![](SingleSectionTrack.png)
+
+Pressing the **Clear** button will assign a single track for each Section you have created. 
+
+![](IndividualTracks.png)
+
+[REGION:tip]
+Pressing **Create Default** will return all Sections to the single Section track. 
+[/REGION]
+
+In our example, we have a Start, a Loop, and an End and we want the Montage to play the Start section immediately, followed by the Loop section.
+Once we get to the Loop section, we want to keep looping that Section until the End Section is told to play through Blueprint or C++.
+
+To do this, we first click on the Start Section in the Section Track. 
+
+![](StartClicked.png)
+
+We can now define what Section should follow the Start Section, so we click the Loop button at the top of the Section area.
+
+![](LoopClicked.png)
+
+When we do this, Loop is added and follows Start on the first track. When we click Loop on the track, then click the Loop button at the top, Loop will turn blue. 
+
+![](LoopNowLoops.png)
+
+Since we associated Loop with itself, the system knows that we want this Section to loop. 
+
+We can preview this by clicking the **Preview** button on the track we want to preview and the **Play** button from the toolbar at the bottom. 
+
+[OBJECT:EmbeddedVideo]
+	[PARAMLITERAL:videoid]
+	QovIb9Tslr8
+	[/PARAMLITERAL]
+	[PARAMLITERAL:width]
+	100
+	[/PARAMLITERAL]
+	[PARAMLITERAL:height]
+	56.25
+	[/PARAMLITERAL]
+	[PARAMLITERAL:units]
+	%
+	[/PARAMLITERAL]
+	[PARAMLITERAL:type]
+	simple
+	[/PARAMLITERAL]
+[/OBJECT]
+
+In this example, whenever we call the Start section to play in either the Standing or Prone Slots, it will play the assigned Start animation before moving onto and looping the assigned Loop animation. 
+This enables us to play the animation of reloading as many times as we need to based on gameplay variables such as whether the player is performing a full or partial reload. 
+We can then call the End Section to play and return to an idle pose whenever we have determined that we are finished reloading. 
+
+To remove a relationship for a Section, click the **X** button next to the item.
+
+![](DeleteSection.png)
+
+Removing the relationship for the Loop section causes it to no longer loop.
+
+![](NoLoops.png)
+
+## Additional Options
+
+The **Element Timing** area is used for informational purposes and is covered on the [](Engine\Animation\AnimMontage\Overview) page, in short, it enables you to see at a glance the sequential playback order of any Sections, Notifies or other Events.
+[Animation Notifies](Engine/Animation/Sequences/Notifies) enable you to set up events to occur at specific points in the animation while [](Engine/Animation/Sequences/Curves) provide a way to change the value of a Material parameter or a Morph Target while the animation is playing back. 
+
+The [Timeline](https://docs.unrealengine.com/latest/INT/Engine/Animation/Sequences/Editor/#timeline) at the bottom of the Montage Editor enables you to playback the Montage, step forward or backward through the Montage, loop the preview playback or create a new animation object by recording Bone manipulation from the viewport. 
+
+## Montage Properties
+
+The following is a breakdown of the Montage asset properties available from the [Asset Details](Engine/Animation/Persona/AnimAssetDetails) window:
+
+![](AssetDetails.png)
+
+|Montage Properties||
+|---|---|
+|[REGION:tablesection]Animation[/REGION]||
+|**Preview Pose Asset**| Points to a [Pose Asset](Engine/Animation/AnimPose) that can be used to preview any Morph Target or Curves. |
+|**Rate Scale**|A multiplier value for how fast the Montage will play back. Default is 1.0.|
+|**Skeleton**|Contains the skeleton associated with this Montage. Cannot be changed in the editor.|
+|**Parent Asset**|During cooking, this data will be used to bake out to Normal Asset. Cannot be changed in the editor.|
+|**Asset Mapping Table**|This is used when Parent Asset is set and works with Asset User Data mapping table. Cannot be changed in the editor.|
+|**Asset User Data**|Array of user defined data stored with the asset.|
+|[REGION:tablesection]Blend Option[/REGION]||
+|**Blend In**| The amount of time at the beginning of Montage playback, during which the character will blend in from its current pose. You can also set the [Blend Option](Engine/Animation/NonLinearBlends) or add a **Custom Curve**. |
+|**Blend Out**|The amount of time at the end of Montage playback, during which the character will blend back to its original pose. You can also set the [Blend Option](Engine/Animation/NonLinearBlends) or add a **Custom Curve**.|
+|**Blend Out Trigger Time**| Time from the Sequence End to trigger blend out. If the value specified is greater than zero, this means Blend Out will finish as the Montage ends. If the value specified is less than or equal to zero, this means it will take the end of the sequence minus the specified value to trigger the Blend Out. |
+|[REGION:tablesection]Sync Group[/REGION]||
+|**Sync Group**| Where you can define the [Sync Group](Engine/Animation/SyncGroups) for the Montage. |
+|**Sync Slot Index**| Defines the index to use for the Sync Group. |
+|[REGION:tablesection]Additive Settings[/REGION]||
+|**Preview Base Pose**|Sets a base preview pose used for additive Blend Spaces.|
+|[REGION:tablesection]Meta Data[/REGION]||
+|**Meta Data**| This is Meta Data that can be saved with the asset. The metadata is a Blueprintable class derived from the Anim MetaData class. This allows you to add custom metadata to animation assets (Anim Sequence, Anim Montage, Anim Composite, and Blendspace are all supported). You can query the data from an animation asset in C++ by calling the `GetMetaData()` method, or by using `GetSectionMetaData()` on Anim Montages. |
+|[REGION:tablesection]Thumbnail[/REGION]||
+|**Orbit Pitch**| The pitch of the orbit camera around the asset. |
+|**Orbit Yaw**| The yaw of the orbit camera around the asset. |
+|**Orbit Zoom**| The offset from the bounds sphere distance from the asset.|
+
+
+
+
+
+
+
+
+
+
+
