Copying //UE4/Dev-Documentation to //UE4/Samples-Main (Source: //UE4/Dev-Documentation @ 3623665)

diff --git a/Engine/Documentation/Source/Engine/Animation/AnimMontage/Overview/MontageOverview.INT.udn b/Engine/Documentation/Source/Engine/Animation/AnimMontage/Overview/MontageOverview.INT.udn
new file mode 100644
--- /dev/null
+++ b/Engine/Documentation/Source/Engine/Animation/AnimMontage/Overview/MontageOverview.INT.udn
@@ -0,0 +1,242 @@
+Availability: Docs
+Crumbs: %ROOT%, Engine, Engine/Animation, Engine\Animation\AnimMontage
+Title: Animation Montage Overview
+Description:Describes how Animation Montage assets work and what they can be used for.
+Type: Overview
+SkillLevel: Intermediate
+Version: 4.16
+Parent: Engine\Animation\AnimMontage
+Order: 1
+Tags: Animation
+Tags: Animation Montage
+Tags: Animation Assets
+topic-image:MontageTopic.png
+Related: Engine/Content/Types/SkeletalMeshes
+Related: Engine/Animation/Sequences
+Related: Engine/Animation/Persona/Modes/Animation
+Related: Engine/Animation/StateMachines
+Related: Engine/Animation/AnimBlueprints
+
+[TOC (start:2 end:2)]
+
+An **Animation Montage** (or **Montage** for short) provides a way for you to control an animation asset directly through **Blueprints** or C++ code.
+With an Animation Montage, you can combine several different [](Engine/Animation/Sequences) into a single asset that you can break up into **Sections** for playback either individually or in combination.
+You can also fire off **Events** within a Montage which can perform a variety of local or replicated tasks, such as playing sound cues or particles effects, changing player values like ammunition count, or even replicating  [](Engine\Animation\RootMotion\) in networked games, provided Root Motion is enabled on the animation.
+
+## Example Montage
+
+The following Animation Sequence of a character reloading a shotgun demonstrates a practical use case for Montages:
+
+[OBJECT:EmbeddedVideo]
+	[PARAMLITERAL:videoid]
+	7qtyXSreAzs
+	[/PARAMLITERAL]
+	[PARAMLITERAL:width]
+	100
+	[/PARAMLITERAL]
+	[PARAMLITERAL:height]
+	56.25
+	[/PARAMLITERAL]
+	[PARAMLITERAL:units]
+	%
+	[/PARAMLITERAL]
+	[PARAMLITERAL:type]
+	simple
+	[/PARAMLITERAL]
+[/OBJECT]
+
+In the video above, the character reaches for a clip, reloads and then returns to an idle pose. 
+While this is perfectly fine to use as a reload animation, a Montage will provide much more control. 
+For example, suppose you wanted the character to reload a variable amount of bullets instead always the same amount. 
+By splitting this animation up into three parts (grabbing a clip, reloading, and returning to idle) and combining those parts in a Montage, you could section them off and selectively play any or all of those sections. 
+
+Below is an example Montage used to handle the reloading of a shotgun, sectioned it off into three parts: _Start_, _Loop_, and _End_. 
+
+![](MontageExample.png)
+
+With Blueprint or C++, we can define that when the player presses a button, the animation is played, beginning at the _Start_ section.
+The looping middle section, called _Loop_, immediately follows the _Start_ section and is set within the Montage to repeat indefinitely.
+If the player releases the reload button, or they've filled the shotgun with bullets, the animation will stop looping and can be set to jump to the _End_ section to transition out of the animation as seen in the video below.
+
+[OBJECT:EmbeddedVideo]
+	[PARAMLITERAL:videoid]
+	uV16SOj_9-A
+	[/PARAMLITERAL]
+	[PARAMLITERAL:width]
+	100
+	[/PARAMLITERAL]
+	[PARAMLITERAL:height]
+	56.25
+	[/PARAMLITERAL]
+	[PARAMLITERAL:units]
+	%
+	[/PARAMLITERAL]
+	[PARAMLITERAL:type]
+	simple
+	[/PARAMLITERAL]
+[/OBJECT]
+
+The reload motion is a little bit faster than the baked sequence, however, we can adjust the play rate of our looping animation if desired. This gives us the ability to control when we want to break out of our reloading motion and return back to an idle pose. 
+A similar usage example could be that a character jumps in the air, loops an animation of being in the air or falling, then plays a landing sequence. The amount of time in the air could vary, and we want to be able to break out of that portion of the Montage at any point.
+
+Some additional uses for Montages include:
+* The ability to play an animation from within the **Event Graph** of an **Animation Blueprint**.
+* Chaining together a complex sequence of animations that you want to think of as a single animation.
+* Looping only a specific portion of animation (or animations) based on code, or a Blueprint script.
+* Handling event-based switching of multiple animations based on code, or a Blueprint script.
+* The ability to assign complex animation sequences to named slots that can be switched between in code, or with Blueprints.
+* Precise switching between various AnimSequences based on code, or with a Blueprint script.
+
+
+## Montage UI
+
+After [creating a Montage](Engine\Animation\AnimMontage\Creation) and opening it inside the [](Engine/Animation/Persona/Modes/Animation), you can use the **Asset Editor** section to define how your Montage functions. 
+
+![](MontageWindow.png)
+
+The Asset Editor panel of a Montage is broken up into the following primary areas: 
+
+1.	Montage
+1.	Sections
+1.	Element Timing
+1.	Notifies
+1.	Curves
+
+
+### Montage
+
+The Montage area breaks down as follows: 
+
+![](MontageArea.png)
+
+1.	**Sections** - Shows any Sections that have been created within the Montage.
+1.	**Slot** - Shows the animations that have been added to the current Slot.
+1.	**Group/Slot** - Shows the currently selected Anim Group and Slot.
+
+#### Sections
+
+Creating **Sections** provides a way to break a Slot into multiple portions of animation. Each Section has a name and location in the Slot's timeline.
+Using the name, you can either jump directly to a particular Section or queue it up to play next when the current section is complete.
+In **Blueprint**, you can query the current Section, jump to a Section, or set the next Section that will play.
+
+Another way to think of Sections is like songs in a music playlist, with Slots being the album.
+You may be currently playing a section (song), but you can queue up or jump to another section within the slot (album) when the current section finishes, or just jump to the one you want to play right now.
+
+#### Slot
+
+Within a Montage, a **Slot** is just a single track that can hold any number of animations. You can pick any name for a Slot and then use that name to blend to the animations within. 
+Going back to the reload Montage example, we could have different versions of the reload for when the player is standing versus when they are prone. 
+In this case, we could place animations for each within separate Slots in our Montage (below) and use the **Slot** node in the **AnimGraph** of our **Animation Blueprint** to determine which set of animations to use based on whether the character is standing or prone. 
+
+![](MultipleSlots.png)
+
+[REGION:caption]
+Above we use the **DefaultProne** slot for when our character is standing and reloading and the **Prone** slot for when our character is reloading while prone. 
+[/REGION]
+
+When multiple Slots are assigned, you can click the ![](PreviewButton.png) preview button for the Slot you want to preview in the editor. 
+
+[REGION:tip]
+For the best results when using multiple Slots, try to ensure that the timing of animations involved from end-to-end is the same length.
+[/REGION]
+
+#### Groups / Slots
+
+**Groups** are a collection of **Slots**, and you can only play a single Montage per group. If you want to play another Montage at the same time, you can create a new Group for it.
+Anim Groups and Slots are centralized within the **Anim Slot Manager** tab, which you can open by clicking the ![](AnimSlotManagerButton.png) icon, and you can use the drop-down menu to select between valid slots.
+When creating a new Montage, Slot nodes start out with a single **Default Slot** so that they are immediately functional.
+
+### Sections
+
+The **Sections** area is where you can establish relationships between the Sections you define in the Montage area. 
+
+![](SectionsUI.png)
+
+1.	**Create Default / Clear** - Create Default creates the default associations between all Sections, stringing them together one after another while Clear wipes out all associations. 
+1.	**Section Buttons** - In this area, you will see a button for each of the Sections you define in the Montage area that you can assign for Section Association. 
+1.	**Section Association Tracks** - Here you can visualize the relationships between animation Sections by clicking Preview All Sections or individual tracks by clicking Preview. 
+
+#### Looping
+
+Sections can be set up to loop indefinitely, which is extremely useful for any action that you need to repeat. 
+By associating the same Section more than once in a Section Association Track, you cause that association to run in a loop. 
+The section will turn blue to show this. As an example, consider the animation in which a character is reloading a shotgun, one shell at a time. 
+You could take just the section in which the character inserts a shell, and loop it. 
+Then, by using **Notifies**, you could create Notify Events in Blueprint that increment the ammo count each time the animation plays through. 
+Once that count reaches a set number (full ammo), you could then switch to an animation of the character closing the receiver and returning to idle.
+
+### Element Timing
+
+The **Element Timing** track takes information from the Montage and Notifies areas to help with the timing of different sections. 
+
+![](ElementTimingTrack.png)
+
+Each node in the track is given a number denoting the trigger order of that object within the Montage, and is color coded as follows:
+
+* **Red** - This denotes Notifies, Notify States and Notify State ending markers.
+* **Light Blue** - This denotes Branching Points.
+* ** Green** - This denotes a Montage Section.
+
+Clicking the ![](ElementOptions.png) button on the Element Timing track presents you with options for showing/hiding additional information. 
+
+### Notifies
+
+![](NotifiesSection.png)
+
+[Animation Notifies](Engine/Animation/Sequences/Notifies) (or **Notifies** for short) are a way to set up events to occur at specific points in the animation. 
+For example, in the image above, we have placed Notifies to play sounds in our Reload Montage for grabbing a clip each time the reload animation plays. 
+We also call a custom event called **SpawnClip** that we can use to spawn a mesh in the world representing an ammo clip. 
+Finally, we use a Branch event during the looping portion to determine if we should exit the loop and move on to the End section of our animation where the character returns to an idle. 
+
+### Curves
+
+![](CurvesArea.png)
+
+[](Engine/Animation/Sequences/Curves) provide a way to change the value of a Material parameter or a Morph Target while an animation is playing back. 
+Their workflow requires that you designate the asset you want to modify (either a Material or a Morph Target), name the curve accordingly, and then adjust keyframe values across the duration of an animation.
+
+## Playing Montages
+
+Montage assets can be played in a few different ways based on your needs. 
+If you want a Skeletal Mesh to play the Montage automatically, you can assign it as the **Anim to Play** for the Skeletal Mesh. 
+If, however, you only want your Montage to play when specific criteria are met, you may want to create **Blueprint** script that calls for the Montage or Sections within the Montage to play using the **Play Montage** node. 
+You can also tell Montages to play through C++ code based on your defined criteria. 
+
+[REGION:note]
+See the [](Engine\Animation\AnimMontage\UserGuide) page for more information. 
+[/REGION]
+
+## Child Montages
+**Child Montages** provide a convenient way to create animation variants without affecting gameplay, and without needing to recreate and maintain Notifies (or other details) across multiple assets. 
+An example use case for this might be varying the auto-attack animations on units in a real-time strategy (RTS) game, or changing a character's idle animation in a role-playing game (RPG), or an action game, based on what kind of weapons or armor they are holding or wearing. 
+A Child Animation Montage can be created from an existing Animation Montage via the **Create Child Montage** context-menu option when right-clicking on an Animation Montage in the **Content Browser**. 
+
+![](ChildMontageOption.png)
+
+The new montage will have the same name as the original, but with "_Child" added to the end.
+
+![](NewChildMontage.png)
+
+When opening a Child Montage, all elements within the Montage will be grayed out and text will appear at the top indicating that it is a child. 
+
+![](GrayedOut.png)
+
+The **Anim Segment** in the **Slot Track** can be right-clicked and overridden with a new animation. 
+
+![](RightClickChange.png)
+
+In the **Slot Track** it is possible to substitute any animation for any other, even if they are of different lengths. 
+The engine will automatically adjust the play rate, start time, and end time, of the replacement animation so that it uses the equivalent portion and runs at the same time as the original. 
+Animations that have not been changed from the parent will be shown in green and animations that have been overridden will be shown in yellow.
+
+![](ReplacedMontage.png)
+
+[REGION:caption]
+Above we have replaced our attack animation with a different attack animation. 
+[/REGION]
+
+[REGION:note]
+Child Montages cannot be used as parents for other Child Montages.
+[/REGION]
+
+
diff --git a/Engine/Documentation/Source/Engine/Animation/AnimMontage/UserGuide/MontageUserGuide.INT.udn b/Engine/Documentation/Source/Engine/Animation/AnimMontage/UserGuide/MontageUserGuide.INT.udn
new file mode 100644
--- /dev/null
+++ b/Engine/Documentation/Source/Engine/Animation/AnimMontage/UserGuide/MontageUserGuide.INT.udn
@@ -0,0 +1,62 @@
+Availability: Docs
+Crumbs: %ROOT%, Engine, Engine/Animation, Engine\Animation\AnimMontage
+Title: Using Animation Montages
+Description:Describes how Animation Montages can be played back at runtime.
+Type: Reference
+SkillLevel: Intermediate
+Version: 4.16
+Parent: Engine\Animation\AnimMontage
+Order: 1
+Tags: Animation
+Tags: Animation Montage
+Tags: Animation Assets
+topic-image:PlayMontage.png
+social-image:PlayMontage.png
+Related: Engine/Content/Types/SkeletalMeshes
+Related: Engine/Animation/Sequences
+Related: Engine/Animation/Persona/Modes/Animation
+Related: Engine/Animation/StateMachines
+Related: Engine/Animation/AnimBlueprints
+
+[TOC(start:2 end:2)]
+
+After [creating a Montage](Engine\Animation\AnimMontage\Creation) and [editing the Montage](Engine\Animation\AnimMontage\Editor), you will probably want to play your Montage at runtime. 
+You can have your Montage automatically start playing by assigning it to a **Skeletal Mesh** or, for more flexibility, you may want to set up a **Blueprint Script** or C++ code that will call your Montage or Sections within your Montage to start playing based on gameplay conditions. 
+
+
+## Using Anim to Play on a Skeletal Mesh
+
+In the **Details** panel for a **Skeletal Mesh**, you can set the **Animation Mode** to **Use Animation Asset** and the **Anim to Play** to your desired **Animation Montage**. 
+
+![](SkeletalMeshMontage.png)
+
+At run time, the Skeletal Mesh will playback your selected Animation Montage automatically. 
+
+## Using Blueprint's Play Montage Node
+
+You can call the **Play Montage** node within Blueprint (below we tell a Skeletal Mesh called **Mesh** to play a Montage when the character is clicked on with the Left Mouse Button). 
+
+![](BlueprintMontage.png)
+
+The Play Montage node enables you to not only select what Montage to play but you can also set the **Play Rate**, **Starting Position** and **Starting Section** as inputs. 
+
+On the output side, there are several callback events that you can use to trigger other script based on the state of the Montage: 
+
+* **On Completed** - called when the Montage finishes playing and is fully blended out.
+* **On Blend Out** - called when the Montage is starting to Blend Out, using [Blend Out Trigger Time](https://docs.unrealengine.com/latest/INT/Engine/Animation/AnimMontage/#montageproperties) or if the Montage ends.
+* **On Interrupted** - called if the Montage is starting to Blend Out, due to interruption by another Montage.
+* **On Notify Begin** and **On Notify End** - are callbacks when using either **Play Montage Notify** or **Play Montage Notify Window** [Anim Notifies](Engine/Animation/Sequences/Notifies) in the Montage asset. 
+	* These Anim Notifies can forward an additional **Notify Name** to differentiate between multiple callbacks from the same Montage. 
+
+In order for the Montage to play, the Skeletal Mesh will need to have its **Animation Mode** set to **Use Animation Blueprint** and the **Anim Class** pointed to an Animation Blueprint. 
+
+![](SetupAnimMode.png)
+
+Inside your Animation Blueprint, on the **AnimGraph**, you can use the **Slot** node to define which Slot from the Montage you want to play. 
+
+![](AnimGraphMontage.png)
+
+[REGION:note]
+See [](Engine/Animation/AnimHowTo/AdditiveAnimations) for another example of how a Montage can be used inside of your Animation Blueprints. 
+[/REGION]
+
diff --git a/Engine/Documentation/Source/Engine/Animation/AnimPose/AnimBPPoseNodes/AnimBPPoseNodes.INT.udn b/Engine/Documentation/Source/Engine/Animation/AnimPose/AnimBPPoseNodes/AnimBPPoseNodes.INT.udn
--- a/Engine/Documentation/Source/Engine/Animation/AnimPose/AnimBPPoseNodes/AnimBPPoseNodes.INT.udn
+++ b/Engine/Documentation/Source/Engine/Animation/AnimPose/AnimBPPoseNodes/AnimBPPoseNodes.INT.udn
@@ -16,6 +16,8 @@ Tags:Animation
 Tags:Animation Blueprint
 Tags:Animation Blending
 
+<!-- Page refactored and split off into two pages, no longer needed
+
 [TOC (start:2 end:2)]
 
 [Pose Assets](Engine\Animation\AnimPose) can be driven at runtime with **Animation Blueprints** as the **AnimGraph** supports a few different Pose Nodes. 
@@ -154,7 +156,7 @@ For additional information see the Unreal Engine Live Stream on Animation Topics
 [/OBJECT]
  
 
-
+ -->
 
 
 
diff --git a/Engine/Documentation/Source/Engine/Animation/AnimPose/AnimPose.INT.udn b/Engine/Documentation/Source/Engine/Animation/AnimPose/AnimPose.INT.udn
--- a/Engine/Documentation/Source/Engine/Animation/AnimPose/AnimPose.INT.udn
+++ b/Engine/Documentation/Source/Engine/Animation/AnimPose/AnimPose.INT.udn
@@ -4,6 +4,7 @@ Crumbs: %ROOT%, Engine, Engine/Animation
 Description:Describes the Animation Pose Asset which can be used to drive animation through weighted curve data. 
 Related: Engine\Animation\Persona\AnimCurves
 Related: Engine\Animation\AnimBlueprints
+Related: Engine\Animation\AnimHowTo\CurveDrivenAnimation
 version: 4.14
 tags:Assets
 topic-image:AnimationPoseAssets_topic.png
@@ -40,42 +41,55 @@ Using these weighted values, we could create a new motion that blends between Po
 For example, below we have created a motion where the character appears to be walking driven through curve data. 
 
 [OBJECT:EmbeddedVideo]
+[PARAMLITERAL:videoid]
+7PRGvePzB2g
+[/PARAMLITERAL]
 [PARAMLITERAL:width]
-640
+100
 [/PARAMLITERAL]
 [PARAMLITERAL:height]
-360
+56.25
 [/PARAMLITERAL]
-[PARAMLITERAL:videoid]
-7PRGvePzB2g
+[PARAMLITERAL:units]
+%
+[/PARAMLITERAL]
+[PARAMLITERAL:type]
+simple
 [/PARAMLITERAL]
 [/OBJECT]
 
 We can also use this process for creating facial animation as seen in the video below. 
 
 [OBJECT:EmbeddedVideo]
+[PARAMLITERAL:videoid]
+UdmU7zMN4tE
+[/PARAMLITERAL]
 [PARAMLITERAL:width]
-640
+100
 [/PARAMLITERAL]
 [PARAMLITERAL:height]
-360
+56.25
 [/PARAMLITERAL]
-[PARAMLITERAL:videoid]
-Z_EO1wfHfMs
+[PARAMLITERAL:units]
+%
+[/PARAMLITERAL]
+[PARAMLITERAL:type]
+simple
 [/PARAMLITERAL]
 [/OBJECT]
 
 Above, we have an Animation Sequence that contains curve data for various poses. 
-We then created a **Pose Asset** from the Animation Sequence to gain access to each of the curves/poses (we isolated **R_smile_pose** and **L_smile_pose** to show how we can make the face smile). 
-Next, we created a new Animation Sequence from the default base pose and added two curves to that sequence (the two smile poses). 
-Finally, we added new weight values midway through the animation that allowed us to create the expression of the face smiling. 
+
+[REGION:note]
+The above examples uses the experimental **Facial Animation Bulk Importer** plugin which imports facial animation curve tables (from FBX) into sound waves.
+[/REGION]
 
 [OBJECT:TopicButtonList]
 	[PARAM:icon]
 		![](%ROOT%/start_icon.png)(convert:false)
 	[/PARAM]
 	[PARAM:title]
-		Getting Started
+		Essentials
 	[/PARAM]
 	[PARAM:links]
 		[REGION:topics fourth]
@@ -94,64 +108,64 @@ Finally, we added new weight values midway through the animation that allowed us
 		![](%ROOT%/blueprint_icon.png)(convert:false)
 	[/PARAM]
 	[PARAM:title]
-		%Engine\Animation\AnimHowTo\CreatePoseAsset:title%
+		%Engine/Animation/AnimHowTo/CreatePoseAsset:title%
 	[/PARAM]
 	[PARAM:description]
-		%Engine\Animation\AnimHowTo\CreatePoseAsset:description%
+		%Engine/Animation/AnimHowTo/CreatePoseAsset:description%
 	[/PARAM]
 	[PARAM:path]
-		[RELATIVE:Engine\Animation\AnimHowTo\CreatePoseAsset]
+		[RELATIVE:Engine/Animation/AnimHowTo/CreatePoseAsset]
 	[/PARAM]
 [/OBJECT]
 [OBJECT:Topic]
 	[PARAM:image]
-		![](topic_howto.png)(w:291 h:120)
+		![](topic_additive.png)(w:291 h:120)
 	[/PARAM]
 	[PARAM:icon]
 		![](%ROOT%/blueprint_icon.png)(convert:false)
 	[/PARAM]
 	[PARAM:title]
-		%Engine\Animation\AnimHowTo\CurveDrivenAnimation:title%
+		%Engine\Animation\AnimPose\AdditivePA:title%
 	[/PARAM]
 	[PARAM:description]
-		%Engine\Animation\AnimHowTo\CurveDrivenAnimation:description%
+		%Engine\Animation\AnimPose\AdditivePA:description%
 	[/PARAM]
 	[PARAM:path]
-		[RELATIVE:Engine\Animation\AnimHowTo\CurveDrivenAnimation]
+		[RELATIVE:Engine\Animation\AnimPose\AdditivePA]
 	[/PARAM]
 [/OBJECT]
 [OBJECT:Topic]
 	[PARAM:image]
-		![](topic_animbp.png)(w:291 h:120)
+		![](Engine\Animation\AnimPose\PoseBlenderNode\PoseNodeTopic.png)(w:291 h:120)
 	[/PARAM]
 	[PARAM:icon]
 		![](%ROOT%/blueprint_icon.png)(convert:false)
 	[/PARAM]
 	[PARAM:title]
-		%Engine\Animation\AnimPose\AnimBPPoseNodes:title%
+		%Engine\Animation\AnimPose\PoseBlenderNode:title%
 	[/PARAM]
 	[PARAM:description]
-		%Engine\Animation\AnimPose\AnimBPPoseNodes:description%
+		%Engine\Animation\AnimPose\PoseBlenderNode:description%
 	[/PARAM]
 	[PARAM:path]
-		[RELATIVE:Engine\Animation\AnimPose\AnimBPPoseNodes]
+		[RELATIVE:Engine\Animation\AnimPose\PoseBlenderNode]
 	[/PARAM]
 [/OBJECT]
 [OBJECT:Topic]
 	[PARAM:image]
-		![](topic_additive.png)(w:291 h:120)
+		![](Engine\Animation\AnimPose\PoseDriverNode\PoseDriveTopic.png)(w:291 h:120)
 	[/PARAM]
 	[PARAM:icon]
 		![](%ROOT%/blueprint_icon.png)(convert:false)
 	[/PARAM]
 	[PARAM:title]
-		%Engine\Animation\AnimPose\AdditivePA:title%
+		%Engine\Animation\AnimPose\PoseDriverNode:title%
 	[/PARAM]
 	[PARAM:description]
-		%Engine\Animation\AnimPose\AdditivePA:description%
+		%Engine\Animation\AnimPose\PoseDriverNode:description%
 	[/PARAM]
 	[PARAM:path]
-		[RELATIVE:Engine\Animation\AnimPose\AdditivePA]
+		[RELATIVE:Engine\Animation\AnimPose\PoseDriverNode]
 	[/PARAM]
 [/OBJECT]
 [/EXCERPT:Topics]
diff --git a/Engine/Documentation/Source/Engine/Animation/AnimPose/PoseBlenderNode/PoseBlender.INT.udn b/Engine/Documentation/Source/Engine/Animation/AnimPose/PoseBlenderNode/PoseBlender.INT.udn
new file mode 100644
--- /dev/null
+++ b/Engine/Documentation/Source/Engine/Animation/AnimPose/PoseBlenderNode/PoseBlender.INT.udn
@@ -0,0 +1,93 @@
+Availability: Docs
+Crumbs: %ROOT%, Engine, Engine/Animation
+Title: Pose Blender
+Description:Describes how to play back a Pose Asset with the Pose Blender node. 
+Type: Reference
+Version: 4.16
+Parent: Engine/Animation/AnimPose
+Order: 1
+topic-image:PoseNodeTopic.png
+Related: Engine\Animation\AnimHowTo\CreatePoseAsset
+Related: Engine\Animation\Persona\AnimCurves
+Related: Engine\Animation\Sequences\Curves
+Related: Engine\Animation\AnimBlueprints
+Related: Engine\Animation\NodeReference\SkeletalControls
+Tags:Animation
+Tags:Animation Blueprint
+Tags:Animation Blending
+Tags:Pose Assets
+
+After creating [](Engine/Animation/AnimPose), one way of driving and playing back a **Pose Asset** at runtime is through the **Pose by Name** and **Pose Blender** nodes: 
+
+* **Pose by Name** - where you can get the transform of a pose by its name. 
+* **Pose Blender** - which outputs the correct pose driven by curves from a Source Pose. 
+
+[REGION:note]
+See [](Engine/Animation/AnimHowTo/CreatePoseAsset) for a step-by-step guide of creating a Pose Asset. 
+[/REGION]
+
+In the example below, we have a face that uses a Pose Asset to smile when the player is close and frown when they move away.
+
+[OBJECT:EmbeddedVideo]
+	[PARAMLITERAL:videoid]
+	d3oGXxq2ekw
+	[/PARAMLITERAL]
+	[PARAMLITERAL:width]
+	100
+	[/PARAMLITERAL]
+	[PARAMLITERAL:height]
+	56.25
+	[/PARAMLITERAL]
+	[PARAMLITERAL:units]
+	%
+	[/PARAMLITERAL]
+	[PARAMLITERAL:type]
+	simple
+	[/PARAMLITERAL]
+[/OBJECT]
+
+Our State Machine inside our Animation Blueprint is set up with three states: 
+
+![](StateMachine.png)
+
+Our default **State** which is our neutral expression, our **Smile** state which is entered when the player enters a trigger volume and our **Frown** state when the player leaves the trigger volume. 
+In this example, we only move between smiling and frowning states after triggering the first smile. Inside each state, we use the **Pose Blender** and our animation sequence to apply the curve driven animation. 
+
+![](PoseBlender.png)
+
+Pose Blender nodes are created automatically when dragging a Pose Asset into the graph. 
+
+![](DragInPoseAsset.png)
+
+### Additional Information
+
+Plugging an Animation Sequence directly into the Final Animation Pose without using a Pose Blender will not automatically apply the curve driven animation. 
+
+![](NoPoseBlender.png)
+
+Additionally, using only a Pose Blender node without supplying any incoming curve data to drive those poses will not output your desired pose. 
+
+![](NoCurveData.png)
+
+You will need to use a method to drive the curve data along with the Pose Driver node to effectively manipulate the final pose. 
+
+![](PoseBlender.png)
+
+You can also convert a Pose Blender to a Pose by Name (or a Pose by Name to a Pose Blender) node through the right-click context menu. 
+
+![](PoseByName.png)
+
+The Pose by Name will enable you to manually specify which pose to use inside the **Details** panel that is used inside the **Pose Asset**. 
+
+![](AdditivePoseByName.png)
+
+For our example, we are controlling two aspects of our character's face, the right side of a smile and the left side of a smile, therefore, a Pose Blender make more sense to use.
+If we wanted to control only one pose within our Pose Asset, the Pose by Name provides the ability to control a single pose by its name (which is usually good for full body context). 
+
+[REGION:note]
+If you are using a Pose Asset with **Additive** enabled, you will also need to use the **Apply Additive** node to correctly display your desired pose. 
+For Pose Assets that are not Additive, you can directly link a Pose by Name node to the Final Animation Pose output. 
+[/REGION]
+
+
+
diff --git a/Engine/Documentation/Source/Engine/Animation/AnimPose/PoseDriverNode/PoseDriverNode.INT.udn b/Engine/Documentation/Source/Engine/Animation/AnimPose/PoseDriverNode/PoseDriverNode.INT.udn
new file mode 100644
--- /dev/null
+++ b/Engine/Documentation/Source/Engine/Animation/AnimPose/PoseDriverNode/PoseDriverNode.INT.udn
@@ -0,0 +1,206 @@
+Availability: Docs
+Crumbs: %ROOT%, Engine, Engine/Animation
+Title: Pose Driver
+Description:Describes how to use the Pose Driver to drive Pose Asset or Curve values based on a bone's movement.
+Type: Reference
+Version: 4.16
+Parent: Engine/Animation/AnimPose
+Order: 1
+topic-image:PoseDriveTopic.png
+Related: Engine\Animation\AnimHowTo\CreatePoseAsset
+Related: Engine\Animation\Persona\AnimCurves
+Related: Engine\Animation\Sequences\Curves
+Related: Engine\Animation\AnimBlueprints
+Related: Engine\Animation\NodeReference\SkeletalControls
+Tags:Animation
+Tags:Animation Blueprint
+Tags:Animation Blending
+Tags:Pose Assets
+
+[REGION:warning]
+The Pose Driver node inside the AnimGraph in Animation Blueprints is currently experimental and some features may not work as expected.
+[/REGION]
+
+[TOC (start:2 end:2)]
+
+In addition to using the [](Engine\Animation\AnimPose\PoseBlenderNode) to playback **Pose Assets** through [Animation Blueprints](Engine/Animation/AnimBlueprints), 
+you can also drive Pose Assets directly at runtime with the **Pose Driver** node available within the [AnimGraph](Engine/Animation/AnimBlueprints/AnimGraph) of an Animation Blueprint. 
+The Pose Driver node enables a bone to drive other aspects of animation, based on a set of 'example poses' or lets you drive a curve value based on a bone's movement. 
+One use for this node would be to create a pose space deformer to drive corrective blend shapes based on a bone's orientation or translation. 
+
+In the clip below, clavicle bone poses are being used to move the shoulder relative to the arm bone's movement. 
+
+[OBJECT:EmbeddedVideo]
+	[PARAMLITERAL:videoid]
+	Cb7CYMmuQdw
+	[/PARAMLITERAL]
+	[PARAMLITERAL:width]
+	100
+	[/PARAMLITERAL]
+	[PARAMLITERAL:height]
+	56.25
+	[/PARAMLITERAL]
+	[PARAMLITERAL:units]
+	%
+	[/PARAMLITERAL]
+	[PARAMLITERAL:type]
+	simple
+	[/PARAMLITERAL]
+[/OBJECT]
+
+We add the Pose Driver node and assign a Pose Asset containing our poses that we want to control. 
+We then set the **Source Bone** to our **clav** bone so that based on the **Drive Source** (which is set to **Rotate**), 
+we can update our **Pose Targets** that drive the pose being used.
+As the animation plays and the clav bone rotates, in the **Pose Targets** section each pose displays a bar to indicate the amount of influence it has to the final pose. 
+[REGION:note]
+The Pose Driver node uses **RBF (Radial Basis Function)** which interpolates driven values based on the orientation of a target bone.
+[/REGION]
+
+## Adding a Pose Driver
+
+To add a Pose Driver node to your **Animation Blueprint**: 
+
+1. On the **AnimGraph**, right-click in the graph and search for the **Pose Driver** node. 
+
+	![](ContextPoseDriverNode.png)
+
+1. Once the node has been added, supply the incoming **Source Pose** and the wire the **Out Pose** to the **Final Animation Pose** or other nodes for blending. 
+
+	![](CompletePoseDriver.png)
+
+1. Use the **Details** panel to define what bones should be evaluated and Pose Asset poses or Curves should be driven. 
+
+	![](BlankDetails.png)
+
+	
+## Pose Targets
+
+Pose Targets are where you specify the Poses from a Pose Asset or Curves you want to drive with the Pose Driver node. 
+
+With a Pose Asset assigned, you can click the **Tools** button and select **Copy All from PoseAsset** to add each of the poses to the Pose Target window. 
+
+![](CopyPoses.png)
+
+Once they have been added, the window will update with the corresponding poses. 
+
+![](PoseTargetsOptions.png)
+
+[REGION:tip]
+You can select each pose by clicking on it in the **Pose Target** window or by clicking on the target line in the viewport. 
+[/REGION]
+
+### Curve based Pose Targets
+
+Pose Targets can also be driven by [](Engine/Animation/Sequences/Curves) (Morph Targets and Material Curves) instead of using a Pose Asset. 
+
+To use curve data instead of a Pose Asset: 
+
+1.	In the **Details** panel of a **Pose Driver** node, under the **Pose Driver** section, change the **Drive Output** to **Drive Curves**. 
+
+	![](DriveOutputSetting.png)
+
+1.	**Compile** the **Animation Blueprint**. 
+
+	![](CompileButton.png)
+
+1.	Click the **Add Target** button to add the number of curves you want to drive.
+
+	![](AddTarget.png)
+
+1.	For each curve, under the **Drive** option select your desired curve and check the **Curve** checkbox, then define its value in the graph. 
+
+	![](CurveSetupComplete.png)
+
+
+## Pose Driver Properties
+
+The **Pose Driver** properties section of the **Details** panel enables you to define the **Source Bones** to use for driving parameters in the **Pose Targets**. 
+You can specify multiple Source Bones or **Only Drive Selected Bones**, define an **Eval Space Bone** to use when reading the Source Bone Transform, set **RBF Parameters** used by the RBF Solver, as well as set the **Drive Source** and **Drive Output**. 
+
+Below are the properties within the Pose Driver section: 
+
+| Option | Description |
+| --- |--- |
+| **Source Bones** | Bone to use for driving parameters based on its orientation. | 
+| **Only Drive Selected Bones** |  If we should filter bones to be driven using the **Only Drive Bones** array. | 
+| **Only Drive Bones** | Only the bones specified in the array will be modified by the Pose Driver node. | 
+| **Eval Space Bone** | Optional other bone space to use when reading Source Bone transform. If not specified, the local space of Source Bone (for example, relative to the parent bone) is used.| 
+| **RBFParams** | [INCLUDE:Engine\Animation\AnimPose\PoseDriverNode#rbftable] | 
+| **Drive Source** | Which part of the transform is read: **Rotation** or **Translation**. | 
+| **Drive Output** | Whether we should drive poses or curves. | 
+
+<!--
+[EXCERPT:rbftable]
+| Option | Description |
+| ----------- |------|
+| **Radius** | Default radius for each target, scaled by per-Target Scale Factor. | 
+| **Function** | Defines the type of [Blend Mode](Engine/Animation/NonLinearBlends) to use. | 
+| **Distance Method** | Sets the measurement of distance to: **Swing Angle**, **Quaternion** or **Euclidean**. | 
+| **Twist Axis** | Axis to use when **Distance Method** is set to **Swing Angle**. | 
+| **Weight Threshold** | Weight below which we shouldn't bother returning a contribution from a target. | 
+[/EXCERPT:rbftable]
+-->
+
+
+## Settings and Relevancy
+
+The **Settings** section of the **Details** panel enables you to define the **Pose Asset** to evaluate and define if [](Engine/Animation/SyncGroups) are used. 
+The **Relevancy** section provides a way for you to determine if relevant animation nodes should ignore the node. 
+
+| Option | Description |
+| --- |--- |
+| **Pose Source** | The **Pose Asset** used as the source of pose data. | 
+| **Sync Group Name** | The name of the Synch Group to reference. | 
+| **Sync Group Role** | The type of membership in the group. | 
+| **Ignore for Relevancy Test** | If true, "Relevant anim" nodes that look for the highest weighted animation in a state will ignore this node.  | 
+
+
+## Additional Examples and Resources
+
+Below are some additional examples of the Pose Driver node and settings. 
+
+Below, the clavicle, lower, and upper arm bones are driving corrective blend shapes deforming the area around the shoulder as it rotates through curves. 
+
+[OBJECT:EmbeddedVideo]
+	[PARAMLITERAL:videoid]
+	sjtbvyDdW9w
+	[/PARAMLITERAL]
+	[PARAMLITERAL:width]
+	100
+	[/PARAMLITERAL]
+	[PARAMLITERAL:height]
+	56.25
+	[/PARAMLITERAL]
+	[PARAMLITERAL:units]
+	%
+	[/PARAMLITERAL]
+	[PARAMLITERAL:type]
+	simple
+	[/PARAMLITERAL]
+[/OBJECT]
+
+Another example of setting up and using the Pose Driver node can be found in the following Live Stream: 
+
+
+[OBJECT:EmbeddedVideo]
+	[PARAMLITERAL:videoid]
+	h2egrj1pXzw
+	[/PARAMLITERAL]
+	[PARAMLITERAL:width]
+	100
+	[/PARAMLITERAL]
+	[PARAMLITERAL:height]
+	56.25
+	[/PARAMLITERAL]
+	[PARAMLITERAL:units]
+	%
+	[/PARAMLITERAL]
+	[PARAMLITERAL:type]
+	simple
+	[/PARAMLITERAL]
+[/OBJECT]
+
+
+
+
+
diff --git a/Engine/Documentation/Source/Engine/Animation/Animation.INT.udn b/Engine/Documentation/Source/Engine/Animation/Animation.INT.udn
--- a/Engine/Documentation/Source/Engine/Animation/Animation.INT.udn
+++ b/Engine/Documentation/Source/Engine/Animation/Animation.INT.udn
@@ -122,6 +122,7 @@ order:3
 		Features
 	[/PARAM]
 	[PARAM:links]
+		* [](Engine\Animation\AnimModifiers "%Engine\Animation\AnimModifiers:description%")
 		* [](Engine\Animation\AnimPose "%Engine\Animation\AnimPose:description%")
 		* [](Engine/Animation/Persona/BoneManipulation "%Engine/Animation/Persona/BoneManipulation:description%")
 		* [](Engine/Animation/AnimationBlending "%Engine/Animation/AnimationBlending:description%")
diff --git a/Engine/Documentation/Source/Engine/Animation/NodeReference/SkeletalControls/LookAt/LookAt.INT.udn b/Engine/Documentation/Source/Engine/Animation/NodeReference/SkeletalControls/LookAt/LookAt.INT.udn
--- a/Engine/Documentation/Source/Engine/Animation/NodeReference/SkeletalControls/LookAt/LookAt.INT.udn
+++ b/Engine/Documentation/Source/Engine/Animation/NodeReference/SkeletalControls/LookAt/LookAt.INT.udn
@@ -1,4 +1,4 @@
-Availability: Docs
+Availability: Public
 Crumbs: %ROOT%, Engine, Engine/Animation, Engine/Animation/NodeReference, Engine/Animation/NodeReference/SkeletalControls
 Title:Look At
 Description:Describes how the Look At control can be used to specify a bone to trace or follow another bone. 
@@ -15,8 +15,6 @@ Related: Engine/Animation/AnimBlueprints/AnimGraph
 Related: Engine/Animation/StateMachines
 Related: Engine/Animation/AnimBlueprints
 
-<!--WB: Is a TOC needed for this?-->
-
 The **Look At** Skeletal Control Node can be used to specify a Bone to trace or follow another Bone in the Skeleton. 
 
 ![](LookAtnode.png)
diff --git a/Engine/Documentation/Source/Engine/Animation/NodeReference/SkeletalControls/ModifyCurve/ModifyCurve.INT.udn b/Engine/Documentation/Source/Engine/Animation/NodeReference/SkeletalControls/ModifyCurve/ModifyCurve.INT.udn
--- a/Engine/Documentation/Source/Engine/Animation/NodeReference/SkeletalControls/ModifyCurve/ModifyCurve.INT.udn
+++ b/Engine/Documentation/Source/Engine/Animation/NodeReference/SkeletalControls/ModifyCurve/ModifyCurve.INT.udn
@@ -1,4 +1,4 @@
-Availability: Docs
+Availability: Public
 Crumbs: %ROOT%, Engine, Engine/Animation, Engine/Animation/NodeReference, Engine/Animation/NodeReference/SkeletalControls
 Title:Modify Curve
 Description:Describes the Modify Curve node which can be used to modify animation curves with arbitrary logic inside Animation Graphs. 
@@ -17,13 +17,11 @@ Related: Engine/Animation/AnimBlueprints/AnimGraph
 Related: Engine/Animation/StateMachines
 topic-image:ModifyCurveNode.png
 
-<!--WB: Do need a TOC for this?-->
-
-The **Modify Curve** control can be used to expose [Animation Curve](Engine/Animation/Sequences/Curves) data which can be modified using arbitrary logic to enable more dynamically-driven animations. 
+The **Modify Curve** control can be used to expose [Animation Curve](Engine/Animation/Sequences/Curves) data which can be modified using arbitrary logic to enable dynamically-driven animations. 
 
 ![](ModifyCurveNode.png)
 
-Below, the Modify Curve node is used to drive a curve we create called **Eye Glow Intensity** where the bear's eyes pulsate and glow over time.
+Below, the Modify Curve node is used to drive a curve we created called **Eye Glow Intensity** where the bear's eyes pulsate and glow over time.
 
 [OBJECT:EmbeddedVideo]
 	[PARAMLITERAL:videoid]
@@ -43,7 +41,7 @@ Below, the Modify Curve node is used to drive a curve we create called **Eye Glo
 	[/PARAMLITERAL]
 [/OBJECT]
 
-To add this node to the [](Engine/Animation/AnimBlueprints/AnimGraph) of your [](Engine/Animation/AnimBlueprints): 
+**To add this node to the [](Engine/Animation/AnimBlueprints/AnimGraph) of your [](Engine/Animation/AnimBlueprints):**
 
 1. Right-click in the **AnimGraph** and search for and add the **Modify Curve** node
 
@@ -51,12 +49,11 @@ To add this node to the [](Engine/Animation/AnimBlueprints/AnimGraph) of your []
 
 	![](AddCurvePin.png)
 
-With the exposed float pin, you can define any logic you wish to drive the value (the logic below uses the Sin of Time multiplied by a coeffiect value to drive the Eye Glow Intensity).
+With the exposed float pin, you can define any logic you wish to drive the value (the logic below uses the **Sin** of **Time** multiplied by a coeffiect value to drive the **Eye Glow Intensity** property).
 
 ![](CurveLogic.png)
 
-
-To remove pins from the Modify Curve node: 
+**To remove pins from the Modify Curve node:**
 
 * Right-click on the node again and select **Remove Curve Pin** or right-click on the pin itself and select **Remove This Curve Pin**.
 
diff --git a/Engine/Documentation/Source/Engine/Animation/NodeReference/SkeletalControls/RigidBody/RigidBody.INT.udn b/Engine/Documentation/Source/Engine/Animation/NodeReference/SkeletalControls/RigidBody/RigidBody.INT.udn
--- a/Engine/Documentation/Source/Engine/Animation/NodeReference/SkeletalControls/RigidBody/RigidBody.INT.udn
+++ b/Engine/Documentation/Source/Engine/Animation/NodeReference/SkeletalControls/RigidBody/RigidBody.INT.udn
@@ -15,8 +15,8 @@ Related: Engine/Physics/PhAT/Reference
 Related: Engine/Animation/AnimBlueprints/AnimGraph
 Related: Engine/Animation/NodeReference/Blend
 
-The **Rigid Body** Skeletal Control node can be used to create large amounts of physically-simulated characters using the high-performance **Immediate Mode** PhysX API. 
-The Rigid Body node (pictued below) can be added to the **AnimGraph** of your [](Engine/Animation/AnimBlueprints) and can even be used with any [](Engine/Animation/NodeReference/Blend/), enabling you to switch from a pose into a simulated physics state.
+The **Rigid Body** Skeletal Control node can be used to create large amounts of physically-simulated characters using the high-performance [Immediate Mode PhysX API](https://www.geforce.com/hardware/technology/physx). 
+The Rigid Body node (pictued below) can be added to the **AnimGraph** of your [](Engine/Animation/AnimBlueprints) and can even be used with any [](Engine/Animation/NodeReference/Blend) enabling you to switch from a pose into a simulated physics state.
 Characters using this simulation can also generate collision with static geometry in the world. 
 
 ![](RigidBodyNode.png)
@@ -43,12 +43,12 @@ Using this node, several simulated characters below are able to respond to an im
 
 ## Setup
 
-Before using this node, inside your [Physics Asset](Engine/Physics/PhAT) set the [](Engine/Physics/PhysicsBodies/) you want to simulate to **Simulated**. 
+Before using this node, inside your [Physics Asset](Engine/Physics/PhAT) set the [](Engine/Physics/PhysicsBodies) you want to simulate to **Simulated**. 
 
 ![](RigidBody_02.png)
 
 [REGION:caption]
-Above we **Right-clicked** on our character's **Pelvis** bone and set all bodies below it to **Simulated** so our character fully simulates physics. 
+Above we right-clicked on our character's **Pelvis** bone and set all bodies below it to **Simulated** so that our character fully simulates physics. 
 [/REGION]
 
 With your Physics Asset set up, inside the AnimGraph of your Animation Blueprint, add the Rigid Body node and the necessary logic to define when to use it. 
diff --git a/Engine/Documentation/Source/Engine/Animation/Optimization/AnimOptimization.INT.udn b/Engine/Documentation/Source/Engine/Animation/Optimization/AnimOptimization.INT.udn
new file mode 100644
--- /dev/null
+++ b/Engine/Documentation/Source/Engine/Animation/Optimization/AnimOptimization.INT.udn
@@ -0,0 +1,267 @@
+Availability: Docs
+Crumbs: %ROOT%, Engine, Engine/Animation
+Title: Animation Optimization
+Description: Describes methods to achieve more performant Animation Blueprints through optimization techniques.
+Type: Overview
+SkillLevel: Advanced
+Version: 4.16
+Parent: Engine/Animation
+Order: 1
+Tags: Animation
+Tags: Optimization
+Tags: Animation Blueprint
+topic-image:OptimizationTopic.png
+Related: Engine/Performance
+
+[TOC (start:2 end:2)]
+
+As you construct your [](Engine/Animation/AnimBlueprints), there are some practices you should keep in mind to ensure the animation in your project runs as smoothly as possible. 
+Some of these are enabled by default while others will require that you consider the approach you take in setting up your Animation Blueprints. 
+Other efforts can be performed through C++ that will also grant you some control over how and when animations are updated, which can lead to improved performance. 
+
+
+## Multi Threaded Animation Update
+
+This allows more animation work to run in worker threads and is enabled by default which you can find in your [Project Settings](https://docs.unrealengine.com/latest/INT/Engine/UI/ProjectSettings/): 
+
+* Inside the **Project Settings** under **General Settings > Anim Blueprints**, make sure **Allow Multi Threaded Animation Update** is enabled.
+
+	![](ProjectSettings.png)
+
+This controls whether by default we allow Animation Blueprint Graph updates to be performed on non-game threads. 
+This also enables some extra checks in the Animation Blueprint compiler that will warn when unsafe operations are being attempted.
+Inside your **Animation Blueprints**, you will also want to make sure that it is set to **Use Multi Threaded Animation Update**.  
+
+* Inside your Animation Blueprint under the **Class Settings**, make sure that **Use Multi Threaded Animation Update** is enabled. 
+
+	![](AnimBPMultiThreadOption.png)
+
+The main driver behind this was to control data access more tightly across threads. To this end, much anim-graph-accessed data has been moved from [`UAnimInstance`](https://docs.unrealengine.com/latest/INT/API/Runtime/Engine/Animation/UAnimInstance/index.html) to a new struct called [`FAnimInstanceProxy`](https://docs.unrealengine.com/latest/INT/API/Runtime/Engine/Animation/FAnimInstanceProxy/index.html). 
+This proxy structure is where the bulk of the data on `UAnimInstance` is found. 
+
+In general, the `UAnimInstance` should not be accessed or mutated from within AnimGraph nodes (Update/Evaluate calls) as these can be run on other threads. 
+There are locking wrappers ([`GetProxyOnAnyThread`](https://docs.unrealengine.com/latest/INT/API/Runtime/Engine/Animation/UAnimInstance/GetProxyOnAnyThread/index.html) and [`GetProxyOnGameThread`](https://docs.unrealengine.com/latest/INT/API/Runtime/Engine/Animation/UAnimInstance/GetProxyOnGameThread/index.html)) that prevent access to the `FAnimInstanceProxy` while tasks are in-flight.
+The idea is that in the worst case, tasks wait for completion before data is allowed to be read from or written to in the proxy. 
+
+From the Anim Graph point of view, only the `FAnimInstanceProxy` can be accessed from animation nodes, not the `UAnimInstance`. 
+Data must be exchanged with the proxy for each tick (either through buffering, copying or some other strategy) in [`FAnimInstanceProxy::PreUpdate`](https://docs.unrealengine.com/latest/INT/API/Editor/AnimGraph/FAnimPreviewInstanceProxy/PreUpdate/index.html) or [`FAnimInstaceProxy::PreEvaluateAnimation`](https://docs.unrealengine.com/latest/INT/API/Runtime/Engine/Animation/FAnimInstanceProxy/PreEvaluateAnimation/index.html). 
+Any data that then needs to be accessed by external objects should then be exchanged/copied from the proxy in [`FAnimInstanceProxy::PostUpdate`](https://docs.unrealengine.com/latest/INT/API/Runtime/Engine/Animation/FAnimInstanceProxy/PostUpdate/index.html).
+
+[REGION:tip]
+This is in conflict with the general use case of `UAnimInstance` where member variables can be accessed from other classes while tasks are in-flight. 
+As a recommendation, try to not directly access the Anim Instance at all from other classes. Instead, the Anim Instance should pull data from elsewhere. 
+[/REGION]
+
+### Example Custom Native AnimInstance 
+
+In the code block below is an example of how one would build a custom native AnimInstance class using the new `FAnimInstanceProxy`, granting access to the internal workings and avoiding copies of shared data between the proxy and the instance: 
+
+	USTRUCT()
+
+	struct FExampleAnimInstanceProxy : public FAnimInstanceProxy
+
+	{
+
+		GENERATED_BODY()
+
+		FExampleAnimInstanceProxy()
+
+		: FAnimInstanceProxy()
+
+		{}
+
+		FExampleAnimInstanceProxy(UAnimInstance* Instance);
+
+		virtual void Update(float DeltaSeconds) override
+
+		{
+
+			// Update internal variables
+	
+		    MovementAngle += 1.0f * DeltaSeconds;
+
+	        HorizontalSpeed = FMath::Max(0.0f, HorizontalSpeed - DeltaSeconds);
+
+		}
+
+	public:
+
+		UPROPERTY(Transient, BlueprintReadWrite, EditAnywhere, Category = "Example")
+
+	    float MovementAngle;
+
+	    UPROPERTY(Transient, BlueprintReadWrite, EditAnywhere, Category = "Example")
+
+	    float HorizontalSpeed;
+
+	};
+
+	UCLASS(Transient, Blueprintable)
+
+	class UExampleAnimInstance : public UAnimInstance
+
+	{
+
+		 GENERATED_UCLASS_BODY()
+
+	private:
+
+	    // The AllowPrivateAccess meta flag will allow this to be exposed to Blueprint,
+
+	    // but only to graphs internal to this class.
+
+	    UPROPERTY(Transient, BlueprintReadOnly, Category = "Example", meta = (AllowPrivateAccess = "true"))
+
+	    FExampleAnimInstanceProxy Proxy;
+
+	    virtual FAnimInstanceProxy* CreateAnimInstanceProxy() override
+
+		{
+
+		    // override this to just return the proxy on this instance
+
+			return &Proxy;
+
+		}
+
+		virtual void DestroyAnimInstanceProxy(FAnimInstanceProxy* InProxy) override
+
+		{
+
+		}
+
+		friend struct FExampleAnimInstanceProxy;
+
+	};
+
+## Animation Fast Path
+
+**Animation Fast Path** provides a way to optimize variable access inside the **AnimGraph** update. 
+This enables the engine to copy parameters internally rather than executing Blueprint code (which involves making calls into the Blueprint Virtual Machine). 
+The compiler can currently optimize the following constructs: member variables negated boolean member variables and members of a nested structure.
+
+The Animation Fast Path option is enabled by default inside the **Project Settings**: 
+
+* Inside the **Project Settings** under **General Settings** and **Anim Blueprints**, make sure **Optimize Anim Blueprint Member Variable Access** is enabled.
+
+	![](FastPathEnabled.png)
+
+To make use of Animation Fath Path, inside the AnimGraph of your Animation Blueprints, ensure that no Blueprint logic is being executed. 
+In the image below, we are reading several float values which are being used to drive multiple Blend Space assets and a Blend resulting in our Final Animation Pose.
+Each node denoted with the lightning icon in the upper right corner are utilizing Fast Path as no logic is being executed. 
+
+![](FastPathExample_1.png)
+
+If we were to change this network to include any form of calculation such as the example depicted below, the associated node would no longer be using Fast Path. 
+
+![](FastPathExample_2.png)
+
+Above, since we are now executing Blueprint logic to generate the values feeding the **TEST_Blend2D** node, it is no longer using Fast Path (and the lighting icon will be removed).
+
+### Fast Path Methods
+
+In order for your Animation Blueprints to use Fast Path, ensure that they either: 
+
+**Access Member Variables Directly**
+
+Below we are using Fast Path by directly accessing and reading the value of our boolean variable to determine our pose. 
+
+![](AccessDirectly.png)
+
+In the next example, we are not using Fast Path as we are performing logic to determine if the boolean variable is equal to true. 
+
+![](AccessDirectlyBPLogic.png)
+
+**Access Negated Boolean Member Variables**
+
+Below we are using Fast Path by reading the value of the negated boolean to determine our pose. 
+
+![](AccessNegatedBooleans.png)
+
+In the next example, we are not using Fast Path as we are performing logic to determine if the boolean variable is not equal to true. 
+
+![](AccessNegatedBooleansBPLogic.png)
+
+**Access Members of a Nested Struct**
+
+Below we break our rotator variable to directly access our Pitch and Yaw variables to feed our Aim Offset. 
+
+![](AccessMembersOfStruct.png)
+
+**Access Members Using "Break Struct" Nodes**
+
+Below we use a Break Struct node to break our rotator variable into XYZ values to feed our Aim Offset. 
+
+![](AccessByBreakStruct.png)
+
+[REGION:note]
+Some **Break Struct** nodes like **Break Transform** will not currently use Fast Path as they perform conversions internally rather than simply copying data. 
+[/REGION]
+
+
+
+### Warn About Blueprint Usage
+
+
+To ensure that your Animation Blueprints are using Fast Path, you can enable the **Warn About Blueprint Usage** option which will cause the compiler to emit warnings to the Compiler Results log whenever a call into the Blueprint Virtual Machine is made from the AnimGraph. 
+
+* To enable **Warn About Blueprint Usage**, enable the option inside the **Class Settings** of your **Animation Blueprint** under **Optimization**. 
+
+	![](WarningOption.png)
+
+	When the compiler identifies any nodes that are not using Fast Path, they will be displayed in the **Compiler Results** log. 
+
+	![](ExampleWarningShown.png)
+
+	Above, since we are executing Blueprint logic in our AnimGraph and have the warning option enabled, we get a warning message in the compiler results and can click on it to take us to the offending node. 
+	This can help track down optimizations that need to be made and will enable you to identify the node(s) that may be the source of the problem. 
+
+<!--
+## Update Rate Optimization (URO)
+
+WB Note - Omitting this for now as it's slated for improvements in 4.18
+
+-->
+
+
+## General Tips
+
+As you start to consider the performance of your animation usage, here are some guidelines that you may want to follow when performing optimizations.
+
+[REGION:note]
+Based on the size and scope of your project, more invasive changes may be needed however this is generally a good place to start. 
+[/REGION]
+
+* **Make sure that the conditions for Parallel Updates are met**
+	 * In `UAnimInstance::NeedsImmediateUpdate` you can see all the conditions that must be met to avoid the update phase of animation running on the game thread. 
+	If root motion is required for character movement, the parallel update cannot be performed as character movement is not multi threaded. 
+* **Avoid calls into the Blueprint Virtual Machine**
+	* Consider [](Engine/Blueprints/TechnicalGuide/NativizingBlueprints) into C++ code. 
+	* Keep your **Event Graph** in your Animation Blueprints empty. Use a custom `UAnimInstance` and `FAnimInstanceProxy` derived class and do all of your work in the proxy during `FAnimInstanceProxy::Update` or `FAnimInstanceProxy::Evaluate` as these are executed on worker threads. 
+	* Ensure that the nodes within your **Anim Graph** of your Animation Blueprint are structured in a way that they are using [Fast Path](#animationfastpath). 
+	* Ensure that **Optimize Anim Blueprint Member Variable Access** is enabled in the **Project Settings** as this controls whether Animation Blueprint nodes that access member variables of their class directly should use the optimized path that avoids a thunk to the Blueprint VM.
+	* Generally the most costly part of an AnimGraph's execution, avoiding calls to the virtual machine is the key to getting maximum performance out of Animation Blueprints. 
+* **Use Update Rate Optimizations (URO)**
+	* This will prevent your animations from ticking too often. Control over how this is applied is up to your game, but we recommend moving toward ~15Hz and under Update Rates at appropriate distances for many characters, as well as disabling interpolation. 
+	* To enable, set your Skeletal Mesh Component to **Enable Update Rate Optimizations** and refer to `AnimUpdateRateTick()`. 
+		* Optionally, you can also enable **Display Debug Update Rate Optimizations** to enable onscreen debugging of URO being applied. 
+* **Enable Component Use Fixed Skel Bounds**
+	* In your Skeletal Mesh Component, enable the **Component Use Skel Bounds** option. 
+	* This will skip using a physics asset and will instead always use the fixed bounds defined in the Skeletal Mesh. 
+	* This will also skip recalculating bounding volumes for culling for every frame, increasing performance. 
+
+### Other Considerations
+
+When [profiling](Engine/Performance) your project, you may see that `FParallelAnimationCompletionTask` is being run for Skeletal Meshes on the main thread after worker threads are done. 
+This will be the bulk of the main thread work that you will see in your profile once the conditions for parallel updates are satisfied and usually consist of a few things, depending on your setup: 
+
+* Moving your components around, updating physics objects for bones for example. 
+	* Try avoiding updating physics for things that don't actually need it as this will be key to reducing this. 
+* Firing off Animation Notifies.
+	* These should be non-Blueprint, again to avoid calls to the Blueprint VM for efficiency. 
+	* These also need to be performed on the game thread as they can affect the animated object's lifetime. 
+* Interpolation of animation if URO is enabled. 
+* Blending of curves if Material or Morph Target curves are in use. 
+
+
+	
diff --git a/Engine/Documentation/Source/Engine/Animation/Persona/BasePoseManager/BasePoseManager.INT.udn b/Engine/Documentation/Source/Engine/Animation/Persona/BasePoseManager/BasePoseManager.INT.udn
--- a/Engine/Documentation/Source/Engine/Animation/Persona/BasePoseManager/BasePoseManager.INT.udn
+++ b/Engine/Documentation/Source/Engine/Animation/Persona/BasePoseManager/BasePoseManager.INT.udn
@@ -25,7 +25,7 @@ The **Retarget Manager** within the **Skeleton Editor** can be used to manage yo
 
 Since retargeting uses a [Skeleton](Engine/Animation/Skeleton) asset and the Skeleton has its original proportions defined by the Skeletal Mesh for which it was initially created, it follows that one-directional retargeting will work smoothly in most cases. 
 However, there will often come times when you will have special animations that are built just for the variant versions. 
-For example, suppose you have multiple characters that share the same Skeleton asset (a base character, a short character and a tall character) and you have a special animation that was created just for the tall version of the character.
+For example, suppose you have multiple characters that share the same Skeleton asset (a base character, a short character, and a tall character) and you have a special animation that was created just for the tall version of the character.
 
 If you imported this new tall-only animation in, you would still need to use the same Skeleton asset as before, which was created from the base version of the character. 
 This will cause the new animation's proportions to be distorted. 
@@ -128,3 +128,4 @@ After selecting a pose from the Pose Asset to use, clicking the **Import** butto
 Below (left) our default pose is shown and (right) our Retargeted Base Pose selected from a pose within our Pose Asset is shown. 
 
 ![](DefaultPose.png)(w:540)![](NewRetargetPose.png)(w:540)
+
diff --git a/Engine/Documentation/Source/Engine/Animation/Persona/PreviewSettings/PreviewSettings.INT.udn b/Engine/Documentation/Source/Engine/Animation/Persona/PreviewSettings/PreviewSettings.INT.udn
--- a/Engine/Documentation/Source/Engine/Animation/Persona/PreviewSettings/PreviewSettings.INT.udn
+++ b/Engine/Documentation/Source/Engine/Animation/Persona/PreviewSettings/PreviewSettings.INT.udn
@@ -1,6 +1,6 @@
-Availability: Docs
+Availability: Public
 Crumbs: %ROOT%, Engine, Engine/Animation, Engine/Animation/Persona
-Title:Preview Settings
+Title:Preview Scene Settings
 Description:User guide for the Preview Settings panel inside the Animation Tools. 
 Type: Reference
 version: 4.16
@@ -16,9 +16,21 @@ Related: Engine/Rendering/LightingAndShadows
 
 [TOC (start:2 end:2)]
 
-The **Preview Settings** enables you to quickly assess what your assets will look like in multiple environments and light scenarios without having to set up these scenes within your level. 
+The **Preview Scene Settings** panel enables you to quickly assess what your assets will look like in multiple environments and light scenarios without having to set up these scenes within your level. 
 Several different settings can be defined and applied from Skeletal Meshes to lighting and Post Process effects all within each of the [](Engine/Animation/Persona/Modes) inside the editor. 
 
+The location of the Preview Scene Settings panel can be found in the upper-right corner of the [Animation Tool](Engine/Animation/Persona/Modes) you are using. 
+
+![](PreviewSceneSettings.png)
+
+If you close the Preview Scene Settings panel, you can re-enable it from the **Window** option from the tool bar. 
+
+![](EnablePreviewSceneSettings.png)
+
+You can also access the Preview Scene Setting panel from the **Viewport Options** drop-down menu. 
+
+![](ViewportMenu.png)
+
 ## Animation
 
 This section is used to define the method in which to apply animation to the preview mesh. 
@@ -40,9 +52,12 @@ This section enables you to change the Skeletal Mesh used for previewing.
 |-------|------|
 | **Preview Mesh (Animation)** | The Skeletal Mesh to use for previewing. |
 
+Each type of asset, whether it's an Animation Sequence, Animation Blueprint, or a Skeleton can store and be overriden with different Preview Meshes. 
+The order of overriding goes from Skeleton to Asset, so if an asset doesn't have an override, it defaults to the Skeleton's Preview Mesh. 
+
 ## Additional Meshes
 
-This section enables you to include additional Skeletal Meshes (such as weapons, items or other objects) to the preview scene. 
+This section enables you to include additional Skeletal Meshes to the preview scene as part of a Preview Mesh Collection.  
 
 ![](PreviewMeshCollection.png)
 
@@ -51,6 +66,9 @@ This section enables you to include additional Skeletal Meshes (such as weapons,
 | **Additional Meshes** | The Preview Mesh Collection to use as included meshes for previewing. |
 | **Skeletal Meshes** | The Skeletal Mesh assets that are included in the specified Preview Mesh Collection. |
 
+Preview Mesh Collections are useful when you have several Skeletal Meshes that are related (like components of a character). 
+Preview Mesh Collections are also very flexible and support the use of unsaved assets within an unsaved collection. 
+They can also be extended by projects that might need to interface with other systems that can "build" a Skeletal Mesh out of other Skeletal Mesh parts. 
 
 ## Settings
 
@@ -128,3 +146,4 @@ Please see [](Platforms/Mobile/PostProcessEffects) for more information.
 [/REGION]
 
 
+
diff --git a/Engine/Documentation/Source/Engine/Animation/Persona/Toolbar/PersonaToolbar.INT.udn b/Engine/Documentation/Source/Engine/Animation/Persona/Toolbar/PersonaToolbar.INT.udn
--- a/Engine/Documentation/Source/Engine/Animation/Persona/Toolbar/PersonaToolbar.INT.udn
+++ b/Engine/Documentation/Source/Engine/Animation/Persona/Toolbar/PersonaToolbar.INT.udn
@@ -1,26 +1,26 @@
-Availability:Public
-Title:Animation Toolbars
+Availability: Public
 Crumbs: %ROOT%, Engine, Engine/Animation, Engine/Animation/Persona
-Description:Overview of the different Toolbars used in each of the Animation Tool Editors. 
-Related: Engine/Animation/Persona/Modes/Skeleton
-Related: Engine/Animation/Persona/Modes/Mesh
-Related: Engine/Animation/Persona/Modes/Animation
-Related: Engine/Animation/Persona/Modes/Graph
-version: 4.16
-type:reference
-topic-image:Engine/Animation/animation_topic.png
-tags:animation
+Title:Animation Toolbars
+Description: Temporary Description
+Type: Reference
+Version: 4.16
+Parent: Engine/Animation/Persona
+Order: 1
+Tags:animation
+topic-image:SkeletonToolbar.png
+Related: Engine/Animation/Persona
+Related: Engine/Content/Types/SkeletalMeshes
+Related: Engine/Animation/AnimBlueprints
 
+[TOC(start:2)]
 
 [EXCERPT:Overview]
-Regardless of what Animation Tool you are working in, each Editor will feature a **Toolbar** which can be located at the top. However, the options on the Toolbar will vary based on the type of Editor you are currently working in.
-The Toolbar is comprised of two main areas, the **Asset Toolbar** which contains several options related to the Editor you are working with and the **Editor Toolbar** which enables you to quickly switch between the different Animation Tool Editors. 
+Inside each [Animation Tool](Engine/Animation/Persona/Modes), you will find a **Toolbar** located along the top of the Editor window. 
+The options located on that Toolbar will vary based on the type of Editor you are currently working in. 
+The Toolbar is comprised of two main areas, the **Asset Toolbar** that includes options related to the Editor you are working with and the **Editor Toolbar** that enables you to quickly switch between the different Animation Tools. 
 [/EXCERPT:Overview]
 
 
-[TOC(start:2)]
-
- 
 ## Asset Toolbar
 
 Located along the left side of the Toolbar are the options that are available based on the type of Editor you are working with, refer to the sections below: 
@@ -65,10 +65,10 @@ The Animation Editor Toolbar is shown whenever you open any animation assets (An
 | ------- | ------------ |
 | ![](saveIcon.png) | Saves any changes made to the current animation asset.  |
 | ![](findincbIcon.png) | Opens the **Content Browser** and automatically navigates to the current animation asset. |
-| ![](createassetIcon.png) | Create new animation assets (animations, anim assets, Blend Spaces or Aim Offsets). When creating animation, this can be based on the **Reference Pose**, **Current Pose** or Current Animation's **Animation Data** or **Preview Mesh**. When using the Preview Mesh, the created animation will include any additional animation data generated from a post-process graph assigned tot he Skeletal Mesh, such as Anim Dynamics for Physics Simulation.   |
+| ![](createassetIcon.png) | Creates a new animation asset. When creating animation, this can be based on the **Reference Pose**, **Current Pose** or Current Animation's **Animation Data** or **Preview Mesh**. When using the preview mesh, the created animation will include any additional animation data generated from a post-process graph assigned to the Skeletal Mesh, such as Anim Dynamics for Physics Simulation as well as retargeting information on the preview mesh.   |
 | ![](ReimportAnimation.png) | Reimport the current animation. |
 | ![](compressionIcon.png) | Enables you to apply different [Compression](Engine/Animation/Sequences) settings to the animation (available for Animation Sequences). |
-| ![](ExportAsset.png) | Export current animation source data or current animation playing on the current preview mesh (includes retargeting, post process graph or other settings applied to the preview mesh).   |
+| ![](ExportAsset.png) | Export current animation source data or current animation playing on the current preview mesh (includes retargeting, post process graph or other settings applied to the preview mesh). |
 | ![](keyIcon.png) | Enables you to add the current Bone Transform of a selected Bone to an [Additive Layer Track.](Engine/Animation/AnimHowTo/LayerEditing)  |
 | ![](applyIcon.png) | Applies [Additive Layer Tracks](Engine/Animation/AnimHowTo/LayerEditing) to Runtime Animation Data. |
 | ![](PreviewMesh.png) | Set a new preview Skeletal Mesh for the current asset (stored per-animation or per-skeleton). |
@@ -102,7 +102,7 @@ The currently open editor will be highlighted:
 
 ![](animationdropdownmenu.png)
 
-A ** * (Asterisk) ** symbol in the lower left corner of the icon indicates that the asset has changes that have not been saved. 
+A ** * (Asterisk)** symbol in the lower left corner of the icon indicates that the asset has changes that have not been saved. 
 
 ![](unsavedasset.png)
 
diff --git a/Engine/Documentation/Source/Engine/Animation/Persona/VirtualBones/VirtualBones.INT.udn b/Engine/Documentation/Source/Engine/Animation/Persona/VirtualBones/VirtualBones.INT.udn
--- a/Engine/Documentation/Source/Engine/Animation/Persona/VirtualBones/VirtualBones.INT.udn
+++ b/Engine/Documentation/Source/Engine/Animation/Persona/VirtualBones/VirtualBones.INT.udn
@@ -1,75 +1,132 @@
-Availability:Docs
+Availability:Public
 Title:Virtual Bones
 Crumbs:
 Description:Overview of Virtual Bones that enables you to add IK or Anim constraint joints inside of the Editor.
 Related: Engine/Content/FBX/Animations
 Related: Engine/Animation/Persona/Modes
 Related: Engine/Content/Types/SkeletalMeshes
-version: 4.15
+version: 4.17
 parent: Engine/Animation
 type:reference
-topic-image: 
+topic-image:VirtualBones_Topic.png
 reference-image:WalkAnim.gif
-tags: Animations, Skeletal Mesh, Skeleton, Bones
-order:
+tags:Animation
+Tags:Skeletal Mesh
+order:100
 
 
 [REGION:fullwidth raw]
 ![Virtual Bones](VirtualBones.png)
 [/REGION]
 
-Inside the [Animation Tools](Engine/Animation/Persona/Modes), namely the Skeleton or Animation Editors, you can add **Virtual Bones** to the Skeleton of your mesh. 
-These serve to help improve iteration time by enabling you to change your target joint hierarchy for IK or aim inside of the Editor. Before this feature was added, you would 
-have to do this outside of the engine to add them to your asset in your modeling application and then reimport all the animations in order to fix the animation data with the 
-new joints included. 
 
-Virtual Bones, by their very nature, are not skinnable, but, instead, constrained between two existing bones on the Skeleton and they automatically have data generated for them 
-for each animation on the Skeleton. For example, you can add a joint that is a child of a hand, but is constrained to the palm joint. Unlike sockets, this joint can be used 
-in Animation Blueprints as a target, such as an IK target or look at target, or you could even modify them in the Animation Blueprint for later use. 
+[TOC(start:2 end:2)]
 
-For additonal information and suggestions for pratical usage of Virtual Bones, like how they have been used to make it easier to retarget or change reference frames for controllers
-and also used for orientation and slope warping. Check out the Animation Techniques used in Paragon training stream for more information. 
 
-[OBJECT:EmbeddedVideo]
-	[PARAMLITERAL:width]
-	640
-	[/PARAMLITERAL]
-	[PARAMLITERAL:height]
-	360
-	[/PARAMLITERAL]
-	[PARAMLITERAL:videoid]
-	1UOY-FMm-xo
-	[/PARAMLITERAL]
-[/OBJECT]
+Inside the [Skeleton Editor](Engine/Animation/Persona/Modes/Skeleton), you can add **Virtual Bones** to the existing skeleton of your mesh using the Skeleton Tree. These serve to help 
+improve iteration time by enabling you to change your target joint hierarchy for aiming or IK (inverse kinematics) from inside the Editor. Before this was possible, you would have to
+add these bones in your 3D modeling software outside of the Editor and then reimport all the animations to fix the animation data with the new joints included. 
+
+Virtual Bones, by their very nature, are not skinnable and are constrained between two existing bones in the skeleton hierarchy and they automatically generate data for each animation.
+For example, you can add a joint that is a child of the hand but is constrained to the palm joint. Unlike sockets, this joint can be used in [Animation Blueprints](Engine/Animation/AnimBlueprints)
+as a target, such an IK target or "look at" target. These could even be modified in an Animation Blueprint for later use.
 
-## Adding Virtual Bones to your Skeleton
 
-To add Virtual Bones to your Skeletal Mesh:
 
-1. In the **Content Browser**, open any Skeletal Mesh or Animation Asset. 
- 
-	![](CB_Skeleton.png)(w:350) 
-1. In the **Skeletal Editor** or **Animation Editor**, select any Bone from the joint hierarchy from the **Skeleton Tree** window. Now, right-click and hover over **Add Virtual Bone** from the context menu. 
- Select a bone from the **Pick Target Bone** window that opens. 
-   
-	[REGION:raw]
-	![](SelectVirtualBone.png)(w:850)
+## Adding Virtual Bones
+
+To add a Virtual Bone to your Skeletal Mesh's skeleton tree, perform the following actions:
+
+1. Use the **Skeleton Tree** window, select any bone you want to be the first connection point. 
+	
+	[REGION:lightbox]
+	[![](Add_SelectBone.png)(w:600)](Add_SelectBone.png)
 	[/REGION]
+	[REGION:caption] Click image for full size. [/REGION]
 
+1. Right-click to bring up the context menu and select **Add Virtual Bone** from the list, and then in the expanded **Pick Target Bone...** select the bone you want to connect.
+	
+	[REGION:lightbox]
+	[![](Add_AddVB_Target.png)(w:600)](Add_AddVB_Target.png)
+	[/REGION]
+	[REGION:caption] Click image for full size. [/REGION]
+	
+	[REGION:tip]
+	You can enable your skeleton hierarchy to always be visible by going to the viewport and selecting **Show** > **Bone** > **All Hierarchy**. 
+	[/REGION]
 
-Now that you've added your Virtual Bone(s), you will see it listed in the Skeleton hierarchy. These are indicated by the prefix "VB" in their name and their lavender color in the Skeleton Tree window. 
+1. Now, you'll see the listed Virtual Bone in the Skeleton Tree hierarchy with the pre-fix **VB**. 
 
-[REGION:lightbox]
-[![](VirtualBoneView.png)(w:550)](VirtualBoneView.png)
-[/REGION]
+	![](Add_VirtualBoneAdded.png)
 
-[REGION:tip]
-If you do not see your bones in the Viewport of the Skeleton or Animation Editors, go to **Show** > **Bones** > **All Hierarchy** so that your Skeleton is always visible. There are also other options here
-to fit your preference if you do not want to always see all the bones in the viewport.
-[/REGION]
+	When you select it from the list, you will see it highlighted in the Skeleton Editor viewport. 
+
+	[REGION:lightbox] 
+	[![](Add_VirtualBoneSelected.png)(w:600)](Add_VirtualBoneSelected.png) 
+	[/REGION] 
+	[REGION:caption] Click image for full size. [/REGION]
+
+	When you play an animation with the Virtual Bone selected, you can see that it animates and updates accordingly.
+
+	[REGION:asyncgif] ![](WalkAnim.gif)(w:350) [/REGION]
+
+## Deleting Virtual Bones
+
+To delete an added Virtual Bone from your Skeletal Mesh's skeleton tree, perform the following actions:
+
+1. In the **Skelton Tree** window, select any Virtual Bone.
+
+	[REGION:lightbox]
+	[![](Delete_SelectVB.png)(w:600)](Delete_SelectVB.png)
+	[/REGION]
+	[REGION:caption] Click image for full size. [/REGION]
+
+1. Right-click to bring up the context menu and select **Delete** from the list. 
+
+	![](Delete_DeleteVB.png)(w:400)
+
+
+## Renaming Virtual Bones
+
+To rename a Virtual Bone in your Skeletal Mesh's skeleton tree, perform the following actions:
+
+1. In the **Skeleton Tree** window, select any Virtual Bone.
+
+	[REGION:lightbox]
+	[![](Delete_SelectVB.png)(w:600)](Delete_SelectVB.png)
+	[/REGION]
+	[REGION:caption] Click image for full size. [/REGION]
+
+1. There are three ways you can rename a Virtual Bone: 
+
+	* Right-click and select **Rename Virtual Bone** from the context menu
+		
+		![](Rename_ContextMenu.png)(w:400)
+	* With the Virtual Bone selected, press **F2**.
+	* Or double-click on the name of the Virtual Bone.
+
+	
+
+## Additional Resources
+
+For additional information and suggestions for practical usage of Virtual Bones, like how they have been used to make retargeting easier, to change reference frames for controllers, for orientation, and 
+slope warping. 
+
+[OBJECT:EmbeddedVideo]
+	[PARAMLITERAL:videoid]
+	1UOY-FMm-xo
+	[/PARAMLITERAL]
+	[PARAMLITERAL:width]
+	50
+	[/PARAMLITERAL]
+	[PARAMLITERAL:height]
+	56.25
+	[/PARAMLITERAL]
+	[PARAMLITERAL:units]
+	%
+	[/PARAMLITERAL]
+	[PARAMLITERAL:type]
+	simple
+	[/PARAMLITERAL]
+[/OBJECT]
 
-In the Animation Editor, you will see your Virtual Bone(s) constrained between the selected joints and that it updates with your animation. 
- 
-[REGION:asyncgif]
-![Virtual Bone](WalkAnim.gif)(w:350)
-[/REGION]
diff --git a/Engine/Documentation/Source/Engine/Animation/PoseSnapshot/PoseSnapshot.INT.udn b/Engine/Documentation/Source/Engine/Animation/PoseSnapshot/PoseSnapshot.INT.udn
--- a/Engine/Documentation/Source/Engine/Animation/PoseSnapshot/PoseSnapshot.INT.udn
+++ b/Engine/Documentation/Source/Engine/Animation/PoseSnapshot/PoseSnapshot.INT.udn
@@ -1,4 +1,4 @@
-Availability: Docs
+Availability: Public
 Crumbs: %ROOT%, Engine, Engine/Animation
 Title: Animation Pose Snapshot
 Description:Animation Pose Snapshots can capture a runtime Skeletal Mesh Pose inside a Blueprint and can be used for blending in or out of additional poses within an Animation Blueprint. 
@@ -19,8 +19,9 @@ Related: Engine/Physics/PhAT
 
 [TOC (start:2 end:2)]
 
-When animating Skeletal Meshes, there may be instances when you want to apply and let physics take control of the mesh (such as a character entering a rag-doll state). 
-After physics has been applied, using the **Animation Pose Snapshot** feature, you can capture a Skeletal Mesh pose (storing all the Bone Transform data) within a [Blueprint](Engine\Blueprints) and retrieve it inside [](Engine\Animation\AnimBlueprints) enabling you to use the saved pose as a source to blend from (as seen in the example video below). 
+When animating **Skeletal Meshes**, there may be instances when you want to let physics take control of the mesh (such as a character entering a rag-doll state). 
+After physics has been applied, using the **Animation Pose Snapshot** feature, you can capture a Skeletal Mesh pose (storing all the Bone Transform data) within a [Blueprint](Engine\Blueprints).
+You can then retrieve that informaiton inside an [Animation Blueprint](Engine\Animation\AnimBlueprints) and use the saved pose as a source to blend from (as seen in the example video below). 
 
 [OBJECT:EmbeddedVideo]
 	[PARAMLITERAL:videoid]
@@ -40,22 +41,21 @@ After physics has been applied, using the **Animation Pose Snapshot** feature, y
 	[/PARAMLITERAL]
 [/OBJECT]
 
-Above, our character enters a rag-doll state when we press a key and we use the **Pose Snapshot** node in our [Character Blueprint](Engine/Blueprints/UserGuide/Types/ClassBlueprint) to save off the pose of our Skeletal Mesh.
-When we press another key, our character blends from that snapshot into and plays a "get up" [](Engine/Animation/AnimMontage) before resuming the normal locomotion state. 
+Above, our character enters a rag-doll state when we press the R key and we use the **Pose Snapshot** node in our [Character Blueprint](Engine/Blueprints/UserGuide/Types/ClassBlueprint) to save off the pose of our Skeletal Mesh.
+When we press the R key again, our character blends from that snapshot into and plays a "get up" [](Engine/Animation/AnimMontage) before resuming the normal locomotion state. 
 This grants us the ability to take whatever pose the character ends up in as a result of physics and generate a smooth blend from that pose into an animation of the character getting back to their feet. 
 
 ## Save Pose Snapshot
 
 In order to save your Skeletal Mesh's pose at runtime, inside your **Character Blueprint**, you will need access to the [Skeletal Mesh Component](Engine/Components/SkeletalMesh) and its **Anim Instance**. 
+After getting the Skeletal Mesh and Anim Instance, you can call the **Save Pose Snapshot** node and enter the desired **Snapshot Name**. 
+You can manually enter a name into the **Snapshot Name** field or create a variable to store the name. 
 
 ![](PoseSnapShotBlueprint.png)
 
-With these in place, you can call the **Save Pose Snapshot** node and enter the desired **Snapshot Name**. 
-You can manually enter a name into the Snapshot Name field or create a variable to store the name. 
-The name you provide as the Snapshot Name **MUST** also be used when attempting to retrieve the Pose Snapshot inside your Animation Blueprint. 
-
-[REGION:note]
-When calling Save Pose Snapshot, the snapshot is taken at the current LOD. For example, if you took the snapshot at LOD1 and then used it at LOD0, any bones not in LOD1 will use the mesh's reference pose. 
+[REGION:warning]
+The name you provide as the **Snapshot Name** must also be used when attempting to retrieve the Pose Snapshot inside your **Animation Blueprint**. 
+Additionally, when calling Save Pose Snapshot, the snapshot is taken at the current LOD. For example, if you took the snapshot at LOD1 and then used it at LOD0, any bones not in LOD1 will use the mesh's reference pose. 
 [/REGION]
 
 ## Retrieve Pose Snapshot
@@ -64,33 +64,33 @@ To retrieve a Pose Snapshot, inside the **AnimGraph** of your **Animation Bluepr
 
 ![](PoseSnapShotAnimBPNode.png)
 
-Below is the graph for the sample use case of a character that is getting up from a rag-doll pose. 
+Below is the graph used for the example of a character getting up from a rag-doll pose. 
 
 ![](PoseSnapShotAnimBP.png)
 
-Above, we have a [State Machine](Engine/Animation/StateMachines) called **Default** that drives our character's locomotion and use an AnimMontage in the **MySlot** node that plays an animation of our character getting up when called to play. 
+Above, we have a [State Machine](Engine/Animation/StateMachines) called **Default** that drives our character's locomotion and uses a **Slot** node called **MySlot** that plays an Animation Montage of our character getting up. 
 We use the [Blend Poses by bool](https://docs.unrealengine.com/latest/INT/Engine/Animation/NodeReference/Blend/#blendposesbybool) node to determine if our character has stopped moving, where if **True** we switch over to our Pose Snapshot. 
-If **False**, we blend from the Pose Snapshot into our slotted AnimMontage before continuing on with our normal Default State Machine. 
+If **False**, we blend from the Pose Snapshot into our slotted Animation Montage before continuing on with our normal Default State Machine. 
 
 ## Snapshot Pose Function
 
-An alternative way of using the Pose Snapshot feature is by calling the **Snapshot Pose** function within Blueprint to save a snapshot to a **Pose Snapshot** variable. 
+An alternative way of using the Pose Snapshot feature is by calling the **Snapshot Pose** function within **Blueprint** to save a snapshot to a **Pose Snapshot** variable. 
 
 ![](SnapshotPose.png)
 
-When using Snapshot Pose, you will need to provide a variable in which to save the snapshot to as shown above. 
+When using Snapshot Pose, you will need to provide a variable to save the snapshot to as shown above. 
 
-Inside your AnimGraph once you have added the **Pose Snapshot** node, in the **Details** panel set the **Mode** to **Snapshot Pin** and check the **(As pin) Snapshot** option. 
+Inside your **AnimGraph** after adding the **Pose Snapshot** node, set the **Mode** to **Snapshot Pin** and check the **(As pin) Snapshot** option inside the **Details** panel. 
 
 ![](PoseSnapshotDetails.png)
 
-This will expose a Pose Snapshot input pin on the node that you can pass in your desired snapshot variable. 
+This will expose a Pose Snapshot input pin on the node where you can pass in your desired snapshot variable. 
 
 ![](SnapshotPoseAnimGraph.png)
 
-## Additional Resources
+## Additional Resource
 
-This feature can also be seen within the following Live Training Stream: 
+Additional information about this feature can be seen in the following Live Training Stream. 
 
 [OBJECT:EmbeddedVideo]
 	[PARAMLITERAL:videoid]
@@ -115,3 +115,5 @@ This feature can also be seen within the following Live Training Stream:
 
 
 
+
+
diff --git a/Engine/Documentation/Source/Engine/Animation/SyncGroups/SyncGroups.INT.udn b/Engine/Documentation/Source/Engine/Animation/SyncGroups/SyncGroups.INT.udn
--- a/Engine/Documentation/Source/Engine/Animation/SyncGroups/SyncGroups.INT.udn
+++ b/Engine/Documentation/Source/Engine/Animation/SyncGroups/SyncGroups.INT.udn
@@ -2,18 +2,30 @@ Availability:Public
 Title: Sync Groups
 Crumbs: %ROOT%, Engine, Engine/Animation
 Description:Sync Groups allow you to maintain the synchronization of animations of different lengths.
-version: 4.9
+Type: Overview
+SkillLevel: Intermediate
+Version: 4.17
+Parent: Engine/Animation
+Order: 1
+Tags: Animation Sequence
+Tags: Animation Notify
 topic-image:SyncGroups_topic.png
+Related: Engine/Animation/Sequences
+Related: Engine/Animation/AnimBlueprints
+
+[TOC (start:2 end:2)]
 
 [EXCERPT:Intro]
 **Sync Groups** are to keep related animations that may have different lengths synchronized. For example, you may have a walk cycle and a run cycle that you would like to blend together so that the character can smoothly accelerate or decelerate. But what if these animations were different lengths, such as the walk being significantly longer than the run? In such a case, directly blending from one to the other would have unnatural results, an unsightly "beat" as the foot animation switches.
 
-Sync Groups solve this problem by allowing one primary animation node to serve as the Leader, and all other related animations will simply scale their time length to match. Typically, the leader is  the node with the greatest blend weight. As the weight blends and the follower's weight exceeds the leader, the follower becomes the leader. In this way, the two animations can work smoothly together offering a seamless transition from one motion to the next. 
+Sync Groups solve this problem by allowing one primary animation node to serve as the Leader, and all other related animations will simply scale their time length to match. Typically, the leader is the node with the greatest blend weight. As the weight blends and the follower's weight exceeds the leader, the follower becomes the leader. In this way, the two animations can work smoothly together offering a seamless transition from one motion to the next. 
 
 It should be noted, however, that since the time of the animations is shifting, certain animation considerations become apparent. For instance, in the example of blending between different walk and run cycles, you would want to make sure your animations all started and ended on the same foot. Establishing these standards early will help everything blend much more smoothly!
 [/EXCERPT:Intro]
 
-### Group Roles: Leaders and Followers
+
+
+## Group Roles: Leaders and Followers
 
 Sync Groups work off the concept of group roles, in which one animation is the "leader" while all other animations are "followers". 
 
@@ -58,9 +70,6 @@ Leaders and followers are designated simply by which animation has the greater b
 [/REGION]
 [/EXCERPT:WalkAnimationScaled]
 
-
-
-
 [EXCERPT:LeaderList]
 * Provides the master animation length that is used by all followers
 * Executes all of its Notifies
@@ -87,7 +96,7 @@ To set the Sync Group, select the animation node in the AnimGraph and look in th
 
 <!---
 Group Role has 3 options:
-* Can Be Leader - This is the default setting. This means that if this animation has the greatest blend weight, then it can become the leader and other animations will scale to match its lenght.
+* Can Be Leader - This is the default setting. This means that if this animation has the greatest blend weight, then it can become the leader and other animations will scale to match its length.
 * Always Leader - If set, then this animation will always be the leader and all animations will always scale to match its length.
 * Never Leader - If set, then this animation can never be the leader and will always scale to match the next weighted animation.
 
@@ -98,5 +107,33 @@ Group Role has 3 options:
 Sync Groups can be set on AnimSequence and BlendSpace nodes within the AnimGraph, as well as inside State Machines.
 [/REGION]
 
+## Marker-Based Animation Syncing
+
+In addition to Sync Groups, animations can also be synchronized using markers within your animation. 
+
+![](AddSyncMarker.png)
+
+To add a Sync Marker, inside the [](Engine/Animation/Sequences/Notifies) track of your Animation Sequence, right-click and select **Add Sync Marker**. 
+
+**Some elements to consider when working with Sync Markers:**
+
+* Only animations within the same Sync Group are synced and the leader drives the positions of followers within the same Sync Group. 
+* The play rate is always that of the master animation (play rate adjustment is currently not available). 
+* Only makers common to all animations within a group are synced. 
+	* For example, if one animation is missing the "Right Foot Down" makers, those markers will be ignored for all animations when determining that frame's position. 
+* Position is synced based on the relative position of the Leader with respect to its common markers. 
+	* For example, if the Leader is 25% of the way between its left and right foot markers, the followers will be synced to 25% of the way of their respective left and right markers. 
+* Marker based sync is used automatically when animations in a Sync Group have enough matching markers otherwise the original scaled length syncing behavior is used. 
+* [AnimMontages](Engine/Animation/AnimMontage) also support marker-based sync while blending out, so you can transition back to other animations seamlessly.
+	* You can find the Sync Group setting [Asset Details](Engine/Animation/Persona/AnimAssetDetails) panel in the AnimMontage. 
+
+[REGION:note]
+Sync Markers can also be added automatically by setting up and applying [](Engine/Animation/AnimModifiers) within your animation. 
+[/REGION]
+
+
+
+
+
 
 
diff --git a/Engine/Documentation/Source/Engine/Audio/SoundActors/SoundActors.INT.udn b/Engine/Documentation/Source/Engine/Audio/SoundActors/SoundActors.INT.udn
--- a/Engine/Documentation/Source/Engine/Audio/SoundActors/SoundActors.INT.udn
+++ b/Engine/Documentation/Source/Engine/Audio/SoundActors/SoundActors.INT.udn
@@ -100,7 +100,7 @@ When **Override Attenuation** is checked, the ability to set additional Attenuat
 | **Spatialize** | Enable the source to be positioned in 3D.  |
 | **Distance Algorithm** | The type of volume versus distance algorithm to use for the attenuation model. |
 | **Attenuation Shape** | The shape of the attenuation volume. |
-| **Radius** | The overall size of the volume. <BR> By default, the volume shape is set to **Sphere**. This option will change based on the Attenuation Shape selected. |
+| **Radius** | The overall size of the volume. By default, the volume shape is set to **Sphere**. This option will change based on the Attenuation Shape selected. |
 | **Falloff Distance** | The distance over which falloff occurs. |
 | **Non-Spatialized Radius** | At what distance to start treating the sound source as spatialized. |
 | **Attenuate with LPF** | Enables attenuation via Low Pass Filter. |
