Copying //UE4/Dev-Documentation to //UE4/Samples-Main (Source: //UE4/Dev-Documentation @ 3623665)

diff --git a/Engine/Documentation/Source/Engine/HLOD/Reference/HLODReference.INT.udn b/Engine/Documentation/Source/Engine/HLOD/Reference/HLODReference.INT.udn
--- a/Engine/Documentation/Source/Engine/HLOD/Reference/HLODReference.INT.udn
+++ b/Engine/Documentation/Source/Engine/HLOD/Reference/HLODReference.INT.udn
@@ -1,26 +1,37 @@
-title:Hierarchical LOD Reference
-description:A reference of the interface elements and properties found in the Hierarchical LOD Outliner.
-crumbs:
-availability:Public
-version:4.13
-parent:Engine/HLOD
-order:2
-type:reference
-tags:HLOD
+Availability: Licensee
+Crumbs: %ROOT%
+Title: Hierarchical LOD Outliner Reference
+Description: A reference page for the interface elements and properties found in the Hierarchical LOD Outliner.
+Type: Reference
+Version: 4.17
+Parent:Engine/HLOD
+Order:2
+Tags:Static Mesh
+Tags:Level of Detail
+Tags:HLOD
 topic-image:Engine/HLOD/defaultHLOD_topic.png
 
-This page details the buttons and properties that exist on the **Hierarchical LOD Outliner**.
+
+This page details the available right-click context menus and Hierarchical LOD (Level of Detail) Setup menus available in the **Hierarchical LOD Outliner** (HLOD). Below you will find information about accessing these menus and what options
+they contain.
+
+[REGION:raw]
+![image alt text](Engine/HLOD/Overview/image_1.png)
+[/REGION]
+
+To access the HLOD property window, perform the following actions: 
+
+* From the file menu, select **Window** > **Hierarchical LOD Outliner**.
 
 ## Property and Interface Reference
 
 ![HLOD Outliner Fully Expanded](image_2.png)
 
-
 ### HLOD Actions
 
 ![](InterfaceElements.png)
 
-Below are the descriptions for the 3 buttons and slider that exist in the **Hierarchical LOD Outliner**.
+Below are the descriptions for the three buttons and slider that exist in the **Hierarchical LOD Outliner**.
 
 [REGION:raw]
 | --- | :---: | :--- |
@@ -28,41 +39,41 @@ Below are the descriptions for the 3 buttons and slider that exist in the **Hier
 | **Delete Clusters** || This will delete all HLODs that have been previously setup in this level. This includes all generated proxy meshes as well as the clusters that have been setup. This will not wipe the previously set values or properties for the HLOD Levels. |
 | **Generate Proxy Meshes** || [INCLUDE:#gpm] |
 | **Forced Viewing Level** || This will allow you to force a specific LOD level to be visible in the editor. This is helpful for debugging any LOD issues you may have. |
-| | %spacer%  | |
+|| %spacer% |||
 [/REGION]
 
 ### Hierarchical LOD Outliner Context Menus
 
-The Hierarchical LOD Outliner has a few right-click context menus, below are how to access them and the options they present.
+When you right-click in the Hierarchical LOD Outliner, you will notice a few context menus. Use the table below to find information about accessing these menus and what the options they contain do.
 
 #### LOD Actor Context Menu
 
-[REGION:raw]
-![](ContextMenu1.png)
-[/REGION]
+Right-click any LOD Actor listed under the **Scene Actor Name** column to bring up this menu where you'll be able to manipulate this specific LODActor.
+
+![](ContextMenu1.png)(w:400)
 
 [REGION:raw]
 | Option | %spacer% | Description |
 | --- | :---: | :--- |
 | **Select LOD Actor** || Selects the LOD Actor in the viewport (same effect as left-clicking the cluster). |
-| **Select Contained Actors** || Selects all the **Actors** contained in the **LOD Cluster**. |
-| **ForceView** || Toggles the view of the **HLOD Actor** in the viewport. |
+| **Select Contained Actors** || Selects all the Actors contained in the LOD Cluster. |
+| **ForceView** || Toggles the view of the HLOD Actor in the viewport. |
 | **Rebuild Proxy Mesh** || Forces the right-clicked cluster to rebuild its proxy mesh. |
-| **Create Containing Hierarchical Volume** || Creates a **Hierarchical Volume** around the **Actors** contained in the **Cluster** that can then be modified to encompass greater or fewer **Actors**. |
+| **Create Containing Hierarchical Volume** || Creates a Hierarchical Volume around the Actors contained in the Cluster that can then be modified to encompass greater or fewer Actors. |
 | **Delete Cluster** || Deletes the right-clicked Cluster. |
 [/REGION]
 
 #### Actor Context Menu
 
-[REGION:raw]
-![](ContextMenu2.png)
-[/REGION]
+Use the arrow next to any LODActor to expand the scene Actors that are included with the generated HLOD cluster. Right-click any Actor in this expanded list to remove or exclude it from the cluster generation.
+
+![](ContextMenu2.png)(w:400)
 
 [REGION:raw]
 | Option | %spacer% | Description |
 | --- | :---: | :--- |
-| **Remove From Cluster** || Removes the right-clicked **Actor** from the **Cluster**. |
-| **Exclude From Cluster Generation** || Excludes the right-clicked **Actor** from **Cluster** generation. If you wish to re-add the excluded **Actor** click and drag it from the **World Outliner** into the **Cluster** in the **Hierarchical LOD Outliner**. |
+| **Remove From Cluster** || Removes the right-clicked Actor from the Cluster. |
+| **Exclude From Cluster Generation** || Excludes the right-clicked Actor from Cluster generation. If you wish to re-add the excluded Actor click and drag it from the World Outliner into the Cluster in the Hierarchical LOD Outliner. |
 [/REGION]
 
 ### Properties
@@ -71,43 +82,90 @@ Below, broken out by major section, are the properties that can be found in a **
 
 ### Cluster Generation Settings
 
+The **Cluster Generation** settings enable you to control how HLOD clusters will be generated to include Actors from your levels by setting the desired bounds of your clusters, how full the cluster should be, and the minimum number of Actors that must be used to generate a cluster.
+
 ![](cluster.png)
 
 [REGION:raw]
-| Property | %spacer% | Description |
+| Property | %spacer2% | Description |
 | --- | :---: | :--- |
 | **Desired Bound Radius** || This is not guaranteed but used to calculate filling factor for auto clustering. |
 | **Desired Filling Percentage** || This is not guaranteed but used to calculate filling factor for auto clustering. |
-| **Min Number of Actors to Build** || Min number of **Actors** to build LODActor. |
+| **Min Number of Actors to Build** || Min number of Actors to build LODActor. |
 [/REGION]
 
 ### Mesh Generation Settings
 
+The **Mesh Generation** settings enable you to control specific properties that will be used when HLOD cluster Actors are merged, like generating lightmaps, combining Materials, transition size, and more.  
+
 ![](meshSettings.png)
 
 [REGION:raw]
 | Property | %spacer% | Description |
 | --- | :---: | :--- |
-| **Transition Screen Size** || The screen radius a mesh object should reach before swapping to the LOD **Actor**. When the parent LOD displays it won’t display any of the children. |
-| [REGION:tablesection]**Merge Settings**[/REGION] |||
-| **Generate Light Map UV** || Whether to generate a lightmap UV for a merged mesh. |
-| **Target Light Map UVChannel** || The target UV channel to assign the lightmap UV. |
-| **Target Light Map Resolution** || The lightmap resolution set for the merged mesh. |
-| **Pivot Point at Zero** || Whether a merged mesh should have a pivot at the World Origin, or at the first merged **Actor** from a generated cluster. |
-| **Merge Physics Data** || Whether to merge the collision primitives when combining the **Actors**. |
-| **Merge Materials** || Whether to merge source materials into one flat material. |
-| [REGION:tablesection]**[Material Settings See Below](#materialsettings)**[/REGION] |||
-| **Bake Vertex Data ** || Incorporates painted vertex colours into model/baked out material. |
-| **Use Vertex Data for Baking Material** || Whether or not vertex data such as vertex colours should be used when baking out materials. |
-| **Use Landscape Culling** || Whether or not to use available landscape geometry to cull away invisible triangles. |
+| **Transition Screen Size** || This is the size on screen (as a radius) that the mesh should reach before swapping to the HLOD Actor. When the parent LOD displays, it won't display any of the children. |
+| [REGION:tablesection] Merge Settings [/REGION] |||
+| **Generate Light Map UV** || When enabled, this will generate a lightmap UV for the merged mesh and assign the Lightmap Coordinate Index to 1. |
+| **Target Light Map Resolution** || The lightmap resolution you'd like to use for the HLOD Actor. |
+| **Pivot Point at Zero** || When the merged mesh should have a pivot located a the World Origin (coordinate 0,0,0), or if the pivot should be located at the first Actor listed when generating the HLOD cluster.|
+| **Merge Physics Data** ||  When enabled, the collision primitives will be merged for the combined Actors. |
+| **Merge Materials** || When enabled, the source Materials will be merged into a single flat material. |
+| **Material Settings** || The settings here enable you to control how your Materials are packed and combined when the HLOD Actor is generated. [INCLUDE:#materialsettings] |
+| **Bake Vertex Data to Mesh** || When enabled, any vertex colors that have been painted on the selected meshes will be incorporated into the generated HLOD Actor and Material that is created. |
+| **Use Vertex Data for Baking Material** || When enabled, vertex data will be used when baking out Materials. |
+| **Use Texture Binning** || This will calculate the varying output texture sizes according to their importance when packing the final atlas texture. |
+| **Use Landscape Culling** || When enabled, landscape geometry can be used to cull away (or remove) any triangle sections of the mesh that are not visible. |
 [/REGION]
 
-### Material Settings
 
-![](material.png)
+[VAR:spacer]![](%ROOT%/spacer.png)(w:150 convert:false)[/VAR]
+[VAR:spacer2]![](%ROOT%/spacer.png)(w:40 convert:false)[/VAR]
+[VAR:spacer3]![](%ROOT%/spacer.png)(w:70 convert:false)[/VAR]
+[VAR:spacer4]![](%ROOT%/spacer.png)(w:10 convert:false)[/VAR]
+<!--
+[EXCERPT:GPM]
+When you generate the proxy meshes this will use the clustered LODs to merge these **Actors** together into a single asset. Depending on what settings you changed for the HLOD levels your static meshes will generate based on this. This can affect your materials, lightmaps, and various other properties as well. 
+
+[REGION:note]
+Be aware that this process can take a very long time depending on your scene complexity and settings you've chosen for your HLOD levels.
+[/REGION]
+[/EXCERPT]
+
+[EXCERPT:TST]
+
+Texture Sizing Type Description
 
 [REGION:raw]
-| Property | %spacer% | Description |
+| Type | %spacer2% | Description |
+| --- | :---: | :--- |	
+| **Use TextureSize for all material properties** || Will use the value entered into Texture Size for sizing. |
+| **Use automatically biased texture sizes based on TextureSize** || This creates a biased set of {normal:diffuse:otherproperties} textures using the ratio of {1:0.5:0.25} so if the texture size is set to 512 it will output a 512 normal map, 256 diffuse map and 128 map for roughness specular etc (diffuse is clamped to 32x32 and roughness specular etc to 16x16) |
+| **Use per property manually overridden texture sizes** || Enables the per texture overrides, and uses those values instead of Texture Size. |
+| **Use Simplygon’s automatic texture sizing** || Uses Simplygon's algorithm for texture sizing based on its decimation settings. This should only be used if you are already using Simplygon. |
+[/REGION]
+[/EXCERPT]
+
+[EXCERPT:blendmode]
+| Selection | %spacer4% | Description|
+| --- | --- | --- |
+| **Opaque** || This blend mode defines a surface through which light neither passes nor penetrates. |
+| **Masked** || This blend mode defines the visibility in a binary (on/off) fashion by marking whether areas of the Material should be visible or not. |
+| **Translucent** || This blend mode defines the visibility of a Material by using an Opacity value. |
+| **Additive** || This blend mode takes the pixels of the Material and adds them to the pixels of the background. |
+| **Modulate** || This blend mode multiplies the value of the Material against the pixels of the background. |
+| **Alpha Composite (Premultiplied Alpha)** ||  This blend mode is used with textures that already have a pre-multiplied alpha. |
+[/EXCERPT:blendmode]
+
+
+-->
+
+
+<!----
+[INCLUDE:#materialsettings]
+
+[EXCERPT:materialsettings]
+[REGION:raw]
+| Property | %spacer3% | Description |
 | --- | :---: | :--- |
 | **Texture Size** || Size of the generated BaseColor texture map. |
 | **Texture Sizing Type** || [INCLUDE:#tst] |
@@ -119,40 +177,20 @@ Below, broken out by major section, are the properties that can be found in a **
 | **Specular Map** || Whether to generate a Specular map. |
 | **Specular Constant** || Allows to you set a roughness constant value in place of a texture map. |
 | **Emissive Map** || Whether to generate a Emissive map. |
+| **Opacity Map** || If checked, an Opacity map will be baked out for the resulting merged Actor. This nullifies the Opacity Constant property.  Must also use the **Blend Mode** for **Translucent** to show correctly in the generated instances. |
+| **Opacity Constant** || Only enabled if Opacity Map is unchecked, this sets a constant value for the Opacity property of the resulting merged Actor's Material. |
+| **Opacity Mask Map** || If checked, an Opacity Mask map will be baked out for the resulting merged Actor. This nullifies the Opacity Mask Constant property. Must also use the **Blend Mode** for **Masked** to show correctly in the generated instances. |
+| **Opacity Mask Constant** || Only enabled if Opacity Mask Map is unchecked, this sets a constant value for the Opacity Mask property of the resulting merged Actor's Material. |
 | **Diffuse Texture Size** || Override for the Diffuse Texture Size. Enabled when **Texture Sizing Type** is set to **Use per property manually overriden texture sizes**. |
 | **Normal Texture Size** || Override for the Normal Texture Size. Enabled when **Texture Sizing Type** is set to **Use per property manually overriden texture sizes**. |
 | **Metallic Texture Size** || Override for the Metallic Texture Size. Enabled when **Texture Sizing Type** is set to **Use per property manually overriden texture sizes**. |
 | **Roughness Texture Size** || Override for the Roughness Texture Size. Enabled when **Texture Sizing Type** is set to **Use per property manually overriden texture sizes**. |
 | **Specular Texture Size** || Override for the Specular Texture Size. Enabled when **Texture Sizing Type** is set to **Use per property manually overriden texture sizes**. |
 | **Emissive Texture Size** || Override for the Emissive Texture Size. Enabled when **Texture Sizing Type** is set to **Use per property manually overriden texture sizes**. |
+| **Opacity Texture Size** || Override for the Opacity Texture Size. Enabled when **Texture Sizing Type** is set to **Use per property manually overriden texture sizes**. |
+| **Opacity Mask Texture Size** || Override for the Opacity Mask Texture Size. Enabled when **Texture Sizing Type** is set to **Use per property manually overriden texture sizes**. |
+| **Blend Mode** || This determines how the material's color is blended with background colors. [INCLUDE:#blendmode] |
 [/REGION]
+[/EXCERPT:materialsettings]
 
-
-[VAR:spacer]![](%ROOT%/spacer.png)(w:150 convert:false)[/VAR]
-
-[VAR:spacer2]![](%ROOT%/spacer.png)(w:40 convert:false)[/VAR]
-
-<!--
-[EXCERPT:GPM]
-When you generate the proxy meshes this will use the clustered LODs to merge these **Actors** together into a single asset. Depending on what settings you changed for the HLOD levels your static meshes will generate based on this. This can affect your materials, lightmaps, and various other properties as well. 
-
-[REGION:note]
-Be aware that this process can take a **very long time** depending on your scene complexity and settings for your HLOD levels.
-[/REGION]
-[/EXCERPT]
-
-[EXCERPT:TST]
-
-Texture Sizing Type Description
-
-[REGION:raw]
-| Type | %spacer2% | Description |
-| --- | :---: | :--- |	
-| **Use TextureSize for all material properties** || Will use the value entered into **Texture Size** for sizing. |
-| **Use automatically biased texture sizes based on TextureSize** || This creates a biased set of **{normal:diffuse:otherproperties}** textures using the ratio of **{1:0.5:0.25}** so if the texture size is set to 512 it will output a 512 normal map, 256 diffuse map and 128 map for roughness specular etc (diffuse is clamped to 32x32 and roughness specular etc to 16x16 |
-| **Use per property manually overridden texture sizes** || Enables the per texture overrides, and uses those values instead of **Texture Size**. |
-| **Use Simplygon’s automatic texture sizing** || Uses **Simplygon's** sizing. |
-[/REGION]
-
-[/EXCERPT]
--->
+--->
\ No newline at end of file
diff --git a/Engine/Documentation/Source/Engine/Landscape/Creation/Plugin/Plugin.INT.udn b/Engine/Documentation/Source/Engine/Landscape/Creation/Plugin/Plugin.INT.udn
--- a/Engine/Documentation/Source/Engine/Landscape/Creation/Plugin/Plugin.INT.udn
+++ b/Engine/Documentation/Source/Engine/Landscape/Creation/Plugin/Plugin.INT.udn
@@ -14,17 +14,17 @@ Creating your own heightmap and weightmap formats for importing landscape data i
 * In order to create new importers, your plugin should create instances of objects implementing `ILandscapeHeightmapFileFormat` and `ILandscapeWeightmapFileFormat`, and register those objects with `ILandscapeEditorModulemodule::RegisterHeightmapFileFormat` and `ILandscapeEditorModulemodule::RegisterWeightmapFileFormat`, respectively.
 
 * Implementing the `ILandscapeHeightmapFileFormat` interface in a plugin requires overriding the following functions:
-1. **const FLandscapeFileTypeInfo& GetInfo() const**: Returns type information indicating which file types this class handles, and whether or not exporting is supported.
-1. **FLandscapeHeightmapInfo Validate(const TCHAR* HeightmapFilename) const** - Validates the named file, or rejects it and returns an error code and message.
-1. **FLandscapeHeightmapImportData Import(const TCHAR* HeightmapFilename, FLandscapeFileResolution ExpectedResolution) const** - Actually imports the file.
-1. **void Export(const TCHAR* HeightmapFilename, TArrayView<const uint16> Data, FLandscapeFileResolution DataResolution, FVector Scale) const** - Exports the file, if this format supports exporting (see the return value from `GetInfo`). This is the only function that doesn't need to be overridden in order to compile. However, if it is called without being overridden, it will call `check`.
-1. **(Destructor)** - Classes that implement this interface should use a virtual destructor, as they will be deleted via a pointer to the interface class.
+1. `const FLandscapeFileTypeInfo& GetInfo() const`: Returns type information indicating which file types this class handles, and whether or not exporting is supported.
+1. `FLandscapeHeightmapInfo Validate(const TCHAR* HeightmapFilename) const` - Validates the named file, or rejects it and returns an error code and message.
+1. `FLandscapeHeightmapImportData Import(const TCHAR* HeightmapFilename, FLandscapeFileResolution ExpectedResolution) const` - Actually imports the file.
+1. `void Export(const TCHAR* HeightmapFilename, TArrayView<const uint16> Data, FLandscapeFileResolution DataResolution, FVector Scale) const` - Exports the file, if this format supports exporting (see the return value from `GetInfo`). This is the only function that doesn't need to be overridden in order to compile. However, if it is called without being overridden, it will call `check`.
+1. `(Destructor)` - Classes that implement this interface should use a virtual destructor, as they will be deleted via a pointer to the interface class.
 
 * Implementing the `ILandscapeHeightmapFileFormat` interface is nearly identical, with only minor differences in some return types:
-1. **const FLandscapeFileTypeInfo& GetInfo() const** - Returns type information indicating which file types this class handles, and whether or not exporting is supported.
-1. **FLandscapeWeightmapInfo Validate(const TCHAR* WeightmapFilename) const** - Validates the named file, or rejects it and returns an error code and message.
-1. **FLandscapeWeightmapImportData Import(const TCHAR* WeightmapFilename, FLandscapeFileResolution ExpectedResolution) const** - Actually imports the file.
-1. **void Export(const TCHAR* WeightmapFilename, TArrayView<const uint16> Data, FLandscapeFileResolution DataResolution, FVector Scale) const** - Exports the file, if this format supports exporting (see the return value from `GetInfo`). This is the only function that doesn't need to be overridden in order to compile. However, if it is called without being overridden, it will call `check`.
-1. **(Destructor)** - Classes that implement this interface should use a virtual destructor, as they will be deleted via a pointer to the interface class.
+1. `const FLandscapeFileTypeInfo& GetInfo() const` - Returns type information indicating which file types this class handles, and whether or not exporting is supported.
+1. `FLandscapeWeightmapInfo Validate(const TCHAR* WeightmapFilename) const` - Validates the named file, or rejects it and returns an error code and message.
+1. `FLandscapeWeightmapImportData Import(const TCHAR* WeightmapFilename, FLandscapeFileResolution ExpectedResolution) const` - Actually imports the file.
+1. `void Export(const TCHAR* WeightmapFilename, TArrayView<const uint16> Data, FLandscapeFileResolution DataResolution, FVector Scale) const` - Exports the file, if this format supports exporting (see the return value from `GetInfo`). This is the only function that doesn't need to be overridden in order to compile. However, if it is called without being overridden, it will call `check`.
+1. `(Destructor)` - Classes that implement this interface should use a virtual destructor, as they will be deleted via a pointer to the interface class.
 
 * For further information and examples, you can see the interfaces in `LandscapeFileFormatInterfaces.h`, the .PNG implementations in `LandscapeFileFormatPng.cpp` and `LandscapeFileFormatPng.h` , and the .RAW implementations in `LandscapeFileFormatRaw.cpp` and `LandscapeFileFormatRaw.h`. All of this code is in the LandscapeEditor module in the engine.
diff --git a/Engine/Documentation/Source/Engine/Landscape/Custom/CreatingCustomHeightmapsAndLayers.INT.udn b/Engine/Documentation/Source/Engine/Landscape/Custom/CreatingCustomHeightmapsAndLayers.INT.udn
--- a/Engine/Documentation/Source/Engine/Landscape/Custom/CreatingCustomHeightmapsAndLayers.INT.udn
+++ b/Engine/Documentation/Source/Engine/Landscape/Custom/CreatingCustomHeightmapsAndLayers.INT.udn
@@ -10,6 +10,7 @@ SkillLevel:Intermediate
 related:Engine/Landscape/Editing
 related:Engine/Landscape/Editing/PaintMode
 related:Engine/Landscape/Creation
+tags:Landscape
 
 [TOC(start:2 end:3)]
 
diff --git a/Engine/Documentation/Source/Engine/LevelStreaming/WorldBrowser/WorldBrowser.INT.udn b/Engine/Documentation/Source/Engine/LevelStreaming/WorldBrowser/WorldBrowser.INT.udn
--- a/Engine/Documentation/Source/Engine/LevelStreaming/WorldBrowser/WorldBrowser.INT.udn
+++ b/Engine/Documentation/Source/Engine/LevelStreaming/WorldBrowser/WorldBrowser.INT.udn
@@ -143,7 +143,7 @@ To see the details for any level in the Level window, select it and then click o
 
 Each level in the world composition can have up to 4 LOD streaming levels. LOD streaming levels are similar to 
 mesh LOD. Depending on streaming distance settings streaming level will be replaced with  corresponding LOD level. LOD 
-levels automatically discovered by World Composition using pattern <Package name>_LOD#. Where number should be from 1 to 4.
+levels automatically discovered by World Composition using pattern [Package name]_LOD#. Where number should be from 1 to 4.
 LOD levels can be created manually or with automatic LOD map generator which requires Simplygon.
 
 [REGION:note]
diff --git a/Engine/Documentation/Source/Engine/Performance/ReducingPackageSize/ReducingPackageSize.INT.udn b/Engine/Documentation/Source/Engine/Performance/ReducingPackageSize/ReducingPackageSize.INT.udn
--- a/Engine/Documentation/Source/Engine/Performance/ReducingPackageSize/ReducingPackageSize.INT.udn
+++ b/Engine/Documentation/Source/Engine/Performance/ReducingPackageSize/ReducingPackageSize.INT.udn
@@ -9,219 +9,406 @@ Parent:Engine/Performance
 
 [VAR:Topic]
 [OBJECT:Topic]
-	[PARAM:image]
-		![%Engine/Performance/ReducingPackageSize:title%](Engine/Performance/ReducingPackageSize/package_topic.png)
-	[/PARAM]
-	[PARAM:title]
-		%Engine/Performance/ReducingPackageSize:title%
-	[/PARAM]
-	[PARAM:description]
-		%Engine/Performance/ReducingPackageSize:description%
-	[/PARAM]
-	[PARAM:path]
-		[RELATIVE:Engine/Performance/ReducingPackageSize]
-	[/PARAM]
+    [PARAM:image]
+        ![%Engine/Performance/ReducingPackageSize:title%](Engine/Performance/ReducingPackageSize/package_topic.png)
+    [/PARAM]
+    [PARAM:title]
+        %Engine/Performance/ReducingPackageSize:title%
+    [/PARAM]
+    [PARAM:description]
+        %Engine/Performance/ReducingPackageSize:description%
+    [/PARAM]
+    [PARAM:path]
+        [RELATIVE:Engine/Performance/ReducingPackageSize]
+    [/PARAM]
 [/OBJECT]
 [/VAR]
 
 [VAR:TopicCompact]
 [OBJECT:TopicCompact]
-	[PARAM:image]
-		![%Engine/Performance/ReducingPackageSize:title%](Engine/Performance/ReducingPackageSize/package_topic.png)
-	[/PARAM]
-	[PARAM:title]
-		%Engine/Performance/ReducingPackageSize:title%
-	[/PARAM]
-	[PARAM:description]
-		%Engine/Performance/ReducingPackageSize:description%
-	[/PARAM]
-	[PARAM:path]
-		[RELATIVE:Engine/Performance/ReducingPackageSize]
-	[/PARAM]
+    [PARAM:image]
+        ![%Engine/Performance/ReducingPackageSize:title%](Engine/Performance/ReducingPackageSize/package_topic.png)
+    [/PARAM]
+    [PARAM:title]
+        %Engine/Performance/ReducingPackageSize:title%
+    [/PARAM]
+    [PARAM:description]
+        %Engine/Performance/ReducingPackageSize:description%
+    [/PARAM]
+    [PARAM:path]
+        [RELATIVE:Engine/Performance/ReducingPackageSize]
+    [/PARAM]
 [/OBJECT]
 [/VAR]
 
 [TOC(start:2 end:3)]
 
 No matter what platform your UE4 project is targeting, reducing the size of the packaged game can be challenging. 
-In the following guide we will cover steps you can take to help reduce your projects final package size to be as small as possible using nothing but the tools provided to you in the UE4 Editor. 
+In the following guide, we will cover steps you can take to help reduce your projects final package size to be as small as possible using nothing but the tools provided to you in the UE4 Editor. 
 
 [EXCERPT:main]
-## Creating a new empty project
+
+## Creating a New Empty Project
+
 When you first start working on your Android based mobile project, you might be tempted to use an existing project as a base to work from or create a new project that has the Starter Content enabled. 
-**DO NOT** do this, instead, create a completely new, empty C++ or Blueprint based project, and then use the [Migration Tool](Engine/Content/Browser/UserGuide/Migrate) to bring over the content that you want to use. 
+**DO NOT** do this, instead, create an entirely new, empty C++ or Blueprint based project, and then use the [Migration Tool](Engine/Content/Browser/UserGuide/Migrate) to bring over the content that you want to use. 
 This way, you ensure that the only content in your project is content that needs to be there.
 
 ## Compressing Cooked Content
+
 The easiest and quickest way to reduce your APK package size is to tell UE4 to compress the APK packages during the packing process. 
 To turn on package compression, you need to do the following in the UE4 Editor.
 
 1. First, open up the project's settings by going to the Main Toolbar and selecting the **Edit** option, then selecting **Project Settings**.
 
-	[OBJECT:ToggleButtonContent]
-	[PARAMLITERAL:category]
-	OS
-	[/PARAMLITERAL]
-	[PARAMLITERAL:id]
-	windows
-	[/PARAMLITERAL]
-	[PARAMLITERAL:active]
-	[/PARAMLITERAL]
-	[PARAM:content]
-	![](T_Project_Settings_Menu.png)
-	[/PARAM]
-	[/OBJECT]
-
-	[OBJECT:ToggleButtonContent]
-	[PARAMLITERAL:category]
-	OS
-	[/PARAMLITERAL]
-	[PARAMLITERAL:id]
-	mac
-	[/PARAMLITERAL]
-	[PARAMLITERAL:active]
-	[/PARAMLITERAL]
-	[PARAM:content]
-	![](T_Project_Settings_Menu_Mac.png)
-	[/PARAM]
-	[/OBJECT]
+    [OBJECT:ToggleButtonContent]
+    [PARAMLITERAL:category]
+    OS
+    [/PARAMLITERAL]
+    [PARAMLITERAL:id]
+    windows
+    [/PARAMLITERAL]
+    [PARAMLITERAL:active]
+    [/PARAMLITERAL]
+    [PARAM:content]
+    ![](T_Project_Settings_Menu.png)
+    [/PARAM]
+    [/OBJECT]
+
+    [OBJECT:ToggleButtonContent]
+    [PARAMLITERAL:category]
+    OS
+    [/PARAMLITERAL]
+    [PARAMLITERAL:id]
+    mac
+    [/PARAMLITERAL]
+    [PARAMLITERAL:active]
+    [/PARAMLITERAL]
+    [PARAM:content]
+    ![](T_Project_Settings_Menu_Mac.png)
+    [/PARAM]
+    [/OBJECT]
 
 1. Under the Project section, click on the **Packaging** section to show the options for how the project will be packaged.
 
-	![](T_Project_Packaging.png)
+    [REGION:lightbox]
+    [![](T_Project_Packaging.png)(w:518)](T_Project_Packaging.png)
+    [/REGION]
+
+    [REGION:caption]
+    Click for full image.
+    [/REGION]
 
-1. Click on the **small white triangle** that is at the bottom of the Packaging settings to expose the **Advanced Project settings**.
+1. Click on the **Advanced Properties** that is at the bottom of the Packaging settings to expose the **Advanced Project Settings**.
 
-	![](T_Expose_Advanced_Props.png)
+    [REGION:lightbox]
+    [![](T_Expose_Advanced_Props.png)(w:518)](T_Expose_Advanced_Props.png)
+    [/REGION]
+
+    [REGION:caption]
+    Click for full image.
+    [/REGION]
 
 1. Look for the option labeled **Create compressed cooked packages** and enable it (if not already enabled).
 
-	![](T_Enable_Comp_Packages.png)
-	
+    [REGION:lightbox]
+    [![](T_Enable_Comp_Packages.png)(w:518)](T_Enable_Comp_Packages.png)
+    [/REGION]
+
+    [REGION:caption]
+    Click for full image.
+    [/REGION]
+    
 If you have not packaged your game with the Create compressed cooked packages checkmark box enabled you should see a huge difference in size when your project is re-packaged with it enabled.
-In fact it is not uncommon for some project's APK packages to be reduced in size by up to 50% when **Create compressed cooked packages** is enabled.
+In fact, it is not uncommon for some project's APK packages to be reduced in size by up to 50% when **Create compressed cooked packages** is enabled.
+
+## Shared Material Shaders & Libraries 
+
+Enabling both **Share Material Shader Code** and **Share Material Native Libraries** will help to reduce the overall package size of your project, but this comes at the cost of increased loading times. To enable these this option in your UE4 project, you will need to do the following:
+
+1. First, on the **Main Toolbar** go to  **Edit** > **Project Settings**.
+
+    [OBJECT:ToggleButtonContent]
+    [PARAMLITERAL:category]
+    OS
+    [/PARAMLITERAL]
+    [PARAMLITERAL:id]
+    windows
+    [/PARAMLITERAL]
+    [PARAMLITERAL:active]
+    [/PARAMLITERAL]
+    [PARAM:content]
+    ![](T_Project_Settings_Menu.png)
+    [/PARAM]
+    [/OBJECT]
+
+    [OBJECT:ToggleButtonContent]
+    [PARAMLITERAL:category]
+    OS
+    [/PARAMLITERAL]
+    [PARAMLITERAL:id]
+    mac
+    [/PARAMLITERAL]
+    [PARAMLITERAL:active]
+    [/PARAMLITERAL]
+    [PARAM:content]
+    ![](T_Project_Settings_Menu_Mac.png)
+    [/PARAM]
+    [/OBJECT]
+
+1. At the bottom of the **Packing** section, there will be the following two checkmark boxes that you should enable.
+
+    * **Share Material Shader Code**
+    * **Share Material Native Libraries**
+
+    [REGION:lightbox]
+    [![](Packing_Share_Material_Code.png)(w:490)](Packing_Share_Material_Code.png)
+    [/REGION]
+
+    [REGION:caption]
+    Click for full image.
+    [/REGION]
+
+## Excluding Editor Content
+
+Enabling this option will make sure that none of the content that the UE4 Editor uses will be packaged in your project. Note that enabling this could cause issues with missing content in projects that might make use of the Editor content. To enable these two options in your UE4 project, you will need to do the following:
+
+1. First, on the **Main Toolbar** go to  **Edit** > **Project Settings**.
+
+    [OBJECT:ToggleButtonContent]
+    [PARAMLITERAL:category]
+    OS
+    [/PARAMLITERAL]
+    [PARAMLITERAL:id]
+    windows
+    [/PARAMLITERAL]
+    [PARAMLITERAL:active]
+    [/PARAMLITERAL]
+    [PARAM:content]
+    ![](T_Project_Settings_Menu.png)
+    [/PARAM]
+    [/OBJECT]
+
+    [OBJECT:ToggleButtonContent]
+    [PARAMLITERAL:category]
+    OS
+    [/PARAMLITERAL]
+    [PARAMLITERAL:id]
+    mac
+    [/PARAMLITERAL]
+    [PARAMLITERAL:active]
+    [/PARAMLITERAL]
+    [PARAM:content]
+    ![](T_Project_Settings_Menu_Mac.png)
+    [/PARAM]
+    [/OBJECT]
+
+1. Under the Project section, click on the **Packaging** section to show the options for how the project will be packaged.
+
+    [REGION:lightbox]
+    [![](T_Project_Packaging.png)(w:518)](T_Project_Packaging.png)
+    [/REGION]
+
+    [REGION:caption]
+    Click for full image.
+    [/REGION]
+
+1. Click on the **Advanced Properties** that is at the bottom of the Packaging settings to expose the **Advanced Project Settings**.
+
+    [REGION:lightbox]
+    [![](T_Expose_Advanced_Props.png)(w:518)](T_Expose_Advanced_Props.png)
+    [/REGION]
+
+    [REGION:caption]
+    Click for full image.
+    [/REGION]
+
+1. Scroll down until you see the **Exclude editor content when cooking** option and enable it by clicking on the checkmark box next to it.
+
+    [REGION:lightbox]
+    [![](Packing_Exclude_Editor_Content.png)(w:518)](Packing_Exclude_Editor_Content.png)
+    [/REGION]
+
+    [REGION:caption]
+    Click for full image.
+    [/REGION]
+
+    [region:note]
+    Note that enabling this could cause issues with missing content in projects that might make use of the Editor content.
+    [/region]
 
 ## Setting up a Project's Levels
-One overlooked area that can lead to the bloating of APK package sizes is failing to properly setup your project's options, like which levels should be loaded by default, or which levels should be used between level transitions. 
+
+One overlooked area that can lead to the bloating of APK package sizes is failing to properly set up your project's options, like which levels should be loaded by default, or which levels should be used between level transitions. 
 To set which level (or levels) should be used for this type of interaction, you can do the following.
 
 1. First, on the **Main Toolbar** go to  **Edit** > **Project Settings**.
 
-	[OBJECT:ToggleButtonContent]
-	[PARAMLITERAL:category]
-	OS
-	[/PARAMLITERAL]
-	[PARAMLITERAL:id]
-	windows
-	[/PARAMLITERAL]
-	[PARAMLITERAL:active]
-	[/PARAMLITERAL]
-	[PARAM:content]
-	![](T_Project_Settings_Menu.png)
-	[/PARAM]
-	[/OBJECT]
-
-	[OBJECT:ToggleButtonContent]
-	[PARAMLITERAL:category]
-	OS
-	[/PARAMLITERAL]
-	[PARAMLITERAL:id]
-	mac
-	[/PARAMLITERAL]
-	[PARAMLITERAL:active]
-	[/PARAMLITERAL]
-	[PARAM:content]
-	![](T_Project_Settings_Menu_Mac.png)
-	[/PARAM]
-	[/OBJECT]
+    [OBJECT:ToggleButtonContent]
+    [PARAMLITERAL:category]
+    OS
+    [/PARAMLITERAL]
+    [PARAMLITERAL:id]
+    windows
+    [/PARAMLITERAL]
+    [PARAMLITERAL:active]
+    [/PARAMLITERAL]
+    [PARAM:content]
+    ![](T_Project_Settings_Menu.png)
+    [/PARAM]
+    [/OBJECT]
+
+    [OBJECT:ToggleButtonContent]
+    [PARAMLITERAL:category]
+    OS
+    [/PARAMLITERAL]
+    [PARAMLITERAL:id]
+    mac
+    [/PARAMLITERAL]
+    [PARAMLITERAL:active]
+    [/PARAMLITERAL]
+    [PARAM:content]
+    ![](T_Project_Settings_Menu_Mac.png)
+    [/PARAM]
+    [/OBJECT]
 
 1. Then, under **Project** in the **Maps & Modes** section look for the **Default Maps** section.
 
-	![](T_Set_Maps.png)
+    [REGION:lightbox]
+    [![](T_Set_Maps.png)(w:518)](T_Set_Maps.png)
+    [/REGION]
+
+    [REGION:caption]
+    Click for full image.
+    [/REGION]
+
+1. Click on the arrow next to **Game Default Map** and from the asset, list select the map you want to use for the games default starting level. 
+    
+    [REGION:lightbox]
+    [![](T_Project_Settings_Default_Game_Map.png)(w:518)](T_Project_Settings_Default_Game_Map.png)
+    [/REGION]
 
-1. Click on the arrow next to **Game Default Map** and from the asset list select the map you want to use for the games default starting level. 
-	
-	![](T_Project_Settings_Default_Game_Map.png)
+    [REGION:caption]
+    Click for full image.
+    [/REGION]
 
-1. Repeat this process until you have levels setup for each required map type.
+1. Repeat this process until you have level's input for each required map type.
 
 ## Selecting which Content to Package or not Package
-In the Packaging section of your project settings you can specify which maps and content should or should not be packaged with your game.  To specify which maps should be included in your project you will need to do the following:
+
+In the Packaging section of your project settings, you can specify which maps and content should or should not be packaged with your game.  To specify which maps should be included in your project you will need to do the following:
 
 [region:warning]
-If you do not specify which maps should be cooked then all maps will be cooked including any testing maps that you might have. Failing to not specify which maps should be cooked will result in your final packaged game being bigger than it needs to be. 
+If you do not specify which maps should be cooked then all maps will be cooked including any testing maps that you might have. Failing not to specify which maps should be cooked will result in your final packaged game being bigger than it needs to be. 
 [/region]
 
 1. First, open up the project's settings by going to the **Main Toolbar** and selecting **Edit** > **Project Settings**.
 
-	[OBJECT:ToggleButtonContent]
-	[PARAMLITERAL:category]
-	OS
-	[/PARAMLITERAL]
-	[PARAMLITERAL:id]
-	windows
-	[/PARAMLITERAL]
-	[PARAMLITERAL:active]
-	[/PARAMLITERAL]
-	[PARAM:content]
-	![](T_Project_Settings_Menu.png)
-	[/PARAM]
-	[/OBJECT]
-
-	[OBJECT:ToggleButtonContent]
-	[PARAMLITERAL:category]
-	OS
-	[/PARAMLITERAL]
-	[PARAMLITERAL:id]
-	mac
-	[/PARAMLITERAL]
-	[PARAMLITERAL:active]
-	[/PARAMLITERAL]
-	[PARAM:content]
-	![](T_Project_Settings_Menu_Mac.png)
-	[/PARAM]
-	[/OBJECT]
+    [OBJECT:ToggleButtonContent]
+    [PARAMLITERAL:category]
+    OS
+    [/PARAMLITERAL]
+    [PARAMLITERAL:id]
+    windows
+    [/PARAMLITERAL]
+    [PARAMLITERAL:active]
+    [/PARAMLITERAL]
+    [PARAM:content]
+    ![](T_Project_Settings_Menu.png)
+    [/PARAM]
+    [/OBJECT]
+
+    [OBJECT:ToggleButtonContent]
+    [PARAMLITERAL:category]
+    OS
+    [/PARAMLITERAL]
+    [PARAMLITERAL:id]
+    mac
+    [/PARAMLITERAL]
+    [PARAMLITERAL:active]
+    [/PARAMLITERAL]
+    [PARAM:content]
+    ![](T_Project_Settings_Menu_Mac.png)
+    [/PARAM]
+    [/OBJECT]
 
 1. Then, under **Project** in the **Packaging** section look for the **Packaging** option.
 
-	![](Content_To_Pack.png)
+    [REGION:lightbox]
+    [![](Content_To_Pack.png)(w:518)](Content_To_Pack.png)
+    [/REGION]
+
+    [REGION:caption]
+    Click for full image.
+    [/REGION]
 
 1. Click on the **small white triangle** at the bottom of the Packaging section to expose the advances packing options.
 
-	![](Advanced_Packing_Section.png)
+    [REGION:lightbox]
+    [![](Advanced_Packing_Section.png)(w:518)](Advanced_Packing_Section.png)
+    [/REGION]
+
+    [REGION:caption]
+    Click for full image.
+    [/REGION]
+
+1. Scroll to the lower part of the Advanced Packaging options till you see a checkmark box for **Cook everything in the project content directory(ignores list of maps below)** option. The options that you see here will allow you to specify which content and maps should be packaged or should not be packaged with your project.
 
-1. Scroll to the bottom of the advanced Packaging options till you see a checkmark box for **Cook everything in the project content directory(ignore list of maps below)** option. The options that you see here will allow you to specify which content and maps should be packaged or should not be packaged with your project.
+    [REGION:lightbox]
+    [![](Advanced_Packing_Options.png)(w:518)](Advanced_Packing_Options.png)
+    [/REGION]
 
-	![](Advanced_Packing_Options.png)
+    [REGION:caption]
+    Click for full image.
+    [/REGION]
 
 1. To include or exclude an item press the white plus sign next to the option you want to use to create a new entry. 
 
-	![](Add_New_Item.png)
+    [REGION:lightbox]
+    [![](Add_New_Item.png)(w:518)](Add_New_Item.png)
+    [/REGION]
 
-1. Then click on the three small white dots next to the new entry and from the windows that is displayed, choose which assets to include or exclude from the build.
+    [REGION:caption]
+    Click for full image.
+    [/REGION]
 
-	![](Select_Items_To_Cook.png)
+1. Then click on the three small white dots next to the new entry and from the windows that are displayed, choose which assets to include or exclude from the build.
+
+    [REGION:lightbox]
+    [![](Select_Items_To_Cook.png)(w:518)](Select_Items_To_Cook.png)
+    [/REGION]
+
+    [REGION:caption]
+    Click for full image.
+    [/REGION]
 
 1. When you have selected an asset, the entry box will now contain a link to where that asset is in your project's folder. 
 
-	![](Selected_Item_Location.png)
+    [REGION:lightbox]
+    [![](Selected_Item_Location.png)(w:518)](Selected_Item_Location.png)
+    [/REGION]
+
+    [REGION:caption]
+    Click for full image.
+    [/REGION]
 
 |Property Name | Description|
 |--------------|------------|
 |Cook everything in the project content directory(ignore list of maps below.|Cook all things in the project content directory.|
-|Cook only maps(this only affects cook all).|Cook only maps(this only affect the cookall flag).|
+|Cook only maps(this only affects cook all).|Cook only maps(this only affect the cook all flag).|
 |Create compressed cooked packages.|Create compressed cooked packages(decreased deployment size.|
-|Do not include editor content in this package may cause game to crash / error if you are using this content|Skip cooking Editor content.|
-|List of maps to include in a packaged build|List of maps to include when no other map list is specified on commandline.|
+|Do not include editor content in this package may cause game to crash/error if you are using this content|Skip cooking Editor content.|
+|List of maps to include in a packaged build|List of maps to include when no other map list is specified on command line.|
 |Additional Asset Directories to Cook|Directories containing .uasset files that should always be cooked regardless of whether they're referenced by anything in your project. Note: These paths are relative to your project Content directory.|
 |Directories to never cook|Directories that should never be cooked.|
-|Additional Non-Asset Directories to Package.|Directories containing files that should always be added to the .pak file(if using a .pak file; otherwise they're copied as individual files) This is used to stage additional file that you manually load via the UFS (Unreal File System) file IO API. Note:These paths are relative to your project Content directory.|
-|Additional Non-Asset Directories to Copy.|Directories containing files that should always be copied when packaging your project, but are not supposed to be part of the .pak file. This is used to stage additional files that you manually load with using the UFS (Unreal File System) file IO API, ex, third-party libraries that perform their own internal file IO. Note:These paths are relative to your project Content directory. |
+|Additional Non-Asset Directories to Package.|Directories containing files that should always be added to the .pak file(if using a .pak file; otherwise they're copied as individual files) This is used to stage additional file that you manually load via the UFS (Unreal File System) file IO API. Note: These paths are relative to your project Content directory.|
+|Additional Non-Asset Directories to Copy.|Directories containing files that should always be copied when packaging your project, but are not supposed to be part of the .pak file. This is used to stage additional files that you manually load with using the UFS (Unreal File System) file IO API, ex, third-party libraries that perform their own internal file IO. Note: These paths are relative to your project Content directory. |
 
 
 ## Checking what Content will be Cooked
+
 To ensure that only the content that is related to your project is added to your APK file, you can check which content is being added to your project's APK by looking in the **Cooked** folder in your project.
-You can find the Cooked folder by going to **(ProjectName)\Saved\Cooked** and then choosing the folder with the format you cooked your project for.
+You can find the Cooked folder by going to **(ProjectName)\Saved\Cooked** and then choose the folder with the format you cooked your project for.
 
 [region:note]
 Please note that you will only be able to see the content in the **Cooked** folder after you complete your first cook. 
@@ -237,7 +424,13 @@ windows
 [PARAMLITERAL:active]
 [/PARAMLITERAL]
 [PARAM:content]
-![](T_Cooked_Content_Folder.png)
+    [REGION:lightbox]
+    [![](T_Cooked_Content_Folder.png)(w:324)](T_Cooked_Content_Folder.png)
+    [/REGION]
+
+    [REGION:caption]
+    Click for full image.
+    [/REGION]
 [/PARAM]
 [/OBJECT]
 
@@ -251,13 +444,19 @@ mac
 [PARAMLITERAL:active]
 [/PARAMLITERAL]
 [PARAM:content]
-![](T_Cooked_Content_Folder_Mac.png)
+    [REGION:lightbox]
+    [![](T_Cooked_Content_Folder_Mac.png)(w:427)](T_Cooked_Content_Folder_Mac.png)
+    [/REGION]
+
+    [REGION:caption]
+    Click for full image.
+    [/REGION]
 [/PARAM]
 [/OBJECT]
 
 The image below shows the Texture folder from the Match 3 sample game. 
 The content in the folder has been sorted by size, so it is easier to see what assets are the biggest.
-Once the biggest assets have been determined, you can then examine the assets inside of UE4 to see if in fact they can be reduced in size without sacrificing the integrity of the asset.
+Once the most significant assets have been determined, you can then examine the assets inside of UE4 to see if in fact they can be reduced in size without sacrificing the integrity of the asset.
 
 [OBJECT:ToggleButtonContent]
 [PARAMLITERAL:category]
@@ -269,7 +468,13 @@ windows
 [PARAMLITERAL:active]
 [/PARAMLITERAL]
 [PARAM:content]
-![](T_Cooked_Content.png)
+    [REGION:lightbox]
+    [![](T_Cooked_Content.png)(w:473)](T_Cooked_Content.png)
+    [/REGION]
+
+    [REGION:caption]
+    Click for full image.
+    [/REGION]
 [/PARAM]
 [/OBJECT]
 
@@ -283,20 +488,29 @@ mac
 [PARAMLITERAL:active]
 [/PARAMLITERAL]
 [PARAM:content]
-![](T_Cooked_Content_Mac.png)
+    [REGION:lightbox]
+    [![](T_Cooked_Content_Mac.png)(w:429)](T_Cooked_Content_Mac.png)
+    [/REGION]
+
+    [REGION:caption]
+    Click for full image.
+    [/REGION]
 [/PARAM]
 [/OBJECT]
 
 ## Development VS Shipping Build Size
+
 When trying to figure out the final size your project will be, keep in mind that the Development version of your project will be slightly larger than your Shipping build size.
 On the Medieval Match example game, the difference in size between a Development and Shipping build was around 14 percent.
 However, since each project has different requirements, the saving between the two different build types for your project could be more or less than a 14 percent savings.
 
 ## Per-Device Texture LOD's
+
 With the release of Unreal Engine 4.8, you can now specify which size Texture should be used on which devices. 
 You can read more about this in the [Texture Guidelines for Mobile Platforms](Platforms/Mobile/Textures) section.
 
 ## Removing Unused Content
+
 When your project is finally ready to be packaged up for the store, make sure to first remove any unused or testing content by selecting it in the Content Browser and deleting it.
 You can delete content by first selecting the content you want to delete and then either press the **Delete** key on the keyboard or use the Delete option in the Right - Click menu. 
 When you do try and delete something from your project, the **Delete Assets** menu will be displayed like in the following image.
@@ -329,73 +543,81 @@ mac
 [PARAMLITERAL:active]
 [/PARAMLITERAL]
 [PARAM:content]
-![](DeleteAssets_Mac.png)
+    [REGION:lightbox]
+    [![](DeleteAssets_Mac.png)(w:305)](DeleteAssets_Mac.png)
+    [/REGION]
+
+    [REGION:caption]
+    Click for full image.
+    [/REGION]
 [/PARAM]
 [/OBJECT]
 
 The Delete Assets menu will inform you if the asset you are trying to delete is referenced by another asset. 
-If it is not referenced, you can simply delete it, however if it is referenced by other assets, you can use the **Replace References** option to replace the reference to that asset to another asset that is supposed to be packaged in your project.
+If it is not referenced, you can simply delete it, however, if it is referenced by other assets, you can use the **Replace References** option to replace the reference to that asset to another asset that is supposed to be packaged in your project.
 
 ## Removing Unused Plug-ins
+
 Disabling any unused plug-ins inside of the UE4 Editor is another way to ensure that unnecessary content and code do not make it into your project's final APK package file.
-This is because some plug-ins require a certain set of base assets and code to work correctly.
-If the plug-in is not disabled, the assets and code that are required to make the plug-in work will not be removed from your project.
+This is because some plug-ins require a particular set of base assets and code to work correctly.
+If the plug-in is not disabled, the assets and code that are needed to make the plug-in work will not be removed from your project.
 While this might not reduce your project's size as much as reducing or removing a large Texture would, every little bit can help when trying to slim your project down to the 100 MB size.
 To disable a plug-in inside of UE4, you need to do the following: 
 
 [region:note]
-	Make sure you fully test your project to ensure that the plug-ins you are disabling do not interfere with the functionality of your project.
+    Make sure you thoroughly test your project to ensure that the plug-ins you are disabling do not interfere with the functionality of your project.
 [/region]
 
 1. First, open up the Plug-ins Manager by going to the **Main Toolbar** and clicking on the **Window** option and then selecting the **Plug-in** option.
 
-	[OBJECT:ToggleButtonContent]
-	[PARAMLITERAL:category]
-	OS
-	[/PARAMLITERAL]
-	[PARAMLITERAL:id]
-	windows
-	[/PARAMLITERAL]
-	[PARAMLITERAL:active]
-	[/PARAMLITERAL]
-	[PARAM:content]
-	![](T_Open_Plugins_Browser.png)
-	[/PARAM]
-	[/OBJECT]
-
-	[OBJECT:ToggleButtonContent]
-	[PARAMLITERAL:category]
-	OS
-	[/PARAMLITERAL]
-	[PARAMLITERAL:id]
-	mac
-	[/PARAMLITERAL]
-	[PARAMLITERAL:active]
-	[/PARAMLITERAL]
-	[PARAM:content]
-	![](T_Open_Plugins_Browser_Mac.png)
-	[/PARAM]
-	[/OBJECT]
+    [OBJECT:ToggleButtonContent]
+    [PARAMLITERAL:category]
+    OS
+    [/PARAMLITERAL]
+    [PARAMLITERAL:id]
+    windows
+    [/PARAMLITERAL]
+    [PARAMLITERAL:active]
+    [/PARAMLITERAL]
+    [PARAM:content]
+    ![](T_Open_Plugins_Browser.png)
+    [/PARAM]
+    [/OBJECT]
+
+    [OBJECT:ToggleButtonContent]
+    [PARAMLITERAL:category]
+    OS
+    [/PARAMLITERAL]
+    [PARAMLITERAL:id]
+    mac
+    [/PARAMLITERAL]
+    [PARAMLITERAL:active]
+    [/PARAMLITERAL]
+    [PARAM:content]
+    ![](T_Open_Plugins_Browser_Mac.png)
+    [/PARAM]
+    [/OBJECT]
 
 1. From the Plug-in Browser, select the different sections and disable any plug-ins you are not using by un-checking the check mark box next to **Enabled**.
   When you have disabled all the plug-ins, you will need to re-start the UE4 Editor and then make sure to re-package your project.
 
-	[OBJECT:EmbeddedVideo]
-	[PARAMLITERAL:width]
-	640
-	[/PARAMLITERAL]
-	[PARAMLITERAL:height]
-	360
-	[/PARAMLITERAL]
-	[PARAMLITERAL:videoid]
-	xLJ01QEvV7Y
-	[/PARAMLITERAL]
-	[/OBJECT]
+    [OBJECT:EmbeddedVideo]
+    [PARAMLITERAL:width]
+    640
+    [/PARAMLITERAL]
+    [PARAMLITERAL:height]
+    360
+    [/PARAMLITERAL]
+    [PARAMLITERAL:videoid]
+    xLJ01QEvV7Y
+    [/PARAMLITERAL]
+    [/OBJECT]
 
 ## Package Black List
-With the release of Unreal Engine 4.9 and later you can now place a text file in your project's **Build/Platform/(Target Platform Folder)** directory that will tell the cooker to exclude partial or complete file paths from being packaged into your project.
+
+You can now place a text file in your project's **Build/Platform/(Target Platform Folder)** directory that will tell the cooker to exclude partial or complete file paths from being packaged into your project.
 Projects can have multiple Black List files for Debug, Development, Test, and Shipping builds that can be setup to include or exclude whatever project data you want.
-You can even have different Black List files for each platform your project supports like one for Android, one for iOS, ect. 
+You can even have different Black List files for each platform your project supports like one for Android, one for iOS, etc. 
 Here is an example of what the Black List files would look like for Tappy Chicken on Android. 
 
 [OBJECT:ToggleButtonContent]
@@ -408,7 +630,13 @@ windows
 [PARAMLITERAL:active]
 [/PARAMLITERAL]
 [PARAM:content]
-![](T_Blacklist_Packages.png)
+    [REGION:lightbox]
+    [![](T_Blacklist_Packages.png)(w:470)](T_Blacklist_Packages.png)
+    [/REGION]
+
+    [REGION:caption]
+    Click for full image.
+    [/REGION]
 [/PARAM]
 [/OBJECT]
 
@@ -422,63 +650,119 @@ mac
 [PARAMLITERAL:active]
 [/PARAMLITERAL]
 [PARAM:content]
-![](T_Blacklist_Packages_Mac.png)
+[REGION:lightbox]
+[![](T_Blacklist_Packages_Mac.png)(w:467)](T_Blacklist_Packages_Mac.png)
+[/REGION]
+
+[REGION:caption]
+Click for full image.
+[/REGION]
 [/PARAM]
 [/OBJECT]
 
-Here is an example from Tappy Chicken where the cooker is told to not include the following files when the game is cooked and packaged.
+Here is an example from Tappy Chicken where the cooker is told not to include the following files when the game is cooked and packaged.
 
 * **Blacklist Text File Location & Name:** TappyChicken/Build/Android/**PakBlacklist-Shipping.txt**:
 
-	[OBJECT:ToggleButtonContent]
-	[PARAMLITERAL:category]
-	OS
-	[/PARAMLITERAL]
-	[PARAMLITERAL:id]
-	windows
-	[/PARAMLITERAL]
-	[PARAMLITERAL:active]
-	[/PARAMLITERAL]
-	[PARAM:content]
-	![](T_Black_List.png)
-	[/PARAM]
-	[/OBJECT]
-
-	[OBJECT:ToggleButtonContent]
-	[PARAMLITERAL:category]
-	OS
-	[/PARAMLITERAL]
-	[PARAMLITERAL:id]
-	mac
-	[/PARAMLITERAL]
-	[PARAMLITERAL:active]
-	[/PARAMLITERAL]
-	[PARAM:content]
-	![](T_Black_List_Mac.png)
-	[/PARAM]
-	[/OBJECT]
-
-The first three entries should be safe for any project, but there may be cases were the **AssetRegistry.bin** is needed at runtime (depending on your project).
+    [OBJECT:ToggleButtonContent]
+    [PARAMLITERAL:category]
+    OS
+    [/PARAMLITERAL]
+    [PARAMLITERAL:id]
+    windows
+    [/PARAMLITERAL]
+    [PARAMLITERAL:active]
+    [/PARAMLITERAL]
+    [PARAM:content]
+    ![](T_Black_List.png)
+    [/PARAM]
+    [/OBJECT]
+
+    [OBJECT:ToggleButtonContent]
+    [PARAMLITERAL:category]
+    OS
+    [/PARAMLITERAL]
+    [PARAMLITERAL:id]
+    mac
+    [/PARAMLITERAL]
+    [PARAMLITERAL:active]
+    [/PARAMLITERAL]
+    [PARAM:content]
+    ![](T_Black_List_Mac.png)
+    [/PARAM]
+    [/OBJECT]
+
+The first three entries should be safe for any project, but there may be cases where the **AssetRegistry.bin** is needed at runtime (depending on your project).
 It is highly recommended that you thoroughly test to ensure that the content you are removing will not negatively impact your project.
-An easy and simple way to go about this is to open up all your project's levels to ensure that they correctly load and have no errors or warnings.
+An easy and straightforward way to go about this is to open up all your project's levels to ensure that they correctly load and have no errors or warnings.
 If you run into further problems and are not sure where they are coming from, check the Logcat for errors.
 
 ## Max Dynamic Point Lights
-With the release of UE 4.9 a new option has been introduced that can help reduce the number of shaders that are generated for dynamic lights. 
-Making this feature is particularly useful for mobile games that do not require dynamic lighting. 
-The property, **Max Dynamic Point Lights** can be found under the **Mobile** section of the **Project Settings**.
 
-![](Max_Dyn_Point_Lights.png)
+The property, **Max Dynamic Point Lights** which can be found under the **Mobile** section of the **Project Settings** can help reduce the number of shaders that are generated for dynamic lights. This feature is particularly useful for mobile games that do not require dynamic lighting.
+
+[REGION:lightbox]
+[![](Max_Dyn_Point_Lights.png)(w:518)](Max_Dyn_Point_Lights.png)
+[/REGION]
+
+[REGION:caption]
+Click for full image.
+[/REGION]
 
 To completely disable this feature, set the **Max Dynamic Point Lights** from **4** to **0** and then re-start the UE4 Editor when prompted.
 
 ![](Max_Dyn_Point_Lights_To_0.png)
 
-Depending on the size of your project and the amount of Materials used, the size saving you can achieve can range from just a few MB to the multiple MB range.
-Using the default Third Person template for an example, we can see what the size savings are when Max Dynamic Point Lights are enabled versus when Max Dynamic Point Lights are disabled completely. 
+Depending on the size of your project and the number of Materials used, the size saving you can achieve can range from just a few MB to the multiple MB ranges.
+Using the default Third Person template, for example, we can see what the size savings are when Max Dynamic Point Lights are enabled versus when Max Dynamic Point Lights are disabled completely. 
 
 | State| Packaged Game Size | Savings|
 | ---------| ------| -----|
 |ON|54.3 MB|0|
 |OFF| 53.2 MB| 1.1 MB|
-[/EXCERPT:main]
\ No newline at end of file
+[/EXCERPT:main]
+
+## Build with Hidden Symbols
+ 
+ You can build the shipping configuration with hidden visibility by default which will result in a smaller .so file. However, this will also remove symbols used to display call stack dumps so your ability to debug crashes, will be limited. To enable these this option in your UE4 project, you will need to do the following:
+
+1. First, open up the project's settings by going to the Main Toolbar and selecting the **Edit** option, then selecting **Project Settings**.
+
+    [OBJECT:ToggleButtonContent]
+    [PARAMLITERAL:category]
+    OS
+    [/PARAMLITERAL]
+    [PARAMLITERAL:id]
+    windows
+    [/PARAMLITERAL]
+    [PARAMLITERAL:active]
+    [/PARAMLITERAL]
+    [PARAM:content]
+    ![](T_Project_Settings_Menu.png)
+    [/PARAM]
+    [/OBJECT]
+
+    [OBJECT:ToggleButtonContent]
+    [PARAMLITERAL:category]
+    OS
+    [/PARAMLITERAL]
+    [PARAMLITERAL:id]
+    mac
+    [/PARAMLITERAL]
+    [PARAMLITERAL:active]
+    [/PARAMLITERAL]
+    [PARAM:content]
+    ![](T_Project_Settings_Menu_Mac.png)
+    [/PARAM]
+    [/OBJECT]
+
+1. Then go to **Platforms** > **Android** > **Advanced Build** and enabled the Build with hidden symbol visibility in shipping config by clicking on the checkmark box.
+
+    [REGION:lightbox]
+    [![](Android_AB_HiddenSymbol.png)(w:594)](Android_AB_HiddenSymbol.png)
+    [/REGION]
+
+    [REGION:caption]
+    Click for full image.
+    [/REGION]
+
diff --git a/Engine/Documentation/Source/Engine/Physics/Cloth/Overview/ClothOverview.INT.udn b/Engine/Documentation/Source/Engine/Physics/Cloth/Overview/ClothOverview.INT.udn
--- a/Engine/Documentation/Source/Engine/Physics/Cloth/Overview/ClothOverview.INT.udn
+++ b/Engine/Documentation/Source/Engine/Physics/Cloth/Overview/ClothOverview.INT.udn
@@ -1,170 +1,408 @@
 Availability:Public
 Crumbs:
-Title:Clothing Tool Overview
-Description:An overview of Cloth creation using in-editor tools in Unreal Engine 4.
+Title:Clothing Tool
+Description:An overview of Cloth creation using the in-Editor tools with Unreal Engine 4.
 Type:Overview
 SkillLevel:Intermediate
-Version:4.16
+Version:4.17
 Parent:Engine/Physics
+Related:Engine/Physics/Cloth/Reference
 Order:1
 Tags:Cloth
 Tags:Physics
-reference-image:SectionSelection.gif
-reference-image:PaintingCloth.gif
 Topic-image:Clothing_Topic.png
 
 
 
+[TOC(start:2 end:2)]
+
 [REGION:warning]
-This is an experimental feature and may change in later versions. In order to use this tool, you must first enable it in the Editor Preferences. You can do this by doing to the main menu > **Edit** > **Editor Preferences** > **Experimental** > and 
-enabling the checkbox next to **Clothing Tool**.
+The Clothing Tool is currently an experimental feature and may change in later versions. To use this tool, you must first enable it in the Editor Preferences. You can do this by doing to the 
+main menu > **Edit** > **Editor Preferences** > **Experimental** > and enabling the check box next to **Clothing Tool**.
 [/REGION]
 
-[TOC(start:2 end:2)]
 
-In Unreal Engine, version 4.16, APEX Cloth has now been replaced with the NVIDIA's NvCloth solver. This is a low-level clothing solver that is responsible for the particle simulation that runs 
+In Unreal Engine, version 4.16, APEX Cloth has now been replaced with the NVIDIA's NvCloth solver. It is a low-level clothing solver that is responsible for the particle simulation that runs 
 clothing. This clothing solver allows integrations to be lightweight and very extensible because we now have direct access to this data. Unlike APEX Clothing, which is a sophisticated SDK that
 is a one-stop shop for clothing simulation that handles both the simulation and the render. 
 
-![](OverviewImage.png)(w:1200)
-
-1. **Section Selection -** Use this to select the Material IDs that will be used for painting cloth onto. 
-1. **Cloth Paint Tab -** This section includes all the necessary tools and properties you can use when painting clothing. 
+[REGION:raw]
+![](ClothingOverviewBanner.png)
+[/REGION]
 
-With the Clothing Tools now available within the Editor, the workflow has been redefined to enable developers to work directly in Unreal Engine to author their content without the need for 
+With the Clothing Tools now available within the Editor, the workflow has been redefined to enable developers to work directly with Unreal Engine to author their content without the need for 
 external dependencies. 
 
 [REGION:imagetable]
-| [INCLUDE:#oldworkflow] | [INCLUDE:#newworkflow] |
+| [INCLUDE:#newworkflow] | [INCLUDE:#oldworkflow] |
 | --- | --- |
-| Old Cloth Workflow | New Cloth Workflow |
+| New Cloth Workflow | Old Cloth Workflow |
 | [INCLUDE:#clickimage] | [INCLUDE:#clickimage] |
 [/REGION]
 
 With the "Old" workflow using APEX Cloth, you were required to develop all your assets outside of Unreal using an external program (NVIDIA's Clothing Tool) or DCC Plugins that require specific
-software (3DS Max/Maya) to author this content. Only during the final stage of development would you bring your authored content into Unreal and assign it to your mesh. This meant that if you had
+software (3DS Max/Maya) to author this content. Only during the final stage of development would you bring your authored content into Unreal and assign it to your mesh which meant that if you had
 to correct anything you would have to go back to your program and re-author the content outside of the Editor. 
 
-The "New" workflow enables you to create your content once and then do all the creation editing of the clothing directly inside of Unreal. This makes the creation and iteration of testing content
+The "New" workflow enables you to create your content once and then do all the creation editing of the clothing directly inside of Unreal. It makes the creation and iteration of testing content
 that much quicker, and you also avoid any disconnects between where the content was created versus where it's being used by being able to see all of your edits for your clothing simulation
 happening in real-time and as they will appear in your game.
 
 
-## Creating a Cloth Section Asset
 
-To create a new Cloth asset, follow these steps.
 
-1. Click the **Section Selection** button in the main toolbar. This will enable you to select the different Material Elements of your Skeletal Mesh.
-	
-	[REGION:asyncgif]
-	![](SectionSelection.gif)(w:500)
-	[/REGION]
-1. With your section selected, you can right-click anywhere on the highlighted element to bring up the **Cloth Asset** options menu.
+## Cloth Painting Workflow
 
-	![](ClothAssetContextMenu.png)(w:450)
+Using the new workflow of creating cloth in-Editor, the Cloth Paint Tool enables you to quickly create clothing for your characters using any existing material element of a Skeletal Mesh.
 
-1. Hover over **Create Cloth Asset from Selection** and then fill in the following areas.
+![](OverviewImage.png)(w:1200)
 
-	![](CreateClothContextMenu.png)(w:350)
+1. **Section Selection -** Use this to select the Material IDs that will be used for painting cloth onto. 
+1. **Cloth Paint Panel -** This section includes all the necessary tools and properties you can use when painting clothing. 
 
-	* **Asset Name** - Give your asset a name so that you can easily find locate it later.
-	* **Remove from Mesh** - If you've got a separate mesh piece of geometry you want to be associated as cloth, you can enable this option. If not, you can leave this unchecked.
-	* **Physics Asset** - If this cloth asset is for a character, use its Physics Asset here to get proper collision for the cloth simulation. 
 
-1. Once you've set the information from the previous step you can click the **Create** button.
+Follow the steps below to learn the process to create your clothing assets, assign them to your character, and along with the basics of painting on your render mesh.
 
-	![](CreateClothButton.png)(w:350)
+### Create and Assign a Cloth Asset
 
-## Assigning a Cloth Section Asset
+To start working with clothing, you will first need to create a clothing asset from your Level of Detail (LOD) section and assign it to your render mesh so that you can paint on it. 
 
-To assign a Cloth asset to a section, follow these steps.
+Follow the steps below to get started:
 
-1. With your section element selected, you can right-click anywhere on the highlighted element to bring up the **Cloth Asset** options menu. Then hover over **Apply Clothing Asset** and 
-select from the available clothing assets to apply. This will associate the Cloth asset you created  with this Material Element geometry.
+1. In the Skeletal Mesh Editor, click the **Section Selection** button from the main toolbar. It will enable you to select the different parts of your Skeletal Mesh that have their Material Element assigned.
 
-	![](ClothSectionAssignment.png)(w:450)
+	[OBJECT:EmbeddedVideo]
+		[PARAMLITERAL:videoid]
+		9Xrfs1waIdA
+		[/PARAMLITERAL]
+		[PARAMLITERAL:width]
+		75
+		[/PARAMLITERAL]
+		[PARAMLITERAL:height]
+		56.25
+		[/PARAMLITERAL]
+		[PARAMLITERAL:units]
+		%
+		[/PARAMLITERAL]
+		[PARAMLITERAL:type]
+		simple
+		[/PARAMLITERAL]
+	[/OBJECT]
 
-## Painting Cloth Weights
+1. Use the left-click to select a part of your mesh that you want to use as cloth. Then right-click to open the context menu to create your Cloth Asset.
+	
+	[REGION:lightbox]
+	[![](ClothAssetContextMenu.png)(w:500)](ClothAssetContextMenu.png)
+	[/REGION]
+	[REGION:caption] Click image for full size. [/REGION]
 
-To begin painting your cloth weights, follow these steps.
+1. From the context menu, select **Create Cloth Asset from Selection** and then fill in the following areas of the menu:
 
-1. Go to the main menu and select **Window**, then locate and click **Cloth Paint**. This will open the Cloth Paint panel where you can click the **Enable Paint Tools** button to access
-the various options for painting cloth weights to onto your geometry. 
+	![](CreateClothContextMenu.png)(w:350)
 
-	![](ClothPanelEnabled.png)(w:350)
+	* **Asset Name** - Give your asset a name so that you can easily find locate it later.
+	* **Remove from Mesh** - If you've got a separate mesh piece of geometry you want to be associated as cloth, you can enable this option. If not, you can leave this unchecked.
+	* **Physics Asset** - If this cloth asset is for a character, use its Physics Asset here to get proper collision for the cloth simulation. 	
 
-1. To start painting, select your assigned Cloth Asset from the **Assets** window.
+	Once you're happy with the settings, click the **Create** button.
 	
-	![](HighlightClothAsset.png)(w:550)
+	![](CreateClothButton.png)(w:350)	
+	
+1. Right-click on the section again to bring up the context menu and hover over the **Apply Clothing Asset** and select from the available clothing assets to apply. It will associate any clothing 
+assets created with the section you've selected.
+
+	![](ClothSectionAssignment.png)(w:450)
+	
+
+### Painting On Your Cloth
+
+Once you're ready to paint on your cloth, you can use the following controls to Paint on your selected cloth asset.
 
-1. In the Cloth Painting section, set a **Paint Value** and then left-click and drag the brush across the surface of your clothing geometry you've selected.
+* Paint - **Left Mouse Button**
+* Erase - **Shift + Left Mouse Button**
+* Cloth Preview - **H**
+
+[REGION:note]
+If you've used the NVIDIA's PhysX DCC Plugin to create Clothing for 3DS Max or Maya, or similar painting tools in other programs, the controls should feel familiar for this operation. 
+[/REGION]
+
+
+1. In the Skeletal Mesh Editor, go to the file menu and select **Window**, then select **Cloth Paint** from the list. It will open the Cloth Paint panel. 
 	
-	[REGION:asyncgif]
-	![](PaintingCloth.gif)(w:500)
+	[REGION:imagetable]
+	| ![](OpenClothPaintPanel.png)(h:500) | | | | | ![](ClothPaintPanel.png)(h:500) |
+	| --- | --- | --- | --- | --- | --- |
+	|||
 	[/REGION]
 
-	[REGION:note]
-	If you've used the APEX Clothing Tools for 3DS Max or Maya, or similar painting tools in general, the controls should feel familiar. 
+	Click the **Enable Paint Tools** button to enable the properties you can use to paint your selected Cloth Asset.
+	
+	![](EnablePaintTools.png)
+	
+1. In the Cloth Paint panel, select your assigned **Cloth Asset** from the **Assets** list.
+
+	![](HighlightClothAsset.png)
+
+1. In the Cloth Painting section, you can select a [Paint Tool](#painttools) type to use and then set a **Paint Value** (by default, the Brush paint tool will be used). Then left-click and drag across the surface of your mesh to start painting.
+
+	[OBJECT:EmbeddedVideo]
+		[PARAMLITERAL:videoid]
+		PtevsF0qJCI
+		[/PARAMLITERAL]
+		[PARAMLITERAL:width]
+		75
+		[/PARAMLITERAL]
+		[PARAMLITERAL:height]
+		55
+		[/PARAMLITERAL]
+		[PARAMLITERAL:units]
+		%
+		[/PARAMLITERAL]
+		[PARAMLITERAL:type]
+		simple
+		[/PARAMLITERAL]
+	[/OBJECT]
+
+	[REGION:tip]
+	Hold the **H** keyboard shortcut to preview your painted cloth.
 	[/REGION]
 
-### Paint Tool Methods
 
-The **Paint Tool** selection enables you to choose the way you can paint your cloth weights; using a brush or a gradient. 
 
-#### Brush
+## Paint Tools
 
-The **Brush** tool enables you to manually paint weight values directly to your cloth. 
+![](PaintToolSelection.png)
 
-![](BrushTools.png)(w:350)
+The **Tool** selection enables you to choose from the brushes available to cloth paint.
 
-To paint your weight values, set a **Paint Value** and left-click and drag the mouse over the cloth geometry. 
+* [Brush](#brush)
+* [Gradient](#gradient)
+* [Smooth](#smooth)
+* [Fill](#fill)
+
+
+### Brush
+
+The **Brush** tool enables you to paint a radius and strength value on your cloth as you drag across cloth asset.
+
+![](Engine\Physics\Cloth\Reference\Brush.png)
+
+Use the **Paint Value** to control the strength of the brush when you are painting the cloth. This value controls how much the painted area will react like cloth based on this value. A value of 
+0 would mean that the skinned vertex could not move and a value of 100 would allow the skinned vertex to move one meter from its original position.
 
 ![](Paint_Brush.png)(w:500)
 
+[REGION:note]
+For additional information about this tool and its properties, see the [Brush Properties](Engine/Physics/Cloth/Reference#brush) reference here.
+[/REGION]
+
+### Gradient
 
-#### Gradient
+The **Gradient** tool enables you to paint a gradual blend between a set of cloth values that are selected. 
 
-The **Gradient** tool enables you to set weight values that will be painted between the selected points. Between these selected points a gradient will be painted that increases from the 
-**Gradient Start Value** to the **Gradient End Value**. 
+![](Engine\Physics\Cloth\Reference\Gradient.png)
 
-![](GradientTools.png)(w:350)
+To paint a gradient, perform the following steps:
 
-To paint your weight values, first set the **Gradient Start Value** and the **Gradient End Value**. Then, left-click and drag along your cloth mesh to paint the start value, which is represented
-by a green color over the painted vertices. Once you've done that, press the **Enter** (Return) key to switch and paint the end value, which is now represented by a red color that is painted over
-the vertices.
+1. Use the **Left-Mouse Button** to paint your **Gradient Start Value**. This is indicated by a **Green** dot over the painted vertice.
+1. Hold the **Ctrl + Left-Mouse Button** to paint your **Gradient End Value**. This is indicated by a **Red** dot over the painted vertice.
+1. When you are finished painting, press the **Enter** key to paint your gradient.
 
-Once you're finished painting your gradient values, press the **Enter** (Return) key again and you will see a gradient painted similarly to the example below.
+[OBJECT:EmbeddedVideo]
+	[PARAMLITERAL:videoid]
+	7iOYc3j2q4Q
+	[/PARAMLITERAL]
+	[PARAMLITERAL:width]
+	75
+	[/PARAMLITERAL]
+	[PARAMLITERAL:height]
+	48
+	[/PARAMLITERAL]
+	[PARAMLITERAL:units]
+	%
+	[/PARAMLITERAL]
+	[PARAMLITERAL:type]
+	simple
+	[/PARAMLITERAL]
+[/OBJECT]
 
 ![](Paint_Gradient.png)(w:500)
 
+[REGION:note]
+For additional information about this tool and its properties, see the [Gradient Properties](Engine/Physics/Cloth/Reference#gradient) reference here.
+[/REGION]
+
+### Smooth
+
+The **Smooth** tool enables you to blur or soften the contrast between painted cloth values. 
+
+![](Engine\Physics\Cloth\Reference\Smooth.png)
 
+Use the **Strength** value to adjust how strong or soft the blurring effect is to create a soft gradient between painted areas.
+
+![](Paint_Smooth.png)(w:500)
 
 [REGION:note]
-Alternatively, if you'd prefer to use a Brush while painting your gradient start and end points you can enable the option **Use Regular Brush for Gradient**. Then in the Brush section, use the **Radius** tool to set your brushes radius.
+For additional information about this tool and its properties, see the [Smooth Properties](Engine/Physics/Cloth/Reference#smooth) reference here.
 [/REGION]
 
+### Fill
+
+The **Fill** tool enables you to replace areas with similar values with another value. 
+
+![](Engine\Physics\Cloth\Reference\Fill.png)
 
+Use the **Fill Value** to set a value to fill the vertices in the area with. Use **Threshold** to set a value that the fill operation should use when sampling the vertices to replace. 
 
-## Assigning Cloth Assets to their Materials
+[OBJECT:ComparisonSlider]
+	[PARAM:before]
+	![Painted Area (White)](Paint_Fill1.png)(w:500)
+	[/PARAM]
+	[PARAM:after]
+	![Painted area using | Fill value of 1.0](Paint_Fill2.png)(w:500)
+	[/PARAM]
+[/OBJECT]
+
+[REGION:note]
+For additional information about this tool and its properties, see the [Fill Properties](Engine/Physics/Cloth/Reference#fill) reference here.
+[/REGION]
 
-To assign Cloth assets to the Material Element of your Skeletal Mesh, follow these steps.
+## Cloth Properties
 
-1. Navigate to the **Asset Details** panel and locate the **Clothing** section. This is where you'll find your available Cloth assets that can be assigned. You can adjust the various properties of your cloth
-by expanding the **Clothing Properties**.
+The **Cloth Config** properties enable you to adjust your cloth to simulate different types of materials, such as that of burlap, rubber, or many other types of cloth materials. 
+
+![](ClothConfig.png)
+
+[REGION:note]
+For additional information about the configuration properties for clothing, see the [Cloth Config Properties](Engine/Physics/Cloth/Reference#config) reference here.
+[/REGION]
+
+## Masks
+
+A **Mask** is a parameter set that can be swapped out for your clothing asset. 
+
+![](Engine\Physics\Cloth\Reference\Masks.png)
+
+This parameter set represents the following **Targets** (or parameter sets):
+
+* **Max Distance** is the maximum distance any point on the cloth can move from its animated position. 
+* **Backstop Distance** is the offset used from any point on the cloth to limit movement within the Max Distance. 
+* **Backstop Radius** uses a sphere radius that when intersected with the Max Distance will prevent any painted point on the cloth from moving into that region. 
+<!--- http://wiki.fitnect.hu/index.php/PhysX_setup --->
+
+You can have multiples of each of these **Targets** but only one of each can be assigned at a time. This makes it possible to test different configurations very quickly in a non-destructive way.
+
+To create a Mask and assign it to a Target, perform the following steps: 
+
+1. Click the **Plus** (+) button to create a new empty Mask. 
+
+	![](Masks01.png)(w:450)
+
+1. Right-click on the New Mask and hover the mouse over **Set Target** to select from the available **Targets** list.
+
+	![](Masks02.png)(w:450)
+
+1. In the Mask window, you will now see the Target you selected listed as an active target for your cloth asset.
+
+	![](Masks03.png)(w:450)
+
+[REGION:note]
+For additional information about masks, see the [Mask Properties](Engine/Physics/Cloth/Reference#masks) reference here.
+[/REGION]
+
+
+<!--- SAVE FOR 4.18 Release Update ------- 
+
+### Levels of Detail (LOD)
+
+There may be times where you want to use a lower **Level of Detail** (LOD) mesh(s) for your base mesh. When you set up your initial cloth asset for your 
+base mesh, it will not automatically carry over to any of your LODs.
+
+
+To create clothing for your LODs, perform the following steps:
+
+1. From the **Viewport** or from the **Cloth Paint** panel, use the **LOD** selection to choose the LOD mesh you want to apply cloth to.
+
+	[REGION:imagetable]
+	| IMAGE1 | IMAGE2 |
+	| --- | --- |
+	| Viewport LOD Selection | Cloth Paint panel LOD Selection |
+	[/REGION]
+
+1. Enable **Section Selection** from the main toolbar and use the **Left-Mouse click** to select the part of your mesh you want to create clothing for. Then right-click to open the context menu and select **Create Clothing LOD from Section**. Now, fill in the following options from the menu: 
+
+	* **Target Asset** - Select the target cloth asset you want to apply the LOD to. For some meshes you may have multiple clothing assets created.
+	* **LOD Index** - Choose the LOD you want to add or possibly one that you would want to replace if already created.
+	* **Remap Parameters** - If the LOD mesh is generally similar to the preceding LOD, the parameters that are set up there will try to be remapped to this LOD.
 	
-	![](AssetDetailsCloth.png)(w:350)
-1. In the **LOD** section, locate the **Material Element** you used to paint your cloth weights too. Use the **Clothing** drop-down selection to pick from the available cloth assets.
+	Once you're happy with the settings, click the **Create** button.
+
+1. Right-click on the section again to bring up the context menu and hover over the **Apply Clothing Asset** and select from the available clothing assets your new LOD. This will associate any clothing 
+assets created with the section you've selected.
+
 
-	![](AssignClothToMaterial.png)(w:350)
+### Repainting	LODs
 
-	[REGION:note]
-	Clothing Assets should only be assigned to the Material Elements that has had a clothing asset created and painted for it to use. 
+If you were to repaint your LOD0 mesh but don't want to have to re-paint all your LOD meshes, you could do the following:
+
+1. In the Cloth Paint panel under the **Assets** window, right-click your clothing asset to bring up the context menu.
+
+	IMAGE HERE
+	
+1. Select **Rebuild Asset Parameter Masks** to automate the process and carry the new painted cloth values over to your other LODs.
+
+	[REGION:imagetable]
+	| --- | --- |
+	| [INCLUDE:#lod0] | [INCLUDE:#lod1] |
 	[/REGION]
 
+	You will now have the newly painted values for LOD0 passed down to any lower LOD meshes like the comparison you can see here with LOD1 before and after.
+
+<!----
+
+[EXCERPT:lod0]
+[OBJECT:ComparisonSlider]
+	[PARAM:before]
+	![LOD0 Before](Tonemapper_On.png)
+	[/PARAM]
+	[PARAM:after]
+	![LOD0 After](Tonemapper_Off.png)
+	[/PARAM]
+[/OBJECT]
+[/EXCERPT:lod0]
 
+----  SAVE FOR 4.18 Release Update --->
 
 
+[EXCERPT:lod1]
+[OBJECT:ComparisonSlider]
+	[PARAM:before]
+	![LOD1 Before](Tonemapper_On.png)
+	[/PARAM]
+	[PARAM:after]
+	![LOD1 After](Tonemapper_Off.png)
+	[/PARAM]
+[/OBJECT]
+[/EXCERPT:lod1]
 
+--->
+
+
+## APEX (Legacy) Assets
+
+If you have created your clothing asset using the APEX Clothing Tool provided by NVIDIA for 3DS Max or Maya, you can continue to use this workflow. There are two ways to go about
+importing your APEX cloth (.apx or .apb) files into the Editor. 
+
+In the Skeletal Mesh Editor, you can use the **Asset Details** and under the **Clothing** section, use the **Add APEX clothing file...** button.
+
+![](APEXLegacyImport.png)
+
+Alternatively, you can enable the in-Editor Clothing Tool and under the **Assets** section of the **Cloth Paint** window, use the **plus** (+) button to import your APEX file. 
+
+![](InEditorClothImportAPEX.png)
+
+When you choose this method, your APEX cloth asset will be converted to the in-Editor tools and appropriately create and assign to your render mesh along with creating any masks for the different parameters
+like Max Distance and Backstop.
 
 
 <!---
@@ -189,4 +427,8 @@ _Click image for full size view._
 [/REGION]
 [/EXCERPT:clickimage]
 
---->
\ No newline at end of file
+--->
+
+
+
+
