Copying //UE4/Dev-Documentation to //UE4/Samples-Main (Source: //UE4/Dev-Documentation @ 3623665)

diff --git a/Engine/Documentation/Source/Engine/Physics/Cloth/Reference/ClothReference.INT.udn b/Engine/Documentation/Source/Engine/Physics/Cloth/Reference/ClothReference.INT.udn
--- a/Engine/Documentation/Source/Engine/Physics/Cloth/Reference/ClothReference.INT.udn
+++ b/Engine/Documentation/Source/Engine/Physics/Cloth/Reference/ClothReference.INT.udn
@@ -1,393 +1,348 @@
-Availability: Docs
-Crumbs:
-Title: Cloth Properties Reference
-Description: A reference page for the available properties used with Cloth asset configuration.
+Availability: Public
+Crumbs:%ROOT%
+Title: Clothing Tool Properties Reference
+Description: A reference page for the available properties used with Clothing Tool.
 Type:reference
-Version:4.16
-Parent:Engine/Physics
+Version:4.17
+Parent:Engine/Physics/Cloth
 Related:Engine/Physics/Cloth/Overview
 Order:2
 Tags:Physics
 Tags:Cloth
 
 
+[VAR:spacer]![](%ROOT%/spacer.png)(w:40 convert:false)[/VAR]
+[VAR:spacer2]![](%ROOT%/spacer.png)(w:40 convert:false)[/VAR]
+[VAR:spacer3]![](%ROOT%/spacer.png)(w:100 convert:false)[/VAR]
+[VAR:spacer4]![](%ROOT%/spacer.png)(w:150 convert:false)[/VAR]
 
-[OBJECT:modulesectionexpanded]
-	[PARAMLITERAL:id]
-		Cloth2
-	[/PARAMLITERAL]
-	[PARAM:heading]
-		Section Selection
-	[/PARAM]
-	[PARAM:content]
-		%C2%
-	[/PARAM]
-[/OBJECT]
+[TOC(start:2 end:2)]
 
-[VAR:C2]
+The **Cloth Paint Tools** have many options and properties that you can use to make very specific clothing simulations. Below you'll find the details about the menu options 
+you can use when creating your clothing assets and details about the Cloth Paint panel that you will use when selecting different painting tools to paint cloth values
+for your render mesh.
 
-[REGION:raw]
-![](SectionSelectionMenu.png)
-[/REGION]
 
-The **Section Selection** enables you to select specific Material Elements of your mesh in the Skeletal Mesh Editor's viewport. In this mode, you are able to create, assign, and remove any 
-cloth assets for this mesh. 
+## Clothing Asset Creation Menus
 
-[REGION:simpletable]
-| Property | Description |
-| --- | --- |
-| Apply Clothing Asset | Select an available clothing asset to apply to the selected section. |
-| Remove Clothing Asset | Remove the currently assigned clothing asset. |
-| Create Clothing Asset from Section | Create a clothing asset from the selected section. [INCLUDE:#createclothsection] |
+In this section, you'll find the details about the properties and settings you can use when creating and applying your clothing assets to your render mesh.
+
+### Section Selection
+
+The **Section Selection** enables you to select the different material elements of your render mesh to create and apply a clothing asset.  In this menu, you can identify the LOD and a material section of your 
+mesh that you've selected, create a cloth asset for your render mesh and its LODs, apply the clothing asset to the selected section, and remove it later if needed. 
+
+![](SectionSelection.png)
+
+[REGION:raw]
+| Property | %spacer% | Description |
+| --- | --- | --- |
+| **LOD Section Selection** || The name of the LOD level used and the Section to create a Cloth Asset for. |
+| **Apply Clothing Asset** || Select a clothing asset to apply to the selected section. |
+| **Remove Clothing Asset** || Remove the currently assigned clothing asset. |
+| [REGION:tablesection] Create Clothing Asset from Section [/REGION] || Create a new clothing asset using the selected section as a simulation mesh. [INCLUDE:#createclothsection] |
+| [REGION:tablesection] Create Clothing Asset from LOD Section [/REGION] || Create a clothing simulation mesh from the selected section and add its LOD to an existing clothing asset. [INCLUDE:#createclothsectionlod] |
 [/REGION]
 
-[/VAR]
 
-<!--- Section Selection excerpts --- 
+
+<!---
 
 [EXCERPT:createclothsection]
-| Property | Description |
-| --- | --- |
-| [REGION:tablesection] Basic [/REGION] ||
-| Asset Name | Set the name for the cloth asset that will be created. |
-| Remove From Mesh | Whether or not to leave this section behind (if driving a mesh with itself). Enable this if driving a high-poly mesh with a low poly.|
-| [REGION:tablesection] Collision [/REGION] ||
-| Physics Asset | The Physics Asset to extract collisions from. This will only extract spheres and sphyls, since this is what the simulation currently supports. |
-| [REGION:tablesection] Auto Fix[/REGION] ||
-| Try Auto Fix |  |
-| Auto Fix Threshold |  |
-| Simulated Particle Max Distance | The max distance to simulate the particles for the cloth solver. |
+![](CreateClothSection.png)(w:350)
+| --- | --- | --- |
+| [REGION:tablesection] Basic [/REGION] |||
+| **Asset Name** || The name entered for the cloth section asset. |
+| **Remove from Mesh** || Whether or not to leave this section behind (if driving a mesh with itself). Enable this is driving a high-poly mesh with a low-poly mesh. |
+| [REGION:tablesection] Collision [/REGION] |||
+| **Physics Asset** || Physics Asset to extract collisions from. Note that this will export Sphere and Sphylls, but does support up to 32 convex faces (or 5 boxes). |
+|| %spacer4% |||
+<!---| [REGION:tablesection] Auto Fix [/REGION] |||--->
+<!---| **Try Auto Fix** ||  |--->
+<!---| **Auto Fix Threshold** ||  |--->
+<!---| **Simulated Particle Max Distance** ||  |--->
 [/EXCERPT:createclothsection]
 
---- End Section Selection Excerpts --->
+[EXCERPT:createclothsectionLOD]
+![](CreateClothLODSection.png)(w:350)
+| --- | --- | --- |
+| [REGION:tablesection] Target [/REGION] |||
+| **Target Asset** || This is the target asset when importing LODs. |
+| **LOD Index** || [INCLUDE:#lodindex] |
+| **Remap Parameters** || If reimporting, this will map the old LOD parameters to the new LOD mesh. If adding a new LOD, this will map the parameters from the preceding LOD. |
+| [REGION:tablesection] Basic [/REGION] |||
+| **Remove from Mesh** || Whether or not to leave this section behind (if driving a mesh with itself). Enable this if driving a high-poly mesh with a low-poly one. |
+|| %spacer4% |||
+[/EXCERPT:createclothsectionLOD]
+
+[EXCERPT:lodindex]
+| --- | --- |
+| **Replace LOD** | Replace the simulation mesh in LOD0 of your selected clothing asset with this section. |
+| **Add LOD** | Use the selected section to add as a new LOD. |
+[/EXCERPT:lodindex]
+
 
+--->
 
 
-[OBJECT:modulesectionexpanded]
-	[PARAMLITERAL:id]
-		Cloth1
-	[/PARAMLITERAL]
-	[PARAM:heading]
-		Cloth Paint
-	[/PARAM]
-	[PARAM:content]
-		%C1%
-	[/PARAM]
-[/OBJECT]
+## Cloth Paint Panel
 
-[VAR:C1]
+The **Cloth Paint** panel houses all the clothing assets, masks, configuration parameters for how your cloth reacts, and the tools you'll use when painting your cloth values. 
 
 ![](ClothPaintPanel.png)
 
-The **Cloth Paint** properties enable you to define how you paint weights on the cloth asset you've selected. These weights determine how the **Cloth Configuration Properties** use 
-these values to applie cloth simulations to your mesh.
+### Assets
 
+The **Assets** window displays your currently created clothing assets that are assigned to a render mesh, enable you to import APEX (.apx or .apb) files, and lets you select from the available LODs 
+created for your mesh to copy parameter values to.
 
-[REGION:simpletable]
-| Property | Description  |
-| --- | --- |
-| Enable Paint Tools  | This enables the tools needed to paint your cloth values. |
-| [REGION:tablesection] Assets [/REGION] ||
-| Asset LOD | Selection box for the available levels of detail (LOD) where you can paint or create a mask. |
-| New Mask | Enables you to create a mask for the selected asset and available level of detail. |
-| [REGION:tablesection] Mask [/REGION] ||
-| Mask Context Menu | Available options for the asset mask. |
-| Set Target | Selects the available targets that you can apply the mask to. [INCLUDE:#settarget] |
-| Move Up | Moves the mask selection up in the Mask list. |
-| Move Down | Moves the mask selection down in the Mask list. |
-| Delete | Removes the selected mask from the list. |
-| Apply | Applies the mask to the physical mesh. |
-| [REGION:tablesection] Brush [/REGION] ||
-| Radius | Sets the radius that will be used for the painting brush. |
-| Strength | Sets the strength of the rush when painting between 0 - 1, with 0 being off and 1 being fully on. |
-| Falloff | Sets the amount of falloff around the edge of the brush to apply between 0 - 1, with 0 being no falloff and 1 being a lot of falloff. |
-| Enable Brush Flow | Enables the continuous 'flow' where paint is applied from the brush every tick. |
-| Ignore Back-Facing | Whether the back-facing triangles should be ignored. |
-| [REGION:tablesection] View [/REGION] ||
-| Color View Mode | Selection used to display the Vertex Colors view modes. [INCLUDE:#colorview] |
-| [REGION:tablesection] Cloth Painting [/REGION] ||
-| Painting Property | Sets the current clothing property which should be visualized and painted. [INCLUDE:#paintingproperty]|
-| Paint Tool | Select which type of paint tool to use. [INCLUDE:#brush] [INCLUDE:#gradient]|
+![](Assets.png)
+
+[REGION:raw]
+| Property | %spacer% | Description |
+| --- | --- | --- |
+| **Name** || The name of the cloth asset created for a LOD section of your physical mesh. |
+| **Import Apex File** || Allows you to import an APEX (.apx or .apb) file that will be converted to use the Clothing Tool. |
+| **Level of Detail (LOD) Selection** || Enables you to select a Level of Detail (LOD) mesh to set parameters, paint values, or copy values to. |
 [/REGION]
 
-[/VAR]
 
+### Masks
 
+The **Masks** window displays any created parameter sets for painted cloth values. These can be assigned a Target value for use with your clothing asset. 
 
-<!--- Cloth Paint panel excerpts --- 
+![](Masks.png)
 
-[EXCERPT:settarget]
-| Property | Description |
-| --- | --- |
-| None | No assigned mask target will be used. |
-| Max Distance | A mask used to store the Max Distance for your cloth asset. |
-| Backstop Distance | A mask used to store the Backstop Distance for your cloth asset.  |
-| Backstop Radius | A mask used to store the Backstop Radius for your cloth asset. |
-[/EXCERPT:settarget]
-
-[EXCERPT:colorview]
-| Property | Description |
-| --- | --- |
-| Off  | Normal view mode (Vertex Color visualization is off). | 
-| RGB Channels | Visualize the RGB color channels. | 
-| Alpha Channel | Visualize the Alpha channel. | 
-| Red Channel | Visualize the Red color channel.  | 
-| Green Channel | Visualize the Red color channel. |
-| Blue Channel | Visualize the Red color channel. | 
-[/EXCERPT:colorview]
-
-[EXCERPT:paintingproperty]
-| Property | Description |
-| --- | --- |
-| Max Distances | Selects the Max Distances property to be visualized and painted. | 
-| Backstop Distances  | Selects the Backstop Distances property to be visualized and painted. | 
-| Backstop Radius | Selects the Backstop Radius property to be visualized and painted. | 
-[/EXCERPT:paintingproperty]
+[REGION:raw]
+| Property | %spacer% | Description |
+| --- | --- | --- |
+| **Name** || The name given to the mask along with its target setting for this parameter set. [INCLUDE:#masksettings]  |
+| **Add Mask** || Adds a new Mask to the available list of Mask parameters. |
+[/REGION]
 
-[EXCERPT:paintingproperty]
-| Property | Description |
-| --- | --- |
-| Max Distances | Selects the Max Distances property to be visualized and painted. | 
-| Backstop Distances  | Selects the Backstop Distances property to be visualized and painted. | 
-| Backstop Radius | Selects the Backstop Radius property to be visualized and painted. | 
-[/EXCERPT:paintingproperty]
 
-[EXCERPT:brush]
-| Property | Description |
-| --- | --- |
-| Brush | Enables you to directly paint the meshes vertices. | 
-| Paint Value | The value to paint for the currently selected Painting Property. |
-[/EXCERPT:brush]
 
-[EXCERPT:gradient]
-| Property | Description |
-| --- | --- |
-| Gradient | Enables you to create a gradient between two sets of vertices. | 
-| Gradient Start Value | Sets the value of (green) points defined at the start of the gradient. |
-| Gradient End Value | Sets the value of (red) point defined at the end of the gradient. |
-| Use Regular Brush for Gradient | Toggle for using the regular brush size for painting the Start and End points of the gradient. |
-[/EXCERPT:gradient]
+<!---
+[INCLUDE:#masksettings]
 
-End Cloth Paint panel excerpts --->
+[EXCERPT:masksettings]
+| --- | --- | --- |
+| [REGION:tablesection] Right-Click Menu Settings [/REGION] |||
+| [REGION:tablesection] Set Target [/REGION] |||
+| **None** || The target has not been set for this parameter set. |
+| **Max Distance** || This is the maximum distance the clothing simulation particle can move from its animated position. |
+| **Backstop Distance** || This is the distance offset from the Max Distance used to limit the movement of the cloth simulation particles. |
+| **Backstop Radius** || This is a radius that when intersected with the Max Distance will prevent any cloth simulation particle from moving into that area. |
+| [REGION:tablesection] Actions [/REGION] |||
+| **Move Up** || Moves the Mask up one slot in the list.  |
+| **Move Down** || Move the Mask down one slot in the list. |
+| **Delete** || Deletes the Mask from the list. |
+| **Apply** || Applies the Mask to the physical mesh. |
+[/EXCERPT:masksettings]
 
+--->
 
 
+### Config
 
+The **Config** properties enables you to adjust how your cloth react so that it can simulate different types of material, such burlap, rubber, leather, and much more.
 
+![](Config.png)
 
-[OBJECT:modulesectionexpanded]
-	[PARAMLITERAL:id]
-		Cloth3
-	[/PARAMLITERAL]
-	[PARAM:heading]
-		Cloth Configuration Properties
-	[/PARAM]
-	[PARAM:content]
-		%C3%
-	[/PARAM]
-[/OBJECT]
+[REGION:raw]
+| Property | %spacer3% | Description |
+| --- | --- | --- |
+| **Physics Asset** || The Physics Asset to extract collisions from when building a simulation. |
+| **Wind Method** || How wind should be processed; Accurate uses drag and lift to make the cloth react differently, Legacy applies similar forces to all clothing without drag and lift (similar to APEX). [INCLUDE:#windmethod] |
+| **Vertical Constraint Config** || Constraint data for vertical constraints. [INCLUDE:#constraintgconfig] |
+| **Horizontal Constraint Config** || Constraint data for horizontal constraints. [INCLUDE:#constraintgconfig] |
+| **Bend Constraint Config** || Constraint data for bend constraints. [INCLUDE:#constraintgconfig] |
+| **Shear Constraint Config** || Constraint data for shear constraints. [INCLUDE:#constraintgconfig] |
+| **Self Collision Radius** || The size of self-collision spheres centered on each vertice. |
+| **Self Collision Stiffness** || The stiffness of the spring force that will resolve self-collisions. |
+| **Self Collision Cull State** || The scale to use for the radius of the culling checks for self-collisions. Any other self-collision body within the radius of this check will be culled. This helps performance with higher resolution meshes by reducing the number of colliding bodies within the cloth. Reducing this will have a negative impact on performance. |
+| **Damping** || The damping of particle motion per-axis. |
+| **Friction** || The friction of the surface when colliding. |
+| **Wind Drag Coefficient** || The drag coefficient for wind calculations, where higher values mean wind has a more lateral effect on cloth. |
+| **Wind Lift Coefficient** || The lift coefficient for wind calculations, where higher values make cloth rise easier in the wind. |
+| **Linear Drag** || The drag applied to linear particle movement per-axis. |
+| **Angular Drag** || The drag to angular particle movement, where higher values should limit material bending (per-axis). |
+| **Linear Inertia Scale** || The scale for linear particle inertia, or how much movement should translate to linear motion (per-axis). |
+| **Angular Inertia Scale** || The scale for angular particle inertia, or how much movement should translate to angular motion (per-axis). |
+| **Centrifugal Inertia Scale** || The scale for centrifugal particle inertia, or how much movement should translate to angular motion (per-axis). |
+| **Solver Frequency** || The frequency of the position solver, where lower values will lead to stretchier, bouncier cloth. |
+| **Stiffness Frequency** || The frequency for stiffness calculations, where lower values will degrade stiffness of constraints. |
+| **Gravity Scale** || The scale of gravity effect on the cloth particle simulation. |
+| **Tether Stiffness** || The scale for stiffness of particle tethers between each other. |
+| **Tether Limit** || The scale for the limit of particle tethers (how far they can separate). |
+| **Collision Thickness** || The "thickness" of the simulated cloth, which is used to adjust collisions. |
+[/REGION]
 
-[VAR:C3]
 
 
-![](ClothingProperties.png)
 
-In the **Asset Details** panel under the **Clothing** section, you can define the various properties of your clothing that is assigned to the mesh's Material Elements. These properties
-define how your cloth will react so that you can have different types of cloth that act like silk, burlap, rubber, etc.
 
-[REGION:simpletable]
-| Property  | Description |
-| --- | --- |
-| [REGION:tablesection] Mesh Clothing Assets [/REGION] ||
-| Name | The name given to the the created or imported Cloth asset. |
-| Details | Specific details about the cloth asset, such as the number of vertices simulating, fixed vertices, bones, etc. |
-| [REGION:tablesection] Cloth Config [/REGION] ||
-| Wind Method | A selection for how wind should be processed. [INLCUDE:#windmethod] |
-| Vertical Constraint Config | Constraint data used for vertical constraints. [INCLUDE:#constrconfig] |
-| Horizontal Constraint Config | Constraint data used for horizontal constraints. [INCLUDE:#constrconfig] |
-| Bend Constraint Config | Constraint data used for bend constraints. [INCLUDE:#constrconfig] |
-| Shear Constraint Config | Constraint data used for shear constraints. [INCLUDE:#constrconfig] |
-| Self Collision Radius | Size of the self collision spheres centered on each vertice.  |
-| Self Collision Stiffness | The stiffness of the spring force that will resolve self collisions. |
-| Damping | The damping of particle motions per-axis (XYZ).  |
-| Friction | The friction of the surface when colliding. |
-| Wind Drag Coefficient | Drag coefficient for wind calculations. Higher values mean that wind has more lateral effect on the cloth. |
-| Wind Lift Coefficient | Lift coefficient for wind calculations. Higher values make cloth rise easier in the wind. |
-| Linear Drag | Drag that is applied to linear particle movement per-axis (XYZ). |
-| Angular Drag | Drag that is applied to angular particle movement. Higher values should limit material bending (per-axis XYZ). |
-| Linear Inertia Scale | Scale for linear particle inertia. This controls how much movement should translate to linear motion per-axis (XYZ). |
-| Angular Inertia Scale | Scale for angular particle inertia. This controls how much movement should translate to angular motion per-axis (XYZ). |
-| Centrifugal Inertia Scale | Scale for centrifugal inertia. This controls how much movement should translate to angular motion per-axis (XYZ). |
-| Solver Frequency | The frequency of the position solver. Lower values will lead to stretchier, bouncier cloth. |
-| Stiffness Frequency | The frequency stiffness calculation. Lower values will degrade stiffness of constraints. |
-| Gravity Scale | The scale of gravity's effect on particles. |
-| Tether Stiffness | The scale for stiffness of particle tethers between each other. |
-| Tether Limit | The scale for the limit of particle tethers. This controls how far they can separate. |
-| Collision Thickness | This adjusts the thickness of the cloth, used to adjust collisions.  |
-| [REGION:tablesection] Import [/REGION] ||
-| Imported File Path | If using APEX Cloth (.apx/apb) assets from an imported file, this will be the original file location path.  |
+<!---
+[INCLUDE:#constraintgconfig]
+
+[EXCERPT:constraintgconfig]
+| --- | --- | --- |
+| **Stiffness** || How stiff this constraint is between the nodes. This affects how closely it will follow the desired position. |
+| **Stiffness Multiplier** || This is a multiplier affecting the **Stiffness** value used. |
+| **Stretch Limit** || The hard limit on how far this constraint can stretch. |
+| **Compression Limit** || The hard limit on how far this constraint can compress. |
+[/EXCERPT:constraintgconfig]
+
+[EXCERPT:windmethod]
+| --- | --- | --- |
+| **Legacy** || Use the Legacy wind mode, where accelerations are modified directly by the simulation with no regard to drag or lift. |
+<!--- | **Accurate** ||  | --->
+[/EXCERPT:windmethod]
+
+--->
+
+
+
+### Import
+
+The **Import** option displays the file path of any [imported APEX files](#assets).
+
+![](Import.png)
+
+[REGION:raw]
+| Property | %spacer2% | Description |
+| --- | --- | --- |
+| **Imported File Path** || If this asset was imported from a file, this will be the original path. |
 [/REGION]
 
 
-[/VAR]
+### Cloth Painting
 
-<!--- Asset Details: Cloth excerpts --- 
+The **Cloth Painting** section enables you to select between your different tools, such as brush, gradient, smooth, and fill. 
 
-[EXCERPT:windmethod]
-| Property | Description |
-| --- | --- |
-| Legacy | This mode uses accelerations modified directly by the simulation with no regard for drag or lift. |
-| Accurate| Use the updated wind calculation for NvCloth based solver, taking into account drag and lift. This will require those properties to be correctly set in the clothing configuration. |
-[/EXCERPT:windmethod]
+Before these properties can be populated, you must first select a clothing asset from your [Assets](#assets) window and then click the **Enable Paint Tools** button.
 
-[EXCERPT:constrconfig]
-| Property | Description |
-| --- | --- |
-| Stiffness | How stiff this constraint is. This affects how closely it will follow the desired position. |
-| Stiffness Multiplier | A multiplier that affects the value set for **Stiffness**. |
-| Stretch Limit | The hard limit on how far this constraint can stretch. |
-| Compression Limit | The hard limit on how far this constraint can compress. |
-[/EXCERPT:constrconfig]
+![](Enable.png)
+
+The type of tool to use when painting cloth values.
 
---- End Asset Details Cloth Excerpts --->
+* [Brush](#brush)
+* [Gradient](#gradient)
+* [Smooth](#smooth)
+* [Fill](#fill)
 
 
+#### Brush 
 
+The **Brush** tool enables you to paint a radius and strength value on your cloth as you drag across cloth asset.
 
+![](Brush.png)
 
 
+[REGION:raw]
+| Property | %spacer% | Description |
+| --- | --- | --- |
+| [REGION:tablesection] View [/REGION] |||
+| **View Min** || When painting float/1D values, this is considered the zero or black point for painted values. |
+| **View Max** || When painting float/1D values, this is considered the one or white point for painted values.  |
+| [REGION:tablesection] Advanced Rollout [/REGION] | | |
+| **Flip Normal** || Whether to flip normals on the mesh preview. |
+| **Cull Backface** || Whether to cull back facing triangles when rendering the mesh preview. |
+| **Opacity** || The opacity value of the mesh preview, which enables you to see through the mesh. |
+| [REGION:tablesection] Tool Settings [/REGION] |||
+| **Paint Value** || The value to paint onto the mesh for this parameter. |
+| [REGION:tablesection] Brush [/REGION] |||
+| **Radius** || The radius of the Brush used for painting. |
+| **Strength** || The strength of the brush (0.0 - 1.0). |
+| **Falloff** || The amount of fall off to apply (0.0 - 1.0). |
+| **Enable Brush Flow** || Enables "flow" painting where the paint is continually applied from the brush every tick. |
+| **Ignore back-facing** || Whether back-facing triangles should be ignored when painting. |
+[/REGION]
 
+#### Gradient 
 
+The **Gradient** tool enables you to paint a gradual blend between a set of cloth values that are selected. 
 
+![](Gradient.png)
 
-Availability: Public
-Title:Physics Simulation
-Crumbs: %ROOT%, Engine
-Description:Subsystem used to calculate collision and simulate physical actors.
-Version: 4.16
-parent:Engine
-order:8
-
-[VAR:Topic]
-[OBJECT:Topic]
-	[PARAM:image]
-		![%Engine/Physics:title%](Engine/Physics/physics_topic.png)
-	[/PARAM]
-	[PARAM:icon]
-		![](%ROOT%/physics_icon.png)(convert:false)
-	[/PARAM]
-	[PARAM:title]
-		%Engine/Physics:title%
-	[/PARAM]
-	[PARAM:description]
-		%Engine/Physics:description%
-	[/PARAM]
-	[PARAM:path]
-		[RELATIVE:Engine/Physics]
-	[/PARAM]
-[/OBJECT]
-[/VAR]
-[VAR:TopicCompact]
-[OBJECT:TopicCompact]
-	[PARAM:image]
-		![%Engine/Physics:title%](Engine/Physics/physics_topic.png)
-	[/PARAM]
-	[PARAM:icon]
-		![](%ROOT%/physics_icon.png)(convert:false)
-	[/PARAM]
-	[PARAM:title]
-		%Engine/Physics:title%
-	[/PARAM]
-	[PARAM:description]
-		%Engine/Physics:description%
-	[/PARAM]
-	[PARAM:path]
-		[RELATIVE:Engine/Physics]
-	[/PARAM]
-[/OBJECT]
-[/VAR]
-
-
-[REGION:banner]
-![Physics](physics_landing_banner.png)
+[REGION:raw]
+| Property | %spacer% | Description |
+| --- | --- | --- |
+| [REGION:tablesection] View [/REGION] |||
+| **View Min** || When painting float/1D values, this is considered the zero or black point for painted values. |
+| **View Max** || When painting float/1D values, this is considered the one or white point for painted values. |
+| [REGION:tablesection] Advanced Rollout [/REGION] | | |
+| **Flip Normal** || Whether to flip normals on the mesh preview. |
+| **Cull Backface** || Whether to cull back facing triangles when rendering the mesh preview. |
+| **Opacity** || The opacity value of the mesh preview, which enables you to see through the mesh. |
+| [REGION:tablesection] Tool Settings [/REGION] |||
+| **Gradient Start Value** || The value of the gradient at the start points. |
+| **Gradient End Value** || The value of the gradient at the end points. |
+| **Use Regular Brush** || Enables the painting of selected points using a brush rather than a point. |
+| [REGION:tablesection] Brush [/REGION] |||
+| **Radius** || The radius of the Brush used for painting. |
+| **Strength** || The strength of the brush (0.0 - 1.0). |
+| **Falloff** || The amount of falloff to apply (0.0 - 1.0). |
+| **Enable Brush Flow** || Enables "flow" painting where the paint is continually applied from the brush every tick. |
+| **Ignore back-facing** || Whether back-facing triangles should be ignored when painting. |
 [/REGION]
 
-[EXCERPT:Intro]
-Unreal Engine 4 uses the PhysX 3.3 physics engine to drive its physical simulation calculations and perform all collision calculations. PhysX provides the ability to perform accurate collision detection as well as simulate physical interactions between objects within the world. Having physics in your game will help improve the immersion value of every scene, as it helps players believe that they are interacting with the scene and that the scene is responding back in some way or another.
-[/EXCERPT:Intro]
-
-[OBJECT:TopicButtonList]
-	[PARAM:icon]
-		![](%ROOT%/start_icon.png)(convert:false)
-	[/PARAM]
-	[PARAM:title]
-		Essentials
-	[/PARAM]
-	[PARAM:links]
-		* [Collision Responses](Engine/Physics/Collision "%Engine/Physics/Collision:description%")
-		* [Collision Reference](Engine/Physics/Collision/Reference "%Engine/Physics/Collision/Reference:description%")
-		* [Physics Constraints User Guide](Engine/Physics/Constraints/ConstraintsUserGuide "%Engine/Physics/Constraints/ConstraintsUserGuide:description%")
-		* [](Engine/Physics/Constraints/ConstraintsBlueprints "%Engine/Physics/Constraints/ConstraintsBlueprints:description%")
-		* [](Engine/Physics/Constraints/ConstraintsReference "%Engine/Physics/Constraints/ConstraintsReference:description%")
-		* [](Engine/Physics/Tracing "%Engine/Physics/Tracing:description%")
-		* [Damping and Friction](Engine/Physics/Constraints/DampingAndFriction "%Engine/Physics/Constraints/DampingAndFriction:description%")
-	[/PARAM]
-[/OBJECT]
-
-[OBJECT:TopicButtonList]
-	[PARAM:icon]
-		![](%ROOT%/vehicle_icon.png)(convert:false)
-	[/PARAM]
-	[PARAM:title]
-		Clothing Tools
-	[/PARAM]
-	[PARAM:links]
-		* [Clothing Tool Overview](Engine/Physics/Cloth/Overview "%Engine/Physics/Cloth/Overview:description%")
-	[/PARAM]
-[/OBJECT]
-
-[OBJECT:TopicButtonList]
-	[PARAM:icon]
-		![](%ROOT%/vehicle_icon.png)(convert:false)
-	[/PARAM]
-	[PARAM:title]
-		Vehicles
-	[/PARAM]
-	[PARAM:links]
-		* [Vehicle User Guide](Engine/Physics/Vehicles/VehicleUserGuide "%Engine/Physics/Vehicles/VehicleUserGuide:description%")
-		* [Vehicle Content Creation Guide](Engine/Physics/Vehicles/VehcileContentCreation "%Engine/Physics/Vehicles/VehcileContentCreation:description%")
-		* [Advanced Vehicle Template](Engine/Physics/Vehicles/DoubleWishboneVehicle "%Engine/Physics/Vehicles/DoubleWishboneVehicle:description%")
-		* [Vehicle API Reference](API:AWheeledVehicle)
-	[/PARAM]
-[/OBJECT]
-
-[OBJECT:TopicButtonList]
-	[PARAM:icon]
-		![](%ROOT%/content_icon.png)(convert:false)
-	[/PARAM]
-	[PARAM:title]
-		Content Creation
-	[/PARAM]
-	[PARAM:links]
-		* [Physics Bodies](Engine/Physics/PhysicsBodies "%Engine/Physics/PhysicsBodies:description%")
-		* [Complex Vs. Simple Collision](Engine/Physics/SimpleVsComplex "%Engine/Physics/SimpleVsComplex:description%")
-		* [Physical Materials](Engine/Physics/PhysicalMaterials "%Engine/Physics/PhysicalMaterials:description%")
-		* [PhAT User Guide and Reference](Engine/Physics/PhAT "%Engine/Physics/PhAT:description%")
-	[/PARAM]
-[/OBJECT]
-
-[OBJECT:TopicButtonList]
-	[PARAM:icon]
-		![](%ROOT%/related_icon.png)(convert:false)
-	[/PARAM]
-	[PARAM:title]
-		Related Links
-	[/PARAM]
-	[PARAM:links]
-		* [Static Mesh Editor](Engine/Content/Types/StaticMeshes/Editor "%Engine/Physics/PhAT:description%")
-		* [FBX Pipeline: Static Meshes](Engine/Content/FBX/StaticMeshes "%Engine/Physics/PhAT:description%")
-	[/PARAM]
-[/OBJECT]
+
+#### Smooth 
+
+The **Smooth** tool enables you to blur or soften the contrast between painted cloth values. 
+
+![](Smooth.png)
+
+[REGION:raw]
+| Property | %spacer% | Description |
+| --- | --- | --- |
+| [REGION:tablesection] View [/REGION] |||
+| **View Min** || When painting float/1D values, this is considered the zero or black point for painted values. |
+| **View Max** || When painting float/1D values, this is considered the one or white point for painted values. |
+| [REGION:tablesection] Advanced Rollout [/REGION] | | |
+| **Flip Normal** || Whether to flip normals on the mesh preview. |
+| **Cull Backface** || Whether to cull back facing triangles when rendering the mesh preview. |
+| **Opacity** || The opacity value of the mesh preview, which enables you to see through the mesh. |
+| [REGION:tablesection] Tool Settings [/REGION] |||
+| **Strength** || The strength of the smoothing (blurring) effect when painting. |
+| [REGION:tablesection] Brush [/REGION] |||
+| **Radius** || The radius of the Brush used for painting. |
+| **Strength** || The strength of the brush (0.0 - 1.0). |
+| **Falloff** || The amount of falloff to apply (0.0 - 1.0). |
+| **Enable Brush Flow** || Enables "flow" painting where the paint is continually applied from the brush every tick. |
+| **Ignore back-facing** || Whether back-facing triangles should be ignored when painting. |
+[/REGION]
+
+
+#### Fill 
+
+The **Fill** tool enables you to replace areas with similar values with another value. 
+
+![](Fill.png)
+
+[REGION:raw]
+| Property | %spacer% | Description |
+| --- | --- | --- |
+| [REGION:tablesection] View [/REGION] |||
+| **View Min** || When painting float/1D values, this is considered the zero or black point for painted values. |
+| **View Max** || When painting float/1D values, this is considered the one or white point for painted values. |
+| [REGION:tablesection] Advanced Rollout [/REGION] | | |
+| **Flip Normal** || Whether to flip normals on the mesh preview. |
+| **Cull Backface** || Whether to cull back facing triangles when rendering the mesh preview. |
+| **Opacity** || The opacity value of the mesh preview, which enables you to see through the mesh. |
+| [REGION:tablesection] Tool Settings [/REGION] |||
+| **Threshold** || Threshold for fill operation that will keep filling until sampled vertices aren't within this range of the original vertex. |
+| **Fill Value** || The value to fill all select vertices to. |
+| [REGION:tablesection] Brush [/REGION] |||
+| **Radius** || The radius of the Brush used for painting. |
+| **Strength** || The strength of the brush (0.0 - 1.0). |
+| **Falloff** || The amount of falloff to apply (0.0 - 1.0). |
+| **Enable Brush Flow** || Enables "flow" painting where the paint is continually applied from the brush every tick. |
+| **Ignore back-facing** || Whether back-facing triangles should be ignored when painting. |
+[/REGION]
+
+
 
 
 
diff --git a/Engine/Documentation/Source/Engine/Physics/Destructibles/DestructibleProperties/DestructibleProperties.INT.udn b/Engine/Documentation/Source/Engine/Physics/Destructibles/DestructibleProperties/DestructibleProperties.INT.udn
--- a/Engine/Documentation/Source/Engine/Physics/Destructibles/DestructibleProperties/DestructibleProperties.INT.udn
+++ b/Engine/Documentation/Source/Engine/Physics/Destructibles/DestructibleProperties/DestructibleProperties.INT.udn
@@ -43,9 +43,9 @@ Version: 4.9
 | **Impact Resistance** | When a chunk takes impact damage due to physical contact, this parameter is the maximum impulse the contact can generate. Weak materials will have this set low, so heavier objects can pass through them during fracture. 0.0 is interpreted at infinite. |
 | **Default Impact Damage** | Impact damage will be taken to this depth. If negative, impact damage is disabled. |
 | [REGION:tablesection] _Debris Parameters_ [/REGION] ||
-| **Debris Lifetime Min** | "Debris chunks" (see debris Depth, below) will be destroyed after a time (in seconds) separated from non-debris chunks.  The actual lifetime is interpolated between these	two bDebrisTimeout, based upon the module's LOD setting.  To disable lifetime, clear the bDebrisTimeout flag in the flags field. If debrisLifetimeMax < debrisLifetimeMin, the mean of the two is used for both.
+| **Debris Lifetime Min** | "Debris chunks" (see debris Depth, below) will be destroyed after a time (in seconds) separated from non-debris chunks.  The actual lifetime is interpolated between these	two bDebrisTimeout, based upon the module's LOD setting.  To disable lifetime, clear the bDebrisTimeout flag in the flags field. If `debrisLifetimeMax < debrisLifetimeMin`, the mean of the two is used for both.
 | **Debris Lifetime Max** | |
-| **Debris Max Separation Min** | "Debris chunks" (see debris Depth, below) will be destroyed if they are separated from their origin by a distance greater than maxSeparation. The actual maxSeparation is interpolated between these two values, based upon the module's LOD setting.  To disable maxSeparation, clear the bDebrisMaxSeparation flag in the flags field. If debrisMaxSeparationMax < debrisMaxSeparationMin, the mean of the two is used for both.
+| **Debris Max Separation Min** | "Debris chunks" (see debris Depth, below) will be destroyed if they are separated from their origin by a distance greater than maxSeparation. The actual maxSeparation is interpolated between these two values, based upon the module's LOD setting.  To disable maxSeparation, clear the bDebrisMaxSeparation flag in the flags field. If `debrisMaxSeparationMax < debrisMaxSeparationMin`, the mean of the two is used for both.
 | **Debris Max Separation Max** | |
 | **Valid Bounds** | "Debris chunks" (see debris Depth, below) will be destroyed if they leave this box. The box translates with the destructible Actor's initial position, but does not rotate or scale. |
 | [REGION:tablesection] _Advanced Parameters_ [/REGION] ||
diff --git a/Engine/Documentation/Source/Engine/Physics/PhAT/HowTo/KinematicWithSimulatedParents/KinematicWithSimulatedParents.INT.udn b/Engine/Documentation/Source/Engine/Physics/PhAT/HowTo/KinematicWithSimulatedParents/KinematicWithSimulatedParents.INT.udn
--- a/Engine/Documentation/Source/Engine/Physics/PhAT/HowTo/KinematicWithSimulatedParents/KinematicWithSimulatedParents.INT.udn
+++ b/Engine/Documentation/Source/Engine/Physics/PhAT/HowTo/KinematicWithSimulatedParents/KinematicWithSimulatedParents.INT.udn
@@ -17,7 +17,7 @@ tags:PhAT
 
 The [](Engine/Physics/PhAT) provides several different ways in which you can simulate physics bodies including the ability to have kinematic physics bodies with simulated parents. 
 This enables you to define child bodies that can be driven by animation data while the parents of those bodies are driven by physics simulation data. 
-This can be useful if you have a character hanging from a ledge playing a hanging idle or traversing across a ledge but want them to react to impulses like falling rocks or other debris and generate a physics based reaction. 
+This can be useful if you have a character that you want to perform a physics based reaction while playing a hanging idle or hanging traversal animation.
 
 In this How-to, we use kinematic bodies with simulated parents to generate the effect of a character hanging from a ledge while the rest of the body is simulated with physics. 
 
@@ -45,7 +45,7 @@ For this guide, we are using the **Blueprint Third Person Template** project wit
 
 	By setting these bones to Kinematic, they will not simulate physics and will follow any animation data. 
 
-	Alternatively, you can **Right-click** bones in the hierarchy list and set their **Physics Type** from the roll-out menu. 
+	Alternatively, you can right-click bones in the hierarchy list and set their **Physics Type** from the roll-out menu. 
 
 	![](KinematicHowTo_04.png)
 
@@ -69,7 +69,7 @@ For this guide, we are using the **Blueprint Third Person Template** project wit
 
 	![](KinematicHowTo_08.png)
 
-	The hands (because they are set to kinematic) will follow the animation data contained within the ThirdPersonJump_Loop motion. 
+	The hands (because they are set to kinematic) will follow the animation data contained within the **ThirdPersonJump_Loop** motion. 
 
 	![](KinematicHowTo_09.png)
 
@@ -147,5 +147,6 @@ The arms and head are set to Kinematic (denoted below with Gold Boxes) while the
 
 
 
+
 
 
diff --git a/Engine/Documentation/Source/Engine/Physics/Physics.INT.udn b/Engine/Documentation/Source/Engine/Physics/Physics.INT.udn
--- a/Engine/Documentation/Source/Engine/Physics/Physics.INT.udn
+++ b/Engine/Documentation/Source/Engine/Physics/Physics.INT.udn
@@ -83,6 +83,7 @@ Unreal Engine 4 uses the PhysX 3.3 physics engine to drive its physical simulati
 	[/PARAM]
 	[PARAM:links]
 		* [Clothing Tool Overview](Engine/Physics/Cloth/Overview "%Engine/Physics/Cloth/Overview:description%")
+		* [Clothing Tool Reference](Engine/Physics/Cloth/Reference "%Engine/Physics/Cloth/Reference:description%")
 	[/PARAM]
 [/OBJECT]
 
diff --git a/Engine/Documentation/Source/Engine/Plugins/Ansel/AnselPhotography.INT.udn b/Engine/Documentation/Source/Engine/Plugins/Ansel/AnselPhotography.INT.udn
--- a/Engine/Documentation/Source/Engine/Plugins/Ansel/AnselPhotography.INT.udn
+++ b/Engine/Documentation/Source/Engine/Plugins/Ansel/AnselPhotography.INT.udn
@@ -1,6 +1,6 @@
 Availability:Public
-TItle:Nvidia Ansel Photography Plugin
-Description:Integration of Nvidia's Ansel in-game photography tool that enables capturing and sharing enhanced images from games
+TItle:NVIDIA Ansel Photography Plugin
+Description:Integration of NVIDIA's Ansel in-game photography tool that enables capturing and sharing enhanced images from games.
 Crumbs:
 Parent:Engine/Plugins
 Order:1
diff --git a/Engine/Documentation/Source/Engine/Plugins/Ansel/Overview/AnselOverview.INT.udn b/Engine/Documentation/Source/Engine/Plugins/Ansel/Overview/AnselOverview.INT.udn
--- a/Engine/Documentation/Source/Engine/Plugins/Ansel/Overview/AnselOverview.INT.udn
+++ b/Engine/Documentation/Source/Engine/Plugins/Ansel/Overview/AnselOverview.INT.udn
@@ -1,6 +1,6 @@
 Availability:Public
-TItle:Ansel Plugin Overview
-Description:Overview of the Unreal Engine 4 integration of Nvidia's Ansel in-game photography tool that enables capturing and sharing enhanced images from games
+TItle:NVIDIA Ansel Plugin Overview
+Description:Overview of the Unreal Engine 4 integration of NVIDIA's Ansel in-game photography tool that enables capturing and sharing enhanced images from games.
 Crumbs:
 Parent:Engine/Plugins/Ansel
 Order:1
@@ -11,9 +11,9 @@ version:4.14
 topic-image:Engine/Plugins/Ansel/Ansel_topic.png
 
 [EXCERPT:intro]
-Nvidia's Ansel technology makes it simple for PC gamers to capture and share enhanced images from their games. Integrating the Ansel UE4 plugin into your game is quick and easy.
+NVIDIA's Ansel technology makes it simple for PC gamers to capture and share enhanced images from their games. Integrating the Ansel UE4 plugin into your game is quick and easy.
 
-An end-user overview of Ansel is available on [Nvidia's Ansel Site](http://www.geforce.com/hardware/technology/ansel).
+An end-user overview of Ansel is available on [NVIDIA's Ansel Site](http://www.geforce.com/hardware/technology/ansel).
 
 The Ansel features available in Unreal Engine 4 are:
 
@@ -29,7 +29,7 @@ The Ansel features available in Unreal Engine 4 are:
 | Unreal Engine Version | 4.14 or later |
 | Operating System | Microsoft Windows |
 | GPU | NVIDIA GTX 780 or better [REGION:none]A full compatibility list is available on the [Ansel Supported GPUs page](http://www.geforce.com/hardware/technology/ansel/supported-gpus)[/REGION] |
-| Display Driver | Nvidia Geforce 372.90 or later |
+| Display Driver | NVIDIA GeForce 372.90 or later |
 [/EXCERPT:intro]
 
 ## Enabling Ansel Photography for Your Game
@@ -40,7 +40,7 @@ The Ansel features available in Unreal Engine 4 are:
 
 	![](Settings_Menu.png)
 
-1.	Select the **Photography** category and locate the **Nvidia Ansel Photography Plugin**. Toggle the _Enabled_ option.
+1.	Select the **Photography** category and locate the **NVIDIA Ansel Photography Plugin**. Toggle the _Enabled_ option.
 
 	![](Plugin_Enabled.png)
 
diff --git a/Engine/Documentation/Source/Engine/Plugins/Ansel/Reference/BlueprintAPI/AnselBlueprintAPI.INT.udn b/Engine/Documentation/Source/Engine/Plugins/Ansel/Reference/BlueprintAPI/AnselBlueprintAPI.INT.udn
--- a/Engine/Documentation/Source/Engine/Plugins/Ansel/Reference/BlueprintAPI/AnselBlueprintAPI.INT.udn
+++ b/Engine/Documentation/Source/Engine/Plugins/Ansel/Reference/BlueprintAPI/AnselBlueprintAPI.INT.udn
@@ -1,6 +1,6 @@
 Availability:Public
-TItle:Ansel Blueprint Node Reference
-Description:Reference of the configuration variables available for the Nvidia Ansel Photography Plugin
+TItle:NVIDIA Ansel Blueprint Node Reference
+Description:Reference of the configuration variables available for the NVIDIA Ansel Photography Plugin
 Crumbs:
 Parent:Engine/Plugins/Ansel
 Order:1
diff --git a/Engine/Documentation/Source/Engine/Plugins/Ansel/Reference/ConsoleVariables/AnselConsoleVars.INT.udn b/Engine/Documentation/Source/Engine/Plugins/Ansel/Reference/ConsoleVariables/AnselConsoleVars.INT.udn
--- a/Engine/Documentation/Source/Engine/Plugins/Ansel/Reference/ConsoleVariables/AnselConsoleVars.INT.udn
+++ b/Engine/Documentation/Source/Engine/Plugins/Ansel/Reference/ConsoleVariables/AnselConsoleVars.INT.udn
@@ -1,6 +1,6 @@
 Availability:Public
-TItle:Ansel Console Variable Reference
-Description:Reference of the configuration variables available for the Nvidia Ansel Photography Plugin
+TItle:NVIDIA Ansel Console Variable Reference
+Description:Reference of the configuration variables available for the NVIDIA Ansel Photography Plugin
 Crumbs:
 Parent:Engine/Plugins/Ansel
 Order:3
@@ -10,7 +10,7 @@ tags:Ansel
 version:4.14
 topic-image:Console_Topic.png
 
-The following console variables are available for configuring the behavior of the Nvidia Ansel Photography Plugin and can be set manually from the console, through your .ini files, or dynamically by a Blueprint. Most of them should be configured during development instead of at runtime, though `r.Photography.Allow` is useful to set dynamically.
+The following console variables are available for configuring the behavior of the NVIDIA Ansel Photography Plugin and can be set manually from the console, through your .ini files, or dynamically by a Blueprint. Most of them should be configured during development instead of at runtime, though `r.Photography.Allow` is useful to set dynamically.
 
 | Variable | Default | Description |
 | --- | --- | --- |
diff --git a/Engine/Documentation/Source/Engine/Plugins/Ansel/Testing/AnselTesting.INT.udn b/Engine/Documentation/Source/Engine/Plugins/Ansel/Testing/AnselTesting.INT.udn
--- a/Engine/Documentation/Source/Engine/Plugins/Ansel/Testing/AnselTesting.INT.udn
+++ b/Engine/Documentation/Source/Engine/Plugins/Ansel/Testing/AnselTesting.INT.udn
@@ -1,5 +1,5 @@
 Availability:Public
-TItle:Ansel Capture Testing
+TItle:NVIDIA Ansel Capture Testing
 Description:Overview of the Unreal Engine 4 integration of Nvidia's Ansel in-game photography tool that enables capturing and sharing enhanced images from games
 Crumbs:
 Parent:Engine/Plugins/Ansel
@@ -10,7 +10,7 @@ tags:Ansel
 version:4.14
 topic-image:Ansel_Testing_Topic.png
 
-**To test the Nvidia Ansel Photography Plugin with your game:**
+**To test the NVIDIA Ansel Photography Plugin with your game:**
 
 1. Run the NvCameraConfiguration.exe application and ensure that **NvCamera Status** is set to _Enable_, which is the default.
 
diff --git a/Engine/Documentation/Source/Engine/Rendering/LightingAndShadows/LightingAndShadows.INT.udn b/Engine/Documentation/Source/Engine/Rendering/LightingAndShadows/LightingAndShadows.INT.udn
--- a/Engine/Documentation/Source/Engine/Rendering/LightingAndShadows/LightingAndShadows.INT.udn
+++ b/Engine/Documentation/Source/Engine/Rendering/LightingAndShadows/LightingAndShadows.INT.udn
@@ -67,7 +67,7 @@ type:landing
 %Engine/Rendering/LightingAndShadows/IndirectLightingCache:Topic%
 %Engine/Rendering/LightingAndShadows/LightMobility/StaticLights:Topic%
 %Engine/Rendering/LightingAndShadows/LightMobility/StationaryLights:Topic%
-<!--%Engine/Rendering/LightingAndShadows/PrecomputedLightingScenarios:Topic%-->
+%Engine/Rendering/LightingAndShadows/PrecomputedLightingScenarios:Topic%
 [/REGION]
 
 
diff --git a/Engine/Documentation/Source/Engine/Rendering/LightingAndShadows/PrecomputedLightingScenarios/PrecomputedLightingScenarios.INT.udn b/Engine/Documentation/Source/Engine/Rendering/LightingAndShadows/PrecomputedLightingScenarios/PrecomputedLightingScenarios.INT.udn
--- a/Engine/Documentation/Source/Engine/Rendering/LightingAndShadows/PrecomputedLightingScenarios/PrecomputedLightingScenarios.INT.udn
+++ b/Engine/Documentation/Source/Engine/Rendering/LightingAndShadows/PrecomputedLightingScenarios/PrecomputedLightingScenarios.INT.udn
@@ -1,17 +1,20 @@
-Availability: Docs
+Availability: Public
 Title: Precomputed Lighting Scenarios
 Crumbs: %ROOT%, Engine, Engine/Rendering/LightingAndShadows
 Description: This page provides an overview of Precomputed Lighting Scenarios in UE4.
 Version: 4.14
 Parent: Engine/Rendering/LightingAndShadows
-Tags: Rendering
+Order: 1
 Type: Reference
+Tags: Rendering
+Tags: Lighting
+Tags: Lightmass
 topic-image:PCLS_Header.png
 
 [TOC(start:1 end:4)]
 
-Unreal Engine 4 (UE4) has support for using different **Precomputed Lighting Scenarios** in levels. This allows a single level to store and display multiple lighting setups, giving you the flexibility of dynamic lighting, but at the fixed cost of precomputed lighting. 
-Having the ability to change between different Precomputed Lighting Scenarios is especially important for Virtual Reality (VR) or Architectural Visualization projects that require high quality renderings with the benefits of fast performance measures. 
+Unreal Engine 4 (UE4) has support for using different **Precomputed Lighting Scenarios** in levels. This allows a single level to store and display multiple lighting setups, giving you the flexibility of dynamic lighting, but at the fixed cost of pre-computed lighting. 
+Having the ability to change between different Precomputed Lighting Scenarios is of particular importance for Virtual Reality (VR) or Architectural Visualization projects that require high quality renderings with the benefits of fast performance measures. 
 As you read through this reference, you'll learn how to use Precomputed lighting in your UE4 projects.
 
 ![](PCLS_Comp_Image_00.png)
@@ -20,12 +23,12 @@ As you read through this reference, you'll learn how to use Precomputed lighting
 In the above image; the Directional Light, Sky Light, and Skybox have been placed in a Lighting Scenario level called `Day Scenario`. The Spotlights for the streetlights have been placed in a different Lighting Scenario level called `Night Scenario`.
 [/REGION]
 
-## Precomputed Lighting Scenarios Limitations
+## Feature Limitations
 While Precomputed Lighting scenarios have a lot of advantages, there are some disadvantages and limitations that you'll need to be aware of while using them. 
 In the following section, we'll go over some of these limitations and how you can avoid (or work around) them.
 
 * Only one Lighting Scenario level should be visible at any time in a game.
-* When a Lighting Scenario level is present, Lightmap data from all Sublevels will be placed inside of it so that only the Day Scenario Lightmaps are loaded when it’s daytime. As a result, Lightmaps will no longer be streamed by Sublevels.
+* When a Lighting Scenario level is present, Lightmap data from all Sublevels will be placed inside of it so that only the Day Scenario Lightmaps are loaded when it's daytime. As a result, Lightmaps will no longer be streamed by Sublevels.
 * A Reflection Capture update is forced when making a Lighting Scenario level visible, which can increase the level load time.
 
 ##  Using Lighting Scenarios
@@ -36,7 +39,7 @@ To use Lighting Scenarios in your UE4 projects, you will need to do the followin
 
     ![](PCLS_Open_Levels.png)
 
-1. With the **Levels Manager** open, right-click on a Sublevel in the **Levels** menu, go to **Lighting Scenario**, and select the **Change to Lighting Scenario** option to make the level a **Lighting Scenario** level.
+1. With the **Levels Manager** open, right-click on a Sub level in the **Levels** menu, go to **Lighting Scenario**, and select the **Change to Lighting Scenario** option to make the level a **Lighting Scenario** level.
 
     ![](PCLS_LS_Level_Option.png)
 
@@ -52,7 +55,7 @@ To use Lighting Scenarios in your UE4 projects, you will need to do the followin
 
     ![](PCLS_Build_Lighting.png)
 
-1. Once the lighting has finished building, open up the **Persistent Level's** Blueprint and add a **Load Stream Level** node and connect it to the **Event Begin Play** node. 
+1. Once the lighting has finished building, open up the **Persistent Level's** Blueprint and then add a **Load Stream Level** node, connecting it to the **Event Begin Play** node. 
 
     ![](PCLS_Level_Load_Nodes.png)
 
diff --git a/Engine/Documentation/Source/Engine/Rendering/LightingAndShadows/ReflectionEnvironment/ReflectionEnvironment.INT.udn b/Engine/Documentation/Source/Engine/Rendering/LightingAndShadows/ReflectionEnvironment/ReflectionEnvironment.INT.udn
--- a/Engine/Documentation/Source/Engine/Rendering/LightingAndShadows/ReflectionEnvironment/ReflectionEnvironment.INT.udn
+++ b/Engine/Documentation/Source/Engine/Rendering/LightingAndShadows/ReflectionEnvironment/ReflectionEnvironment.INT.udn
@@ -179,6 +179,13 @@ By default, lightmap mixing will be enabled which means it will affect existing
 
 ![](Engine\Rendering\LightingAndShadows\ReflectionEnvironment\Enable_LM_Mixing.png)
 
+You can fine tune the amount of lightmap mixing that will take place by adjusting the following commands via the UE4 console. 
+ 
+* **r.ReflectionEnvironmentBeginMixingRoughness (default = 0.1)**
+* **r.ReflectionEnvironmentEndMixingRoughness (default = 0.3)**
+* ** r.ReflectionEnvironmentLightmapMixBasedOnRoughness (default = 1)**
+* **r.ReflectionEnvironmentLightmapMixLargestWeight (default = 1000)**
+
 ## High Quality Reflections
 
 While the default reflection quality settings strike a good balance between performance and visual quality, there could be instances where you want to achieve even higher quality reflections. The following sections describe the available methods for achieving high quality reflections.
diff --git a/Engine/Documentation/Source/Engine/Rendering/LightingAndShadows/VolumetricFog/VolumetricFog.INT.udn b/Engine/Documentation/Source/Engine/Rendering/LightingAndShadows/VolumetricFog/VolumetricFog.INT.udn
--- a/Engine/Documentation/Source/Engine/Rendering/LightingAndShadows/VolumetricFog/VolumetricFog.INT.udn
+++ b/Engine/Documentation/Source/Engine/Rendering/LightingAndShadows/VolumetricFog/VolumetricFog.INT.udn
@@ -20,14 +20,22 @@ support varying densities and any number of lights affecting the fog.
 [REGION:fullwidth raw]
 ![](VolumetricFog.png)(w:1339)
 [/REGION]
+[REGION:caption] In this scene, Volumetric Fog is coming from the directional light source off screen through the arch and surrounding area to create shadowed fog.  [/REGION]
 
-## Global Controls
+
+## Volumetric Fog Controls
+
+When setting up and adjusting Volumetric Fog, you can control it globally or locally in your scene. The global controls enable you to use the Exponential Height Fog component to control fog for the the entire scene. 
+The local controls enable you to control fog by way of a particle system in areas where the particles can spawn.
+
+### Global Controls
 
 To control Volumetric Fog, you can adjust the properties in your **Exponential Height Fog** and on each **Light** to control the Light's contribution amount. 
 
-### Exponential Height Fog
 
-Volumetric Fog controls can be found in the Exponential Height Fog component under the **Volumetric Fog** section. The exponential height distribution provides a global density for Volumetric Fog.
+#### Exponential Height Fog
+
+Volumetric Fog controls can be found in the **Exponential Height Fog** component under the **Volumetric Fog** section. The exponential height distribution provides a global density for Volumetric Fog.
 
 ![](VolumetricFogProperties.png)
 
@@ -36,33 +44,35 @@ Volumetric Fog controls can be found in the Exponential Height Fog component und
 | **Scattering Distribution** | This determines how directional the volumetric scattering is; a value of 0, means light scatters equally in all directions, while a value close to 1 causes scattering, predominantly in the direction of the light (you have to be looking at the light to see its scattering). |
 | **Albedo** | This is the overall reflectiveness of the participating media. Clouds, fog, and mist, which are based on water particles, have an Albedo close to 1. |
 | **Extinction Scale** | Controls how much the participating media blocks light. |
-| **View Distance** | The distance from the camera over which Volumetric Fog will be computed. There are a limited number of Z slices in the volume texture, so pushing out the View Distance will increase under-sampling artifacts near the camera. |
+| **View Distance** | This is the distance from the camera over which Volumetric Fog will be computed. In the created volume texture for the fog there are a limited number of Z slices depending on this distance. Increasing the distance will result in under-sampling that can cause artifacts to appear. The number of Z slices can be adjusted by using r.VolumetricFog.GridSizeZ, where higher is better quality, but will be more expensive. |
 | **Override Light Colors with Fog Inscattering Colors** | When enabled, uses the **Fog Inscattering Color**, **Directional Inscattering Color**, and **Inscattering Texture** properties to override the light color with Volumetric Fog. |
 
-### Lights
 
-Each Light's contribution amount to the scene (and whether it shadows the fog) can be controlled with the following properties.
+#### Lights
+
+Each Light's contribution amount to the scene (and whether it shadows the fog) can be controlled by adjusting the following properties on each Light's Details panel under the **Light** section.
 
 ![](VolumetricFogLightProperties.png)
 
 | Property | Description |
 | --- | --- |
 | **Volumetric Scattering Intensity** | Controls how much this light will contribute to the Volumetric Fog. When set to 0, there is no contribution. |
-| **Cast Volumetric Shadow** | Toggles whether or not to cast a volumetric shadow for lights contributing to Volumetric Fog. When enabled, Point and Spot Lights are approximately three times more expensive than non-shadow casting lights. |
-
+| **Cast Volumetric Shadow** | Toggles whether or not to cast a volumetric shadow for lights contributing to Volumetric Fog. When shadow casting is enabled, Point and Spot Lights are approximately three times more expensive because they contribute to the volume textures shadowing, where as non-shadow casting lights only contribute to the fog but do not shadow. |
 
 [OBJECT:ComparisonSlider]
 	[PARAM:before]
-	![Volumetric Scattering Intensity: 1 (Default)](VolumetricFogEnabled.png)(w:800)
+	![Volumetric Scattering Intensity: 1 (Default)](VolumetricFogEnabled.png)(w:950)
 	[/PARAM]
 	[PARAM:after]
-	![Volumetric Scattering Intensity: 0 ](VolumetricFogDisabled.png)(w:800)
+	![Volumetric Scattering Intensity: 0 ](VolumetricFogDisabled.png)(w:950)
 	[/PARAM]
 [/OBJECT]
 
 In this example, the Spot Light's contribution to the Volumetric Fog has been disabled by setting the **Volumetric Scattering Intensity** to 0.
 
-## Local Controls
+
+
+### Local Controls
 
 Materials using the **Volume** domain describe Albedo, Emissive, and Extinction for a given point in space. Albedo is in the range [0-1], while Emissive and Extinction are world space densities with
 any value greater than 0. 
@@ -70,29 +80,33 @@ any value greater than 0.
 
 ![](VolumetricFogParticleMaterial.png)(w:700)
 [REGION:caption]
-This is an example of the simplest Volume material for a Particle System.
+This is an example of the simplest Volume Material for a Particle System.
 [/REGION]
 
 [REGION:note]
-Volume materials currently only work on particles and only positions inside of the particle radius are valid, which is usually handled by a SphereMask. 
+Volume materials currently only work on particles and positions inside of the particle radius will be valid, which is usually handled by a SphereMask.
 [/REGION]
 
+
 Placing a single Particle System with the material causes a sphere of density to be added to the Volumetric Fog. The effect is fully three dimensional (3D) with no billboards involved. 
 
-![](VF_ParticleInGame.png)(w:700)
+![](VF_ParticleInGame.png)(w:900)
+
+Taking this approach a step further, you could use multiple spherical fog particles with noise from textures to limit fog to a certain area of your scene. 
 
-Taking this approach a step further, you could use multiple spherical fog particle with noise from textures to limit fog to a certain area of your scene. 
+![](FogParticles.png)(w:900)
+
+In this scene, the fog particles have populated these low-lying areas to create a localized fog effect that is shadowed using Voluemtric Fog. 
 
 ## Temporal Reprojection
 
 The volume textures (voxels) used by Volumetric Fog are relatively low-resolution and aligned to the camera frustum. Volumetric Fog uses a heavy temporal reprojection filter with a different sub-voxel jitter
-per frame to smooth out the aliasing. As a side-effect, fast-changing lights, like flashlights and muzzle flashes, leave lighting trails. Disable the contribution of these lights by setting
-**Volumetric Scattering Intensity** set to 0.
+per frame to smooth out the aliasing. As a side-effect, fast-changing lights, like flashlights and muzzle flashes, leave lighting trails. To disable the contribution of these lights, set **Volumetric Scattering Intensity** to 0.
 
 ## Performance
 
-The GPU cost of Volumetric Fog is primarily controlled by the volume texture resolution, which is set from the Engine Shadow Scalability level. Volumetric Fog costs 1 millisecond on PlayStation 4 at High Settings,
-and 3 milliseconds on an Nvidia 970 GTX on Epic settings, which has eight times more voxels to operate on.
+The GPU cost of Volumetric Fog is primarily controlled by the volume texture resolution, which is set from the Shadow level of the [Engine Scalability](Engine/Performance/Scalability/ScalabilityReference). Volumetric Fog costs 1 millisecond on PlayStation 4 at **High** settings,
+and 3 milliseconds on an NVIDIA 970 GTX on **Epic** settings, which has eight times more voxels to operate on.
 
 * Particles using **Volume** domain can add a significant GPU cost, depending on their 3D overdraw and instruction count. Use the console command `profilegpu` to inspect this cost.
 * Point and Spot Lights that have **Cast Volumetric Shadow** enabled, cost approximately three times more than unshadowed Point and Spot Lights. 
@@ -130,6 +144,27 @@ Below are some common questions or issues that may arise when using Volumetric F
 		1. Keep the global fog density low, but use a higher **Volumetric Scattering Intensity** for the Directional Light. Also, adjust the **Scattering Distritbution** to nearly **0.9** in your Exponential Height Fog Actor.
 		1. Keep the global fog density low, but increase it in certain areas with Volume particles.
 * **Can Exponential Height Fog and Volumetric Fog be used at the same time?** 
-	* At this time, Volumetric Fog replaces **Fog Inscattering Color** within the Volumetric Fog **View Distance**. Because Volumetric Fog is physically-based and Exponential Height Fog is not, it's impossible to precisely match these in the distance. This means that some settings in the Exponential Height Fog will have no effect on Volumetric Fog.
+	* At this time, Volumetric Fog replaces **Fog Inscattering Color** within the Volumetric Fog **View Distance**. Because Volumetric Fog is physically-based and Exponetial Fog is not, it's impossible to blend the two systems in the distance for them to 
+	precisely match. This also means that some settings in the Exponential Height Fog component will have no effect on Volumetric Fog. 
 * **Can the Volumetric Fog's center be decoupled from the camera? This would be ideal for top-down games...**
 	* Not currently, though, a standalone volume would be ideal for this. However, it's hard to integrate them with translucency efficiently.
+
+## Training Stream
+
+[OBJECT:EmbeddedVideo]
+	[PARAMLITERAL:videoid]
+	N4mkgbwLg7U
+	[/PARAMLITERAL]
+	[PARAMLITERAL:width]
+	50
+	[/PARAMLITERAL]
+	[PARAMLITERAL:height]
+	56.25
+	[/PARAMLITERAL]
+	[PARAMLITERAL:units]
+	%
+	[/PARAMLITERAL]
+	[PARAMLITERAL:type]
+	simple
+	[/PARAMLITERAL]
+[/OBJECT]
diff --git a/Engine/Documentation/Source/Engine/Rendering/Materials/ExpressionReference/Utility/UtilityExpressions.INT.udn b/Engine/Documentation/Source/Engine/Rendering/Materials/ExpressionReference/Utility/UtilityExpressions.INT.udn
--- a/Engine/Documentation/Source/Engine/Rendering/Materials/ExpressionReference/Utility/UtilityExpressions.INT.udn
+++ b/Engine/Documentation/Source/Engine/Rendering/Materials/ExpressionReference/Utility/UtilityExpressions.INT.udn
@@ -15,10 +15,7 @@ Topic-image:Utility_topic.png
 [/EXCERPT:TOC]
 
 
-In this page, you'll find the available Utility Expressions that enable you to input values to affect Material's output value, like for example, using the GIReplace node will Replace an objects
-indirect bounce color with a given value you input, or LinearInterpolate will blend between two inputs using an alpha value to control how much.
-
-There are many others that enable you to control your Materials that you can look at the examples below to learn more about. 
+**Utility Expressions** nodes are nodes that can affect Materials in a number of different ways than one might be used to. For example the GIReplace node will Replace an objects indirect bounce color with a given value you input while the Linear Interpolate node will help blend between two Textures based on an Alpha input. On the following page you will find detailed descriptions for all of the Utility Expressions nodes that are available in the UE4 Material Editor.
 
 
 ## A - H
@@ -27,7 +24,7 @@ There are many others that enable you to control your Materials that you can loo
 [EXCERPT:AntialiasedTextureMask]
 ### AntialiasedTextureMask
 
-The **AntialiasedTextureMask** expression allows you to create a material using a soft (anti-aliased) transition mask. The mask can be used to blend between two complex material properties or to fade out an alpha blended material (works well with SoftMasked). Simply specify a texture with the mask specified in one channel (red, green, blue, or alpha), set the used channel in the expression and specify the comparison value. Assuming the channel stores a grayscale value in the range 0=black to 1=white the comparison function defines if the resulting mask should be 0 or 1. This expression is a parameter, allowing the **Texture** property to be overridden by child MaterialInstances.
+The **AntialiasedTextureMask** expression allows you to create a material using a soft (anti-aliased) transition mask. The mask can be used to blend between two complex material properties or to fade out an alpha blended material (works well with SoftMasked). Simply specify a texture with the mask specified in one channel (red, green, blue, or alpha), set the used channel in the expression and specify the comparison value. Assuming the channel stores a grayscale value in the range 0 = black to 1 = white the comparison function defines if the resulting mask should be 0 or 1. This expression is a parameter, allowing the **Texture** property to be overridden by child MaterialInstances.
 
 | Item | Description |
 | -------- | ----------- |
@@ -203,7 +200,7 @@ The **Desaturation** expression desaturates its input, or converts the colors of
 ![DesaturationExample.png](DesaturationExample.png)
 
 [REGION:note]
-**Programmers:** Define desaturated color `D`, input color `I` and luminance factor `L`.  The output will be `O = (1-Percent)*(D.dot(I)) + Percent*I`
+**Programmers:** Define desaturated color `D`, input color `I` and luminance factor `L`.  The output will be `O = (1 - Percent)*( D.dot( I )) + Percent * I`
 [/REGION]
 [/EXCERPT:Desaturation]
 
@@ -228,11 +225,11 @@ The **Distance** expression computes the (Euclidian) distance between two points
 
 **Pseudo code:**
 
-	Result = length(A - B)
+	Result = length (A - B)
 
 **Low level HLSL code:**
 	
-	float Result = sqrt(dot(A-B, A-B))
+	float Result = sqrt (dot (A-B, A-B))
 	
 [/EXCERPT:Distance]
 
@@ -443,11 +440,11 @@ The **Noise** expression creates a procedural noise field, giving you control ov
 | **Scale** | Changes the overall size of the noise cells. Lower numbers make the noise bigger. |
 | **Quality** | A look/performance setting. Lower values are faster but may look worse, higher values are slower but may look better. |
 | **Function** | [INCLUDE:#NoiseFunctions] |
-| **Turbulence** | Whether to use an absolute value when combining noise levels, results in a veined or billowy look. |
+| **Turbulence** | With Turbulence on, each noise octave will add only absolute values to the result. Changes the visual characteristics and can great shapes resembling sharp mountain ridges |
 | **Levels** | How many different levels of noise at different scales to combine, multiplies the computational cost by the number of levels. |
 | **Output Min** | The lowest value output by the noise calculation. |
 | **Output Max** | The highest value output by the noise calculation. |
-| **Level Scale** | Controls the scale of individual levels when using multiple levels of noise. This defaults to 2 so each level is twice the frequency at half the magnitude. |
+| **Level Scale** |  Level scale is always active and determines how much the scale changes for each new octave. |
 | **Tiling** | For noise functions that support it, allows noise to tile. This is more expensive, but useful when baking noise into a seamless wrapping texture. |
 | **Repeat Size** | When tiling, how often should the noise repeat. |
 || [REGION:tablesection]Inputs[/REGION] ||
@@ -466,7 +463,7 @@ The **Noise** expression creates a procedural noise field, giving you control ov
 |Fast Gradient | 3D Texture | High quality for direct use, **BAD** for bumps. about 16 instructions per level, 1 Texture lookup. Always tiles with a repeat of 16, "Tiling" mode is not an option for Fast Gradient noise.|
 | Gradient | Computational | High quality for direct use and bumps. Non-tiled: about 80 instructions per level, no Textures. Tiling: about 143 instructions per level, no Textures. |
 | Value | Computational | Low quality, but pure computation. Non-tiled: about 53 instructions per level, no Textures. Tiling: about 118 instructions per level, no Textures. Formerly mislabeled as Gradient noise. |
-| Voronoi | | Also know as Worley or Cellular noise. Quality=1 searches 8 cells, Quality=2 searches 16 cells, Quality=3 searches 27 cells, Quality=4 searches 32 cells. All are about about 20 instructions per cell searched. |
+| Voronoi | | Also know as Worley or Cellular noise. Quality = 1 searches 8 cells, Quality = 2 searches 16 cells, Quality = 3 searches 27 cells, Quality = 4 searches 32 cells. All are about about 20 instructions per cell searched. |
 [/EXCERPT]
 -->
 
@@ -500,7 +497,7 @@ In this example, the Previous Frame Switch is using a constant value to control
 	YbjqiNZr4K4
 	[/PARAMLITERAL]
 	[PARAMLITERAL:width]
-	50
+	100
 	[/PARAMLITERAL]
 	[PARAMLITERAL:height]
 	56.25
diff --git a/Engine/Documentation/Source/Engine/Rendering/Materials/Functions/Reference/Misc/Misc.INT.udn b/Engine/Documentation/Source/Engine/Rendering/Materials/Functions/Reference/Misc/Misc.INT.udn
--- a/Engine/Documentation/Source/Engine/Rendering/Materials/Functions/Reference/Misc/Misc.INT.udn
+++ b/Engine/Documentation/Source/Engine/Rendering/Materials/Functions/Reference/Misc/Misc.INT.udn
@@ -11,7 +11,7 @@ tags:Materials
 
 Miscellaneous functions include a variety of one-off functions that do not readily fit into any of the other existing categories. They are placed here to avoid too many categories with only one function inside.
 
-## Math Functions
+## Misc Functions
 
 The following is a list of all the functions underneath the Misc category.
 
diff --git a/Engine/Documentation/Source/Engine/Rendering/Nvidia/NVIDIAAftermath/NVIDIA_Aftermath.INT.udn b/Engine/Documentation/Source/Engine/Rendering/Nvidia/NVIDIAAftermath/NVIDIA_Aftermath.INT.udn
new file mode 100644
--- /dev/null
+++ b/Engine/Documentation/Source/Engine/Rendering/Nvidia/NVIDIAAftermath/NVIDIA_Aftermath.INT.udn
@@ -0,0 +1,72 @@
+Availability: Public 
+Crumbs: %ROOT%
+Title: NVIDIA Aftermath
+Description: Integration of NVIDIA's Aftermath technology that enables post-mortem GPU crash analysis on NVIDIA GeForce based GPUs. 
+Type: Reference
+Version:4.16
+Parent:Engine/Rendering/NVIDIA
+Order:2
+Tags:NVIDIA
+topic-image:Aftermath_Topic.png
+
+[REGION:fullwidth raw]
+![](Aftermath_Banner.png)(w:900)
+[/REGION]
+
+[REGION:note]
+It is recommended to upgrade to NVIDIA Driver 382.53 or later as it has some improvements that make it a lot faster to use. 
+[/REGION]
+
+As of version 4.15.2, Unreal Engine supports NVIDIA's **Aftermath** C++ library for Windows-based developers, which provides additional data from NVIDIA GeForce-based GPUs after a crash occurs. 
+This data contains crucial information about what the GPU was doing when the crash occurred, enabling you to track down what's happening with the GPU in your own projects.
+
+Aftermath is a lightweight tool that is unobtrusive and reduces the performance footprint needed by some debugging tools. In fact, it's lightweight enough that it can even ship with finished 
+games to provide the data you need from your customers' machines. Aftermath enables programmers to insert markers into their code that assist in tracking the root causes of crashes. This is 
+now being used in the Unreal Engine Editor to track down and fix reported or captured issues.
+
+For additional information and how to use it in your own project, visit [NVIDIA's Aftermath documentation](https://developer.nvidia.com/nvidia-aftermath) page.
+
+## Enabling NVIDIA Aftermath
+
+If you're using Unreal Engine version 4.15.2, NVIDIA Aftermath is automatically enabled. You can use the following console command to set it in your **ConsoleVariables.ini** configuration file:
+
+	r.DX11NVAfterMathEnabled
+
+If you're using Unreal Engine version 4.16 or later, NVIDIA Aftermath is automatically disabled. You can use the following console command to set it in your **ConsoleVariables.ini** configuration file:
+
+	r.GPUCrashDebugging=1
+
+You can also enable it by passing the following argument to the command line:
+
+	-gpucrashdebugging
+
+
+## Logs
+
+Once you've enabled NVIDIA Aftermath, you should see the following output in your log:
+
+	LogD3D11RHI: [Aftermath] Aftermath enabled and primed
+
+You can access your log in Unreal Engine 4 by going to **File Menu** > **Window** > **Developer Tools** > **Output Log** or by opening the log text file from your project folder.
+
+When a crash happens, the log will display something similar to the following:
+
+	LogD3D11RHI: Error: Result failed at X:\[Project Folder]\Engine\Source\Runtime\Windows\D3D11RHI\Private\D3D11Viewport.cpp:290 with error DXGI_ERROR_DEVICE_REMOVED DXGI_ERROR_DEVICE_HUNG
+	LogRHI: Error: [Aftermath] Status: Timeout
+	LogRHI: Error: [Aftermath] GPU Stack Dump
+	LogRHI: Error: [Aftermath] 0: Frame2769
+	LogRHI: Error: [Aftermath] 1: FRAME
+	LogRHI: Error: [Aftermath] 2: Scene
+	LogRHI: Error: [Aftermath] 3: ComputeLightGrid
+	LogRHI: Error: [Aftermath] 4: Compact
+	LogRHI: Error: [Aftermath] GPU Stack Dump
+
+In this particular example, the GPU has crashed and the resulting lines with the `[Aftermath]` preceeding indicate the status and where in the frame the problem is so that you can investigate what 
+was the root cause.   
+
+## Considerations
+
+For an average project that has approximately 200 to 300 markers, Aftermath is fast enough, but slowdowns can happen when you have certain things like lots of per-object shadows that can 
+increase its cost. For this reason, the feature is not enabled by default.  
+
+
diff --git a/Engine/Documentation/Source/Engine/Rendering/Nvidia/NVIDIAAlternateFrameRendering/NVIDIA_AlternateFrameRendering.INT.udn b/Engine/Documentation/Source/Engine/Rendering/Nvidia/NVIDIAAlternateFrameRendering/NVIDIA_AlternateFrameRendering.INT.udn
new file mode 100644
--- /dev/null
+++ b/Engine/Documentation/Source/Engine/Rendering/Nvidia/NVIDIAAlternateFrameRendering/NVIDIA_AlternateFrameRendering.INT.udn
@@ -0,0 +1,40 @@
+Availability: Public
+Crumbs: %ROOT%
+Title: NVIDIA SLI ALternate Frame Rendering
+Description: Integration of NVIDIA's Alternate Frame Rendering technology that enables support for games using SLI. 
+Type: Reference
+Version:4.16
+Parent:Engine/Rendering/Nvidia
+Order:1
+Tags:Nvidia
+topic-image:AFR_topic.png
+
+
+[REGION:fullwidth raw]
+![](AFR_Banner.png)
+[/REGION]
+
+
+As of version 4.15, Unreal Engine supports NVIDIA's **Alternate Frame Rendering** (AFR) for packaged games running on NVIDIA SLI configurations. AFR works by alternating frames between the linked 
+GPUs; one frame is rendered by GPU 1, then the next frame is rendered by GPU 2, then back to GPU 1, and so on. By using multiple GPUs on a single monitor, its output can increase in image quality
+and performance! 
+
+Developers who intend to use AFR will still need to work directly with NVIDIA to test their games and to have it automatically switch over to use the Alternating Frame Rendering functionality
+where appropriate.
+
+For additional details, you can read NVIDIA's documentation on their [SLI modes, specifically Alternate Frame Rendering](https://docs.nvidia.com/gameworks/content/technologies/desktop/sli.htm) here.
+
+## Force Enabling AFR
+
+For some applications, the NVIDIA Control Panel enables you to add a program for which the SLI Rendering mode will have AFR forcibly enabled. To do this, follow these steps:
+
+1. Open the **NVIDIA Control Panel** from your task tray, then navigate to **Manage 3D Settings**. 
+1. Click on the **Program Settings** tab, then under **Select a program to cumstomize** use the drop-down to select a program to add. 
+1. Next to the option for  **SLI Rendering Mode**, you can select **Force Alternate Frame Rendering**.
+
+![](AFRSetting.jpg)
+
+[REGION:note]
+This feature is not guaranteed to improve the quality or performance of your application and is not intended for use with the Unreal Engine 4 Editor. It is recommended to talk directly with 
+NVIDIA to set up this functionality with their own driver.
+[/REGION]
diff --git a/Engine/Documentation/Source/Engine/Rendering/Nvidia/Nvidia.INT.udn b/Engine/Documentation/Source/Engine/Rendering/Nvidia/Nvidia.INT.udn
--- a/Engine/Documentation/Source/Engine/Rendering/Nvidia/Nvidia.INT.udn
+++ b/Engine/Documentation/Source/Engine/Rendering/Nvidia/Nvidia.INT.udn
@@ -1,14 +1,17 @@
-Availability:Docs
-Title:NVIDIA Feature Support
+Availability: Public
+Crumbs: %ROOT%
+Title: NVIDIA Feature Support
 Description: Integrated and supported technologies from NVIDIA for use with your games.
-Crumbs:
+Type: Landing
+Version:4.16
 Parent:Engine/Rendering
 Order:101
-type:landing
-version:4.16
+Tags:Nvidia
 topic-image:NVIDIA_logo_Topic.png
 
 
+
+
 [REGION:raw]
 ![](nvidia_logo_horizontal.png)(w:700)
 [/REGION]
@@ -17,7 +20,9 @@ topic-image:NVIDIA_logo_Topic.png
 Here you will find the integrated technologies and relevant information from NVIDIA for use with your games. 
 
 
-%Engine\Nvidia\NVAlternateFrameRendering:Topic%
-%Engine\Nvidia\NVAftermath:Topic%
+[DIR(output:"topiccompact" parent:"Engine/Rendering/Nvidia" org:"hierarchy" tags:"Nvidia" end:"1")]
+
+
+
 
 
diff --git a/Engine/Documentation/Source/Engine/Rendering/ParticleSystems/Reference/Modules/Event/EventModules.INT.udn b/Engine/Documentation/Source/Engine/Rendering/ParticleSystems/Reference/Modules/Event/EventModules.INT.udn
--- a/Engine/Documentation/Source/Engine/Rendering/ParticleSystems/Reference/Modules/Event/EventModules.INT.udn
+++ b/Engine/Documentation/Source/Engine/Rendering/ParticleSystems/Reference/Modules/Event/EventModules.INT.udn
@@ -18,7 +18,7 @@ The **Event Generator** module will generate an event (or multiple events) based
 | ---- | ---- |
 |[REGION:tablesection]Events[/REGION]||
 |**Type** | The type of event. The possible types include:[INCLUDE:#EventGenType]|
-|**Frequency** | How often to trigger the event. <=1 means every time. Use this to trigger an event on every other collision, say.|
+|**Frequency** | How often to trigger the event. `<=1` means every time. Use this to trigger an event on every other collision, say.|
 |**First Time Only** | Boolean to trigger the event only on the first time it would fire.|
 |**Last Time Only** | Boolean to trigger the event only on the last occurrence of when it would fire.|
 |**Use Reflected Impact Vector** | Boolean to decide whether you orient the result of the collision event in the direction of the impact vector instead of the hit normal.|
diff --git a/Engine/Documentation/Source/Engine/Rendering/ParticleSystems/Reference/Modules/Spawn/SpawnModules.INT.udn b/Engine/Documentation/Source/Engine/Rendering/ParticleSystems/Reference/Modules/Spawn/SpawnModules.INT.udn
--- a/Engine/Documentation/Source/Engine/Rendering/ParticleSystems/Reference/Modules/Spawn/SpawnModules.INT.udn
+++ b/Engine/Documentation/Source/Engine/Rendering/ParticleSystems/Reference/Modules/Spawn/SpawnModules.INT.udn
@@ -19,7 +19,7 @@ Allows for spawning particles based on the distance the emitter has traveled. Us
 |**Unit Scalar** | A float value giving the scalar to apply to the distance traveled. The value from SpawnPerUnit is divided by this value to give the actual number of particles per unit.|
 |**Spawn Per Unit** | A float distribution giving the amount of particles to spawn per unit. The value is retrieved using the EmitterTime.|
 |**Ignore Spawn Rate When Moving** | If _true_, process the default spawn rate when not moving. When the emitter is moving, skip processing the default spawn rate. If _false_, return the _Process Spawn Rate_ setting.|
-|**Movement Tolerance** | A float value giving the tolerance for moving vs. not moving w.r.t. the _Ignore Spawn Rate When Moving_ flag, i.e. if (DistanceMoved < (UnitScalar * MovementTolerance)), then consider it not moving.|
+|**Movement Tolerance** | A float value giving the tolerance for moving vs. not moving w.r.t. the _Ignore Spawn Rate When Moving_ flag, i.e. if (`DistanceMoved < (UnitScalar * MovementTolerance)`), then consider it not moving.|
 |**Max Frame Distance** | The maximum movement for a single frame. If the particle moves beyond this distance over a single frame, no particles are emitted. If left at 0, this check will not be performed. This is useful for handling teleporting objects.|
 |**Ignore Movement Along X** | If _true_, the X-component of the movement is ignored.|
 |**Ignore Movement Along Y** | If _true_, the Y-component of the movement is ignored.|
diff --git a/Engine/Documentation/Source/Engine/Rendering/PostProcessEffects/AutomaticExposure/AutomaticExposure.INT.udn b/Engine/Documentation/Source/Engine/Rendering/PostProcessEffects/AutomaticExposure/AutomaticExposure.INT.udn
--- a/Engine/Documentation/Source/Engine/Rendering/PostProcessEffects/AutomaticExposure/AutomaticExposure.INT.udn
+++ b/Engine/Documentation/Source/Engine/Rendering/PostProcessEffects/AutomaticExposure/AutomaticExposure.INT.udn
@@ -23,8 +23,8 @@ Topic-image:EyeAdap_Topic.png
 | -------- | ----------- |
 | **Low Percent** | The eye adaptation will adapt to a value extracted from the luminance histogram of the scene color. The value defines the lower percentage of the histogram that is used to find the average scene luminance. We want to avoid clamping of the bright image parts so it is best to ignore most of the dark areas. e.g. 80 avoid 80% of the darks. Values should be in the range [0, 100]. Values in the range 70-80 give the best results. |
 | **High Percent** | The eye adaptation will adapt to a value extracted from the luminance histogram of the scene color. The value defines the upper percentage of the histogram that is used to find the average scene luminance. We want to clip off a few percent as it is ok to have some bright pixels (usually sun). Values should be in the range [0, 100]. Values in the range 80-98 give the best results. |
-| **Min Brightness** | This value limits the lower brightness that the eye adapts within. Value must be greater than 0 and should be <= **EyeAdaptationMaxBrightness**. Actual values depend on the HDR range the content is using. |
-| **Max Brightness** | This value limits the upper brightness that the eye adapts within. Value must be greater than 0 and should be >= **EyeAdaptationMinBrightness**. Actual values depend on the HDR range the content is using. |
+| **Min Brightness** | This value limits the lower brightness that the eye adapts within. Value must be greater than 0 and should be less than or equal to **EyeAdaptationMaxBrightness**. Actual values depend on the HDR range the content is using. |
+| **Max Brightness** | This value limits the upper brightness that the eye adapts within. Value must be greater than 0 and should be greater than or equal to **EyeAdaptationMinBrightness**. Actual values depend on the HDR range the content is using. |
 | **Speed Up** | The speed at which the adaptation occurs from a dark environment to a bright environment. |
 | **Speed Down** | The speed at which the adaptation occurs from a bright environment to a dark environment. |
 | **Auto Exposure Bias** | Controls biasing of the exposure settings. The value is logarithmic meaning: [INCLUDE:#ExposureOffset] Fractional values like 1.5 or values outside of the shown examples work as well. If eye adaptation is active, the value is added to the automatic exposure. |
