Copying //UE4/Dev-Documentation to //UE4/Samples-Main (Source: //UE4/Dev-Documentation @ 3623665)

diff --git a/Engine/Documentation/Source/Engine/Rendering/PostProcessEffects/ColorGrading/ColorGrading.INT.udn b/Engine/Documentation/Source/Engine/Rendering/PostProcessEffects/ColorGrading/ColorGrading.INT.udn
--- a/Engine/Documentation/Source/Engine/Rendering/PostProcessEffects/ColorGrading/ColorGrading.INT.udn
+++ b/Engine/Documentation/Source/Engine/Rendering/PostProcessEffects/ColorGrading/ColorGrading.INT.udn
@@ -1,135 +1,437 @@
 Availability: Public
 Crumbs: %ROOT%
-Title: Color Grading
+Title: Color Grading and Filmic Tonemapper
 Description: Tonemapping and color correction effects for adjusting scene colors.
 Type: Overview
 SkillLevel:Beginner
 Version:4.16
 Parent:Engine/Rendering/PostProcessEffects
+related:Engine/Rendering/PostProcessEffects/UsingLUTs
 Order:4
 Tags:Post Process
-Topic-image:ColorGrading_Topic.png
+Topic-image:ColorGrading_Overview.png
+Social-image:ColorGrading_Overview.png
 
+[VAR:spacer]![](%ROOT%/spacer.png)(w:150 convert:false)[/VAR]
+[VAR:spacer2]![](%ROOT%/spacer.png)(w:40 convert:false)[/VAR]
+[VAR:spacer3]![](%ROOT%/spacer.png)(w:0 convert:false)[/VAR]
 
 
-Within the Unreal engine, the term **Color Grading** covers the Tone Mapping function (HDR to LDR transformation) and the further Color Correction (LDR color to screen color transformation).
+[TOC(start:2 end:2)]
+
+![](ColorGrading_Overview.png)
+[REGION:caption] The Color Grading tools enable you to change the look and tone of your scene, like with this example that is showing the scene using different Scene Colors. [/REGION]
+
+Within the Unreal Engine, the term **Color Grading** covers the Tone Mapping function (HDR to LDR transformation) that is used with [High Dynamic Range (HDR) display output](Engine/Rendering/HDRDisplayOutput)
+and further the color correction (LDR color to screen color transformation) processing of the image.
 
 ## Tone Mapping
 
 [EXCERPT:ToneMapping_Desc]
-The purpose of the **Tone Mapping** function is to map the wide range of HDR (high dynamic range) colors into the small LDR (low dynamic range) so a monitor can display the color. This is done after normal rendering during post processing. A global tone mapper is a function with three inputs (RGB) and three outputs (RGB). A local tone mapper also takes into account the neighborhood of the pixel but is much more computationally intensive (means slower). A good tone mapper function tries to preserve the color of a pixel even if the color is very bright.
+The purpose of the **Tone Mapping** function is to map the wide range of high dynamic range (HDR) colors into low dynamic range (LDR) that a display can output. This is the last stage of post processing
+that is done after the normal rendering during post processing. The process of tone mapping can be thought of as a way to simulate the response that film has to light. 
 [/EXCERPT:ToneMapping_Desc]
 
-The following graph shows how each channel in the range from 0 to 5 is mapped to the new range of 0 to 1:
-
-![Tone Mapping Graphs](newtone.jpg)
+### Academy Color Encoding System (ACES) Filmic Tone Mapper
 
-This is the used function:
+[REGION:note]
+As of Unreal Engine version 4.15, the filmic tone mapper was enabled as the default tone mapper. Any content that was previously developed may look different due to this change. To use the old tone mapper 
+you can enter the following console command:
+
+	r.TonemapperFilm 0
+	
+The Tone Mapper property values that closest match the "old" tone mapper can be used by entering:
+* **Slope:** 0.98
+* **Toe:** 0.3
+* **Shoulder:** 0.22
+* **Black Clip:** 0
+* **White Clip:** 0.025
+[/REGION]
 
-	GammaColor = LinearColor / (LinearColor + 0.187) * 1.035;
+The Filmic Tone Mapper that is used in UE4 matches the industry standard set by the [Academy Color Encoding System (ACES)](http://www.oscars.org/science-technology/sci-tech-projects/aces) for television 
+and film. This ensures that consistent color is preserved across multiple formats and displays while also as a way to future proof the source material since it will not have to be adjusted for each medium 
+that comes along. The filmic tone mapper also uses the same global tone mapper that was previously used in UE4, with the exception that there is now a filmic response meaning that the S-curve shape now better
+simulates film stock for a better overall look. 
 
-Note that bright colors become gradually brighter but much less than dark ones. Black remains black and there is a mostly linear part of the curve that is a bit steeper than the non tone-mapped curve would be. That results in some contrast enhancing. While it is normal to get a different look when using a tone mapper, in order to get good results the source image needs to have a higher dynamic range. This can pay off in a more realistic film like look. 
+Two places you'll immediately notice a difference is with Bloom and Exposure Levels.
 
-The presented tone mapper formula already has two editable constants and while more math would allow further tweaking, finding a good function is not easy as quality, flexibility, and performance need to be considered. We decided to stick to a simple formula and modify the resulting LDR color with a simple color lookup. As we already mapped the HDR color into a limited range, we get several expressions of dark colors and still can modify the bright colors. This method is simple to understand and offers a lot of flexibility with local control and consistent performance. 
+#### Physically Correct Emissive and Bloom
 
-| Property | Description |
-| -------- | ----------- |
-| **Scene Color Tint** | A color that is a filter color (multiply) that is applied to the HDR scene color. [INCLUDE:#SceneColorTint] |
-| **Tonemapper Toe Factor** | Allows the dark colors in the tonemapper to be tweaked. Value is in the range [0.0, 1.0]. [INCLUDE:#TonemapperToeFactor] |
+Emissive values are now physically correct so that as the emissive power increases the color will become lighter, similarly to how colored lights work in the real world. As color gets tone mapped, if the final
+color is bright enough to start saturating the film / sensor, it will become white.
 
-<!--
-[EXCERPT:SceneColorTint]
 [OBJECT:ComparisonSlider]
 	[PARAM:before]
-	![Scene Color Tint - (1.0, 1.0, 1.0)](scene_tint_full_small.png)(w:600)
+	![Filmic Tone Mapper | Emissive](Bloom_FilmicToneMapper.png)(w:700)
 	[/PARAM]
 	[PARAM:after]
-	![Scene Color Tint - (0.25, 0.25, 0.25)](scene_tint_quarter_small.png)(w:600)
+	![Old Tone Mapper | Emissive](Bloom_OldToneMapper.png)(w:700)
 	[/PARAM]
 [/OBJECT]
-[/EXCERPT:SceneColorTint]
 
-[EXCERPT:TonemapperToeFactor]
+
+
+In this example, the "old" tonemapper is compared to the new filmic tone mapper. With a high enough emissive power, the color will start to become white, unlike the old tone mapper where values would just
+become overly saturated causing areas of the material to lose its detail. With the filmic tone mapper, you can even see that the Bloom in this scene retains its physical correctness as well as it is never 
+overly saturated, thus retaining it's reflective primary color value.
+
+#### Exposure Levels
+
+Exposure levels are physically correct so that colors will continue to have shape to them rather than lose detail.
+
+![](Exposure_Default.png)(w:500)
+
+In this example, there are a variety of colors and material types with varying roughness and metalness. The Exposure Bias is also set to 0. 
+
 [OBJECT:ComparisonSlider]
 	[PARAM:before]
-	![Toe Factor - 0](toe_factor_0_small.png)(w:600)
+	![Filmic Tone Mapper | Exposure Bias: 3](Exposure_FilmicToneMapper.png)(w:700)
 	[/PARAM]
 	[PARAM:after]
-	![Toe Factor - 1](toe_factor_1_small.png)(w:600)
+	![Old Tone Mapper | Exposure Bias: 3](Exposure_OldToneMapper.png)(w:700)
 	[/PARAM]
 [/OBJECT]
-[/EXCERPT:TonemapperToeFactor]
--->
 
-## Color Correction
+In this comparison, the Exposure Bias has been increased to a value of 3 so really show the response of the Filmic tone mapper compared to the "old" tone mapper implementation. Also, a value of three means
+that the processed image will be eight times brighter than the default value of 0 shown in the previous example above.
+
+As the exposure level is increased, the difference between the filmic and the "old" tone mapper will become more apparent. The spheres in the filmic example will continue to have shape to their color
+and shading, even as they appear to look brighter. The sphere using the "old" tone mapper will start to have their colors blend with the shading. The spheres become really bright as a solid color, but 
+you can't necessarily tell that they have a high exposure value, whereas the filmic tone mapper takes this into account. The filmic response looks like a natural camera effect if you were to increase its
+exposure.
+
+### Settings
+
+In the **Tone Mapper** section of the Post Process Volume, you'll find the properties that match the ACES standardized film stock. You can adjust these tone mapper controls to emulate other types of 
+film stock for your project. 
+
+[REGION:note]
+It is recommended that changes to these properties be used project-wide for a specific look and not changed dynamically or on a shot-by-shot basis. Instead, you should use the 
+[Color Grading](#colorcorrectionproperties) properties for any artistic adjustments.
+[/REGION]
 
-[EXCERPT:LookupTexture]
-Color correction is implemented through lookup tables (LUTs). Instead of using three 1d lookup tables, a single 3d lookup table is used as this offers more sophisticated color transformations (e.g. desaturation). The following image shows the 16x16x16 color neutral LUT unwrapped to a 256x16 texture (how it is currently visible in the texture browser):
+![](ToneMapperProperties.png)(w:500)
 
 [REGION:raw]
-![Neutral LUT](RGBTable16x1.png)
+| Property | %spacer3% | Description |
+| -------- | -------- | ----------- |
+| [REGION:tablesection] Tonemapper [/REGION] |||
+| **Slope** || This will adjust the steepness of the S-curve used for the tone mapper, where larger values will make the slope steeper (darker) and lower values will make the slope less steep (lighter). Value is in the range of [0.0, 1.0]. [0.0, 1.0]. [INCLUDE:#slopetable]  |
+| **Toe** || This will adjust the dark color in the tone mapper. Value is in the range of [0.0, 1.0]. [0.0, 1.0]. [INCLUDE:#toetable] |
+| **Shoulder** || This will adjust the bright color in the tone mapper. Value is in the range of [0.0, 1.0] [INCLUDE:#shouldertable] |
+| **Black Clip** || This will set where the crossover happens where black's start to cut off their value. In general, this value should **NOT** be adjusted. Value is in the range of [0.0, 1.0] [INCLUDE:#blackcliptable] |
+| **White Clip** || This will set where the crossover happens where white's start to cut off their values. This will appear as a subtle change in most cases. Value is in the range of [0.0, 1.0] [INCLUDE:#whitecliptable] |
 [/REGION]
 
 [REGION:note]
-You can use the RGBTable16x1.png image from this page, but do not copy and paste as the image displayed on the monitor is not necessarily representative of the source. 
+If you'd like to test different tone mapper values in a real-time grpah like the ones above, you can use the set up default [UE4 Tonemapper example](https://www.desmos.com/calculator/h8rbdpawxj) here.
+[/REGION]
+
+
+<!--- Tone Mapper Properties Excerpts ---
+
+
+[EXCERPT:slopetable]
+| [INCLUDE:#slopegraph] | [INCLUDE:#slopeimage] |
+| --- | --- |
+|||
+[/EXCERPT:slopetable]
+
+[EXCERPT:toetable]
+| [INCLUDE:#toegraph] | [INCLUDE:#toeimage] |
+| --- | --- |
+|||
+[/EXCERPT:toetable]
+
+[EXCERPT:shouldertable]
+| [INCLUDE:#shouldergraph] | [INCLUDE:#shoulderimage] |
+| --- | --- |
+|||
+[/EXCERPT:shouldertable]
+
+[EXCERPT:blackcliptable]
+| [INCLUDE:#blackclipgraph] | [INCLUDE:#blackclipimage] |
+| --- | --- |
+|||
+[/EXCERPT:blackcliptable]
+
+[EXCERPT:whitecliptable]
+| [INCLUDE:#whiteclipgraph] | [INCLUDE:#whiteclipimage] |
+| --- | --- |
+|||
+[/EXCERPT:whitecliptable]
+
+---------------------------------
+
+[EXCERPT:slopegraph]
+[OBJECT:ComparisonSlider]
+	[PARAM:before]
+	![Slope: 0.88 (Default)](DefaultSettings_FilmicToneMapper.png)(w:450)
+	[/PARAM]
+	[PARAM:after]
+	![Slope: 0.6](Slope_0_6.png)(w:450)
+	[/PARAM]
+[/OBJECT]
+[/EXCERPT:slopegraph]
+
+[EXCERPT:slopeimage]
+[OBJECT:ComparisonSlider]
+	[PARAM:before]
+	![Slope: 0.88 (Default)](DefaultToneMapper.png)(w:450)
+	[/PARAM]
+	[PARAM:after]
+	![Slope: 0.6](SlopeToneMap.png)(w:450)
+	[/PARAM]
+[/OBJECT]
+[/EXCERPT:slopeimage]
+
+------------------------------
+
+[EXCERPT:toegraph]
+[OBJECT:ComparisonSlider]
+	[PARAM:before]
+	![Toe: 0.55 (Default)](DefaultSettings_FilmicToneMapper.png)(w:450)
+	[/PARAM]
+	[PARAM:after]
+	![Toe: 0.8](Toe_0_8.png)(w:450)
+	[/PARAM]
+[/OBJECT]
+[/EXCERPT:toegraph]
+
+[EXCERPT:toeimage]
+[OBJECT:ComparisonSlider]
+	[PARAM:before]
+	![Toe: 0.55 (Default)](DefaultToneMapper.png)(w:450)
+	[/PARAM]
+	[PARAM:after]
+	![Toe: 0.8](ToeToneMap.png)(w:450)
+	[/PARAM]
+[/OBJECT]
+[/EXCERPT:toeimage]
 
-* [LUT Texture Example](RGBTable16x1.png) - **right-click** and **Save as**
+-----------------------------
+
+[EXCERPT:shouldergraph]
+[OBJECT:ComparisonSlider]
+	[PARAM:before]
+	![Shoulder: 0.26 (Default)](DefaultSettings_FilmicToneMapper.png)(w:450)
+	[/PARAM]
+	[PARAM:after]
+	![Shoulder: 1](Shoulder_1.png)(w:450)
+	[/PARAM]
+[/OBJECT]
+[/EXCERPT:shouldergraph]
 
+[EXCERPT:shoulderimage]
+[OBJECT:ComparisonSlider]
+	[PARAM:before]
+	![Shoulder: 0.26 (Default)](DefaultToneMapper.png)(w:450)
+	[/PARAM]
+	[PARAM:after]
+	![Shoulder: 1](ShoulderToneMap.png)(w:450)
+	[/PARAM]
+[/OBJECT]
+[/EXCERPT:shoulderimage]
+
+-----------------------------
+
+[EXCERPT:blackclipgraph]
+[OBJECT:ComparisonSlider]
+	[PARAM:before]
+	![Black Clip: 0 (Default)](DefaultSettings_FilmicToneMapper.png)(w:450)
+	[/PARAM]
+	[PARAM:after]
+	![Black Clip: 0.1](BlackClip_0_1.png)(w:450)
+	[/PARAM]
+[/OBJECT]
+[/EXCERPT:blackclipgraph]
+
+[EXCERPT:blackclipimage]
+[OBJECT:ComparisonSlider]
+	[PARAM:before]
+	![Black Clip: 0 (Default)](DefaultToneMapper.png)(w:450)
+	[/PARAM]
+	[PARAM:after]
+	![Black Clip: 0.1](BlackClipToneMap.png)(w:450)
+	[/PARAM]
+[/OBJECT]
+[/EXCERPT:blackclipimage]
+
+-----------------------------
+
+[EXCERPT:whiteclipgraph]
+[OBJECT:ComparisonSlider]
+	[PARAM:before]
+	![White Clip: 0.04 (Default)](DefaultSettings_FilmicToneMapper.png)(w:450)
+	[/PARAM]
+	[PARAM:after]
+	![White Clip: 0.2](WhiteClip_0_2.png)(w:450)
+	[/PARAM]
+[/OBJECT]
+[/EXCERPT:whiteclipgraph]
+
+[EXCERPT:whiteclipimage]
+[OBJECT:ComparisonSlider]
+	[PARAM:before]
+	![White Clip: 0.04 (Default)](DefaultToneMapper.png)(w:450)
+	[/PARAM]
+	[PARAM:after]
+	![White Clip: 0.2](WhiteClipToneMap.png)(w:450)
+	[/PARAM]
+[/OBJECT]
+[/EXCERPT:whiteclipimage]
+
+--- End Tone Mapper Properties Excerpts --->
+
+
+## Color Correction
+
+Color correction, or color grading, is used to alter or enhance the overall color of light in your scene. With the advent of HDR displays, the need to preserve color
+before it's processed is very important. It ensures that the color will be properly displayed.  
+
+In the past, the color correction would have been implemented through [Lookup Tables (LUTs)](Engine/Rendering/PostProcessEffects/UsingLUTs), which happens in LDR and on the final color that was output 
+to the monitor in sRGB color space. This presents a problem when working with HDR displays because the LUT is just a snapshot in time of the currently supported display it is
+being adjusted for and cannot be used to the same effect on an HDR display. To solve these types of issues, the [color correction controls](#colorcorrectioncontrols) do their work in Scene Referred Linear Space, 
+which means all colors are captured before they are tone mapped. This makes it possible to only adjust your colors on one HDR display for any color correction to look correct on all displays the image is being 
+outputted to, whether they are HDR or LDR.
+
+[REGION:note]
+You can learn more about the HDR pipeline and how this works for the benefit of your content in the [High Dynamic Range Display Output](Engine/Rendering/HDRDisplayOutput) page.
+[/REGION]
+
+### Color Correction Properties
+
+In the **Color Grading** section of the Post Process Volume, you'll find the properties that enable you to have the most artistic control over the scene. 
+
+#### Settings
+
+[REGION:imagetable]
+| ![](RGBColorWheel.png)(w:500) | ![](HSVColorWheel.png)(w:500) |
+| --- | --- |
+|RGB|HSV|
+[/REGION]
+
+Under each section, you can use the color wheel to select and drag the color values around. You also can choose between the following modes:
+
+* RGB - This will adjust the Red, Green, Blue values.
+* HSV - This will adjust the Hue and Saturation values.
+
+[REGION:tip]
+For more precise control over the slider values, you can hold the **Shift** key while dragging them.
 [/REGION]
 
-A modified texture might look like this one:
+![](ColorGradingProperties.png)(w:500)
 
 [REGION:raw]
-![Sepia Tone LUT](LUT_Sepia.png)
+| Property | %spacer2% | Description |
+| -------- | -------- | ----------- |
+| [REGION:tablesection] White Balance [/REGION] || The properties in this section are used to adjust the colors in the scene so that whites appear truly white. This allows for other colors in the scene to be correctly lit under the given lighting in the scene. |
+| **Temp** || This will adjust the white balance in relation to the temperature of the light in the scene. When the light temperature and this one match the light will appear white. When a value is used that is higher than the light in the scene it will yield a "warm" or yellow color, and, conversely, if the value is lower, it would yield a "cool" or blue color. |
+| **Tint** || This will adjust the white balance temperature tint for the scene by adjusting the cyan and magenta color ranges. Ideally, this setting should be used once you've adjusted the white balance **Temp**  property to get accurate colors. Under some light temperatures, the colors may appear to be more yellow or blue. This can be used to balance the resulting color to look more natural.  |
+| [REGION:tablesection] Global [/REGION] || The properties in this section are a global set of color corrections you can use for your scene. |
+| **Saturation** || This will adjust the intensity (purity) of the colors (hue) that are being represented. A higher saturation intensity will result in colors appearing more like their purest forms (red, green, blue) and when saturation is lowered colors will appear more gray or washed-out. |
+| **Contrast** || This will adjust the tonal range of light and dark color values in your scene. Lowering the intensity will remove highlights and lighten the image resulting in a washed-out appearance, whereas a higher intensity will tighten the highlights and darken the overall image. |
+| **Gamma** || This will adjust the luminance intensity of the image's mid-tones to accurately reproduce colors. Lowering or raising this value will result in the image being washed-out or too dark. |
+| **Gain** || This will adjust the luminance intensity of the image's whites (highlights) to accurately reproduce colors. Raising or lowering this value will result in the image highlights being washed-out or too dark. |
+| **Offset** || This will adjust the luminance intensity of the image's blacks (shadows) to accurately reproduce colors. Raising or lowering this value will result in the image shadows being washed-out or too dark. |
+| **Scene Color Tint** || This is a multiplier for color that is a filter color applied to the HDR scene color. |
+| **Color Grading LUT Intensity** || A scaling factor that controls the effect of the color correction. |
+| **Color Grading LUT** || The LUT texture to use as a lookup table for color correction. [INCLUDE:Engine/Rendering/PostProcessEffects/ColorGrading#ColorGradingLUT] |
+| [REGION:tablesection] Shadows [/REGION] || The properties in this section are used to adjust the color correction values for shadows in the scene. |
+| **Saturation** || This will adjust the intensity (purity) of the colors (hue) that are being represented. A higher saturation intensity will result in colors appearing more like their purest forms (red, green, blue) and when saturation is lowered colors will appear more gray or washed-out. |
+| **Contrast** || This will adjust the tonal range of light and dark color values in your scene. Lowering the intensity will remove highlights and lighten the image resulting in a washed-out appearance, whereas a higher intensity will tighten the highlights and darken the overall image. |
+| **Gamma** || This will adjust the luminance intensity of the image to accurately reproduce colors. Raising or lowering this value will result in the image mid-tones being washed-out or too dark. |
+| **Gain** || This will adjust the luminance intensity of the image's whites (highlights) to accurately reproduce colors. Raising or lowering this value will result in the image highlights being washed-out or too dark. |
+| **Offset** || This will adjust the luminance intensity of the image's blacks (shadows) to accurately reproduce colors. Raising or lowering this value will result in the image shadows being washed-out or too dark. |
+| **Shadows Max** || This is a multiplier of the properties that affect the color correction properties that have been adjusted in the **Shadows** section. |
+| [REGION:tablesection] Mid-tones [/REGION] || The properties in this section are used to adjust the color correction values for mid-tones in the scene. |
+| **Saturation** || This will adjust the intensity (purity) of the colors (hue) that are being represented. A higher saturation intensity will result in colors appearing more like their purest forms (red, green, blue) and when saturation is lowered colors will appear more gray or washed-out. |
+| **Contrast** || This will adjust the tonal range of light and dark color values in your scene. Lowering the intensity will remove highlights and lighten the image resulting in a washed-out appearance, whereas a higher intensity will tighten the highlights and darken the overall image. |
+| **Gamma** || This will adjust the luminance intensity of the image to accurately reproduce colors. Raising or lowering this value will result in the image mid-tones being washed-out or too dark. |
+| **Gain** || This will adjust the luminance intensity of the image's whites (highlights) to accurately reproduce colors. Raising or lowering this value will result in the image highlights being washed-out or too dark. |
+| **Offset** || This will adjust the luminance intensity of the image's blacks (shadows) to accurately reproduce colors. Raising or lowering this value will result in the image shadows being washed-out or too dark. |
+| [REGION:tablesection] Highlights [/REGION] || The properties in this section are used to adjust the color correction values for Highlights in the scene. |
+| **Saturation** || This will adjust the intensity (purity) of the colors (hue) that are being represented. A higher saturation intensity will result in colors appearing more like their purest forms (red, green, blue) and when saturation is lowered colors will appear more gray or washed-out. |
+| **Contrast** || This will adjust the tonal range of light and dark color values in your scene. Lowering the intensity will remove highlights and lighten the image resulting in a washed-out appearance, whereas a higher intensity will tighten the highlights and darken the overall image. |
+| **Gamma** || This will adjust the luminance intensity of the image to accurately reproduce colors. Raising or lowering this value will result in the image mid-tones being washed-out or too dark. |
+| **Gain** || This will adjust the luminance intensity of the image's whites (highlights) to accurately reproduce colors. Raising or lowering this value will result in the image highlights being washed-out or too dark. |
+| **Offset** || This will adjust the luminance intensity of the image's blacks (shadows) to accurately reproduce colors. Raising or lowering this value will result in the image shadows being washed-out or too dark. |
+| **HighLights Min** || This is a multiplier of the properties that affect the color correction properties that have been adjusted in the **Highlights** section.  |
 [/REGION]
 
-To make use of a LUT, you need to assign the LUT texture in the post process volume you want to use it. 
 
-| Property | Description |
-| -------- | ----------- |
-| **Color Grading Intensity** | A scaling factor that controls the effect of the color correction. |
-| **Color Grading LUT** | The LUT texture to use as a lookup table for color correction. [INCLUDE:Engine/Rendering/PostProcessEffects/ColorGrading#ColorGradingLUT] |
+### Workflow for Good Color Correction
 
-<!--
-[EXCERPT:ColorGradingLUT]
-[REGION:imagetable]
-| ![Natural Scene](lut_none.png)(w:265) | ![Greenish Scene](lut_green.png)(w:265) | ![Reddish Scene](lut_red.png)(w:265) |
-| ------ | ------ | ------ |
-| ![Neutral LUT](RGBTable16x1.png) | ![Greenish LUT](LUT_Greenish.png) | ![Reddish LUT](LUT_Reddish.png) |
-[/REGION:imagetable]
-[/EXCERPT:ColorGradingLUT]
--->
+These are some good habits that should be used when working with the color grading tools in your Post Process Volume. 
+
+* The **Tone Mapper** controls should be set for a project-wide look. Ideally, these properties should not be adjusted and certainly should not be adjusted dynamically or on a shot-by-shot basis. 
+Use the [Color Grading](#colorcorrectionproperties) properties for this.
 
-### LUT Texture Creation
+* Start by using the **Global** color correction properties and doing an initial pass on your scene before touching any of the other properties. The majority of your color grading changes should be 
+done with these properties.
 
-The procedure:
+* You should not try to use the **Gain** property to adjust exposure for your scene. This should be used for fine tweaks for different colors. Instead use the **Exposure Bias** in the **Lens** > 
+**Auto-Exposure** section.
 
-1. Make representative screen shots of the scenes you want to adjust and put them in one Photoshop document.
-1. Load a neutral 256x16 LUT into Photoshop.
+* Use **Offset** in minor amounts to make your shadowed or darkened areas a little lighter. If you want things to look a little hazy or have more bloom, then you can set a higher value, conversely 
+if your scene is overly hazy or has too much bloom you can lower the value to subtract color from the scene.
 
-	[REGION:note]
-	You can use the RGBTable16x1.png image from this page, but do not copy and paste as the image displayed on the monitor is not necessarily representative of the source. 
+* Instead of using a [Lookup Table (LUT)](Engine/Rendering/PostProcessEffects/UsingLUTs) for coloring the scene, since that happens after tone mapping, you should use the Color Grading properties 
+to get the look you want so that it's consistent on all displays you output to (LDR and HDR). However, it can be a good practice to quickly get a look to set the tone with a LUT and then adjust the 
+color correction to get a similar look.
 
-	* [LUT Texture Example](RGBTable16x1.png) - **right-click** and **Save as**
 
-	[/REGION]
+## Training Streams
 
-1. Insert the LUT into the Photoshop document with the screenshots (Select All in the LUT document, copy, switch to screenshot document, paste).
-1. Apply color manipulations (best by adding adjustment layers, otherwise you need to flatten everything before and cutting out the 256x16 later gets more tricky).
-1. Select the 256x16 LUT (Select the LUT layer in the layer panel, Menu: Select/Load Selection, Ok).
-1. Copy merged the content of the LUT (Menu: Edit/Copy Merged).
-1. Paste and save the 256x16 texture in some uncompressed format that the engine can read (Menu: File/New, Menu: Edit/Paste, Menu: File/Save As).
-1. Import the texture in the editor and specify the ColorLookupTable Texture Group. 
+[OBJECT:EmbeddedVideo]
+	[PARAMLITERAL:videoid]
+	A-wectYNfRQ
+	[/PARAMLITERAL]
+	[PARAMLITERAL:width]
+	50
+	[/PARAMLITERAL]
+	[PARAMLITERAL:height]
+	56.25
+	[/PARAMLITERAL]
+	[PARAMLITERAL:units]
+	%
+	[/PARAMLITERAL]
+	[PARAMLITERAL:type]
+	simple
+	[/PARAMLITERAL]
+[/OBJECT]
 
 
-Various color corrections are possible. Here are a few examples:
 
-* Brightness
-* Saturation
-* Simple contrast (linear with clamping)
-* Higher quality contrast (e.g. curves with a steeper linear part in the middle)
-* Selective changes to the darks, midtones, and bright areas of an image (e.g. curves)
-* Selective changes to specific colors (best expressed in some color space where the brightness is in a separate channel e.g. LAB)
-* Adjustments can even be done in a different color space (e.g. LAB keeps brightness and color independent which can be very useful) 
 
-[/EXCERPT:LookupTexture]
+
+<!--
+[EXCERPT:ColorGradingLUT]
+[REGION:imagetable]
+| [INCLUDE:#colorgradingnone] | [INCLUDE:#colorgradinggreen] | [INCLUDE:#colorgradingred] |
+| ------ | ------ | ------ |
+| ![Neutral LUT](RGBTable16x1.png) | ![Greenish LUT](LUT_Greenish.png) | ![Reddish LUT](LUT_Reddish.png) |
+[/REGION:imagetable]
+[/EXCERPT:ColorGradingLUT]
+
+[EXCERPT:colorgradingnone]
+[REGION:lightbox]
+[![Natural Scene](lut_none.png)(w:300)](lut_none.png)
+[/REGION]
+[REGION:caption]Click image for full size view.[/REGION]
+[/EXCERPT:colorgradingnone]
+
+[EXCERPT:colorgradinggreen]
+[REGION:lightbox]
+[![Greenish Scene](lut_green.png)(w:300)](lut_green.png)
+[/REGION]
+[REGION:caption]Click image for full size view.[/REGION]
+[/EXCERPT:colorgradinggreen]
+
+[EXCERPT:colorgradingred]
+[REGION:lightbox]
+[![Reddish Scene](lut_red.png)(w:300)](lut_red.png)
+[/REGION]
+[REGION:caption]Click image for full size view.[/REGION]
+[/EXCERPT:colorgradingred]
+-->
\ No newline at end of file
diff --git a/Engine/Documentation/Source/Engine/Rendering/PostProcessEffects/DepthOfField/DepthOfField.INT.udn b/Engine/Documentation/Source/Engine/Rendering/PostProcessEffects/DepthOfField/DepthOfField.INT.udn
--- a/Engine/Documentation/Source/Engine/Rendering/PostProcessEffects/DepthOfField/DepthOfField.INT.udn
+++ b/Engine/Documentation/Source/Engine/Rendering/PostProcessEffects/DepthOfField/DepthOfField.INT.udn
@@ -110,7 +110,7 @@ Note that the circle size might not always match to the actual effect as the imp
 | **Occlusion** | Controls how much the blurred geometry extends beyond its usual silhouette and its opacity. A value of 0.18 gives a fairly natural occlusion result. A value of 0.4 may be necessary to solve layer color leaking issues. Very small values (less than 0.18) generally result in negating the blur effect, though can work well when objects are very near to the camera. [INCLUDE:#Occlusion] |
 | **Color Threshold** | Controls the threshold where adaptive DoF switches to using full-resolution based on color. Smaller values cause more of the scene to use full-resolution processing. |
 | **Size Threshold** | Controls the threshold where adaptive DoF switches to using full-resolution based on size. Larger values cause more of the scene to use full-resolution processing. |
-| **Sky Distance** | Artificial distance to allow the skybox to be in focus (e.g. 200000), <=0 to switch the feature off, only for GaussianDOF, can cost performance. |
+| **Sky Distance** | Artificial distance to allow the skybox to be in focus (e.g. 200000), `<=0` to switch the feature off, only for GaussianDOF, can cost performance. |
 
 <!--
 [EXCERPT:FocalDistance]
diff --git a/Engine/Documentation/Source/Engine/Rendering/PostProcessEffects/PaniniProjection/PaniniProjection.INT.udn b/Engine/Documentation/Source/Engine/Rendering/PostProcessEffects/PaniniProjection/PaniniProjection.INT.udn
--- a/Engine/Documentation/Source/Engine/Rendering/PostProcessEffects/PaniniProjection/PaniniProjection.INT.udn
+++ b/Engine/Documentation/Source/Engine/Rendering/PostProcessEffects/PaniniProjection/PaniniProjection.INT.udn
@@ -29,7 +29,7 @@ parameter, you can refer directly to the original research paper referenced at t
 
 r.Upscale.Panini.D > 0 directly emphasizes the effect. However it not the only parameter. r.Upscale.Panini.S 
 lerp with the paper's hardcompression. And if still some sphere still looks elipsed in the corners, then this
-parameter can also be tweaked as well. However, for r.Upscale.Panini.S < 0, start shows unrendered black
+parameter can also be tweaked as well. However, for `r.Upscale.Panini.S < 0`, start shows unrendered black
 looking pixels in the corners.
 
 ![](qagame_fov120_vertical_compression.png)
diff --git a/Engine/Documentation/Source/Engine/Rendering/PostProcessEffects/PostProcessMaterials/PostProcessMaterials.INT.udn b/Engine/Documentation/Source/Engine/Rendering/PostProcessEffects/PostProcessMaterials/PostProcessMaterials.INT.udn
--- a/Engine/Documentation/Source/Engine/Rendering/PostProcessEffects/PostProcessMaterials/PostProcessMaterials.INT.udn
+++ b/Engine/Documentation/Source/Engine/Rendering/PostProcessEffects/PostProcessMaterials/PostProcessMaterials.INT.udn
@@ -1,30 +1,40 @@
-Availability: Public
-Crumbs: %ROOT%
-Title: Post Process Materials
-Description: How to author and blend custom Post Process passes with the material editor.
-Type: Overview
-SkillLevel:Beginner
+Availability:Public
+Crumbs:
+Title:Post Process Materials
+Description:How to author and blend custom Post Process passes with the Material Editor. 
+Type:reference
 Version:4.16
 Parent:Engine/Rendering/PostProcessEffects
 Order:8
 Tags:Post Process
 Topic-image:PPMaterials_Topic.png
 
+[TOC(start:2 end:4)]
+
+[REGION:banner]
+![](Teaser_Banner.png)
+[/REGION]
 
+<!---
 ![](Teaser3.png)
 ![](Teaser0.png)
 ![](Teaser1.png)
 ![](Teaser2.png)
+-->
 
-[TOC(start:2 end:4)]
+
+Post Process Materials enable you to set up materials that can be used with your post process to create visual screen effects for damage, area type effects, or overall look for your game
+that can only be achieved via Post Process Materials. 
+
+In the following sections, you'll learn about setting up Post Process Materials, the available settings you can use, along with some examples of how to setup some of your own Post Process
+Materials using different buffers, blending between Post Process Materials, and more. 
 
 ## Post Processing Graph
 
 [EXCERPT:PostProcessMaterials_Desc]
 The engine already features complex post processing based on a graph of post processing nodes. The **Post Processing Materials**
-can be additionally inserted in some specific position. See [FAQ](#FAQ) **r.CompositionGraphDebug** to get a dump of the full graph.
-The graph is actually not only doing post processing but also some parts of the lighting. Over time we plan to make more
-parts be customizable through the material editor.
+can be additionally inserted in some specific position. See [Frequently Asked Questions section](#FAQ) for **r.CompositionGraphDebug** to get a dump of the full graph.
+The graph is actually not only doing post processing but also some parts of the lighting. 
 [/EXCERPT:PostProcessMaterials_Desc]
 
 Most of the time the graph is automatically creating intermediate render targets. This means if you want to blend with the former
@@ -39,7 +49,7 @@ Post Process Materials should be used sparingly and only when you really need th
 ## Using a Post Process Material
 
 Through post process settings (usually defined with post process volumes or camera settings) it is possible to blend (so called) blendable assets.
-At the moment, only **materials** and **material instances** are blendable assets. The engine provides a few post process materials but with this feature
+At the moment, only **Materials** and **Material Instances** are blendable assets. The engine provides a few post process materials but with this feature
 you can create your own **custom post processing** without any programmer help.
 
 Simply assign one or more post process materials to a post process volume in the **Blendables** section. First press the **+** to add
@@ -59,7 +69,7 @@ A post processing material needs to specify the material domain **post process**
 
 ![](DomainPostProcess.png)
 
-The material should only use *EmissiveColor* to output the new color. Additionally, you can define where during the post processing this pass should be applied
+The material should only use **Emissive Color** to output the new color. Additionally, you can define where during the post processing this pass should be applied
 and, if there are multiple, in which order they should be processed (Priority):
 
 ![](PostProcessMaterialProps.png)
@@ -122,7 +132,7 @@ With camera movement, you can perceive soft linear transition between the two ef
 
 The following shows a top down view of a level that has how two volumes. The large unbound volume has a red material instance and the small volume has a green material instance specified as blendable.
 The smaller volume has a higher priority. The material params are getting blended depending on the camera position.
-The fuzzy borders are defined by the BlendRadius property specified in the volumes and extend the volume shape.
+The fuzzy borders are defined by the Blend Radius property specified in the volumes and extend the volume shape.
 
 With the correct setup all blending happens as expected.
 
@@ -181,7 +191,7 @@ Anti-Aliasing generally makes this a bit more difficult as a GBuffer pixel/texel
 
 ## CustomDepth
 
-This separate feature allows masking of certain objects by rendering them into another depth buffer (called custom depth buffer).
+This separate feature enables masking of certain objects by rendering them into another depth buffer (called custom depth buffer).
 This adds extra draw calls but no more materials. The rendering is fairly cheap as we only output depth. The feature can be activated on the mesh
 (e.g. Static Mesh Properties / Render Custom Depth):
 
@@ -199,11 +209,41 @@ This is the material we used to visualize it:
 
 ![](CustomDepthMat.png)
 
+## CustomDepth Stencil
+
+Custom Depth Stencil is an extension of Custom Depth where you can use a stencil, or cutout, of your rendered object to then do visually interesting things like the example below that enables you to 
+visualize occlude objects, draw object outlines, or only be visible from certain viewing angles. There is a lot of potential with what you can do by having access to the stencil of an Actor in your
+scene. Using the following settings to enable and assign a stencil value.
+
+![](CustomStencilSettings.png)
+
+In this scene, Custom Depth is enabled on three objects and a **Custom Depth Stencil Value** is set for each, but without any post processing pass to visualize the content, this feature remains invisible.
+
+![](CustomDepthStencilScene.png)(w:520)
+
+Once you've setup your Post Process Material you will be able to visualize how the Custom Depth Stencil looks, The occluded objects render with a randomly assigned color based on the **Custom Depth Stencil Value** used. 
+
+![](CustomDepthStencilVisualization.png)(w:520)
+
+This is a setup of the Material we used to visualize it:
+
+[REGION:lightbox]
+[![](CustomDepthStencilMaterial.png)(w:768)](CustomDepthStencilMaterial.png)
+[/REGION]
+[REGION:caption]
+_Click image for full size._
+[/REGION]
+
+This is by no means, the only way to use Custom Depth Stencil and in this particular Material setup the stencil is being divided so that it uses values between 1 and 255, a mask is being used 
+for any value that is between these values, a random color is being created for these values as well so as the Custom Depth Stencil Value changes so does the color, and finally the mask that is
+created is being used in a way that only colors the stencil if the object is occluded.
+
+
 ## Temporal Anti-Aliasing or Why the GBuffer Jitters
 
 Temporal anti-aliasing is a unique UE4 feature that greatly improves image quality with a very moderate performance cost.
 
-By default, a post process material is inserted in the end of the post processing graph (after tone mapper). This means you get the final LDR color after tonemapping, color grading and after the temporal anti-aliasing
+By default, a post process material is inserted in the end of the post processing graph (after tone mapper). This means you get the final LDR color after tone mapping, color grading and after the temporal anti-aliasing
 was applied. This is the best spot for many simple post processing effects - for performance and ease of use.
 
 Here you can see how we used the custom depth input to visualize a silhouette around specific objects:
@@ -256,7 +296,7 @@ Scalar parameter:
 	
 [EXCERPT:PostMatWarning]
 [REGION:warning]
-In order to get the parameters you need to create a material instance from the post process material!
+In order to get the parameters you need to create a material instance from the post process material.
 [/REGION]
 [/EXCERPT:PostMatWarning]
 
@@ -270,7 +310,7 @@ The following issues still need to be fixed:
     * Certain lookups do not work in certain passes (e.g. some we might not fix because they would cost too much performance).
     * MaterialFunction might report an error but it still can be used in a Material with the Post Process domain.
   * **Material**
-    * UV in the PostProcesMaterial might not be in 0-1 range (e.g. in editor when reducing viewport size), this aligns with the lookup but makes it hard to implement something like an vignette effect.
+    * UV in the PostProcesMaterial might not be in 0-1 range (e.g. in editor when reducing viewport size), this aligns with the lookup but makes it hard to implement something like a vignette effect.
     * The asset thumbnail of a post process material does not look right.
     * Outputting alpha is not yet supported (should go through Opacity).
     * The preview material in the material editor does not look right.
@@ -294,7 +334,7 @@ The following issues still need to be fixed:
 
   * **How much memory do I pay for a post process?**
 
-    Memory cost depends on screen resolution. Before tonemapping we use HDR (8bytes per pixel), after we use LDR (4 bytes per pixel).
+    Memory cost depends on screen resolution. Before tone mapping we use HDR (8bytes per pixel), after we use LDR (4 bytes per pixel).
 
   * **How can I lower the render cost for post processing?**
 
diff --git a/Engine/Documentation/Source/Engine/Rendering/PostProcessEffects/UsingLUTs/UsingLUTs.INT.udn b/Engine/Documentation/Source/Engine/Rendering/PostProcessEffects/UsingLUTs/UsingLUTs.INT.udn
new file mode 100644
--- /dev/null
+++ b/Engine/Documentation/Source/Engine/Rendering/PostProcessEffects/UsingLUTs/UsingLUTs.INT.udn
@@ -0,0 +1,205 @@
+Availability: Public
+Crumbs: %ROOT%
+Title: Using Lookup Tables (LUTs) for Color Grading
+Description: How to create and use your own lookup tables for color grading in Unreal Engine 4.
+Type:how-to
+SkillLevel:Beginner
+Version:4.16
+Parent:Engine/Rendering/PostProcessEffects
+Related:Engine/Rendering/PostProcessEffects/ColorGrading
+Order:1000
+Tags:Post Process
+Topic-image:ColorGrading_Topic.png
+Social-image:ColorGrading_Overview.png
+
+[VAR:spacer]![](%ROOT%/spacer.png)(w:150 convert:false)[/VAR]
+[VAR:spacer2]![](%ROOT%/spacer.png)(w:40 convert:false)[/VAR]
+[VAR:spacer3]![](%ROOT%/spacer.png)(w:0 convert:false)[/VAR]
+
+
+![](Engine/Rendering/PostProcessEffects/ColorGrading/ColorGrading_Overview.png)
+
+
+[TOC(start:2 end:2)]
+
+A **Lookup Table** (LUT) can be used to achieve color correction using a Post Process Volume. Instead of using three one-dimensional (1D) lookup tables, a single three-dimensional (3D) Lookup Table is used. This offers 
+more sophisticated color transformation, which can be used for something like desaturation. 
+
+[REGION:imagetable]
+| ![Neutral LUT](RGBTable16x1.png)(w:350)  | ![Sepia Tone LUT](LUT_Sepia.png)(w:350) |
+| --- | --- |
+| Color Neutral LUT | Sepia Toned LUT |
+[/REGION]
+
+A Lookup Table used in Unreal Engine 4 (UE4) is a 16x16x16 color neutral LUT unwrapped to a 256x16 texture. These examples show a color neutral and a sepia toned one. If you were to apply the neutral LUT, you wouldn't see any changes
+to a default scene in UE4, however, if you used the sepia toned one, you'd get something like this: 
+
+[OBJECT:ComparisonSlider]
+	[PARAM:before]
+	![Color Neutral LUT](LUT1.png)(w:600)
+	[/PARAM]
+	[PARAM:after]
+	![Sepia Tone LUT](LUT4.png)(w:600)
+	[/PARAM]
+[/OBJECT]
+
+
+
+## Using a LUT Texture
+
+To make use of a LUT, use the following properties to assign the LUT texture to the Post Process Volume you want to use it with and adjust its intensity.
+
+[REGION:raw]
+| Property | %spacer2% | Description |
+| -------- | --------- | ----------- |
+| **Color Grading LUT Intensity** || A scaling factor that controls the effect of the color correction. |
+| **Color Grading LUT** || The LUT texture to use as a lookup table for color correction. [INCLUDE:#ColorGradingLUT] |
+[/REGION]
+
+<!--
+[EXCERPT:ColorGradingLUT]
+[REGION:imagetable]
+| [INCLUDE:#colorgradingnone] | [INCLUDE:#colorgradinggreen] | [INCLUDE:#colorgradingred] |
+| ------ | ------ | ------ |
+| ![Neutral LUT](RGBTable16x1.png) | ![Greenish LUT](LUT_Greenish.png) | ![Reddish LUT](LUT_Reddish.png) |
+[/REGION:imagetable]
+[/EXCERPT:ColorGradingLUT]
+
+[EXCERPT:colorgradingnone]
+[REGION:lightbox]
+[![Natural Scene](lut_none.png)(w:300)](lut_none.png)
+[/REGION]
+[REGION:caption]Click image for full size.[/REGION]
+[/EXCERPT:colorgradingnone]
+
+[EXCERPT:colorgradinggreen]
+[REGION:lightbox]
+[![Greenish Scene](lut_green.png)(w:300)](lut_green.png)
+[/REGION]
+[REGION:caption]Click image for full size.[/REGION]
+[/EXCERPT:colorgradinggreen]
+
+[EXCERPT:colorgradingred]
+[REGION:lightbox]
+[![Reddish Scene](lut_red.png)(w:300)](lut_red.png)
+[/REGION]
+[REGION:caption]Click image for full size.[/REGION]
+[/EXCERPT:colorgradingred]
+-->
+
+## Creating Your Own LUT Texture
+
+To create your own LUT texture, you'll first need an image editing application (like Photoshop or GIMP) that can make color adjustments to a captured image. For this demonstration, 
+the **Sun Temple** project available in the Epic Games Launcher via the **Learn** tab will be used, and Photoshop will be the application of choice for the LUT image editing. Some terms may be specific
+to Photoshop in this example.
+
+1. In your UE4 project, start by taking at least one screenshot that is representative of your scene that you want to apply color correction too. You can use the [Screenshot Tool](Engine/Basics/Screenshots) available
+in the Editor to capture a high-resolution image of your scene. 
+	
+	![](LUT1.png)(w:850)
+
+1. Open the screenshot in Photoshop and proceed to add your own **Adjustment** layers from the Adjustments tab located above the Layers panel on the right. Here you'll find many different icons that represent the various types of changes you can make to your images, such as brightness and contrast, hue, saturation, and many others.
+
+	[REGION:lightbox]
+	[![](LUT2.png)(w:850)](LUT2.png)
+	[/REGION]
+	[REGION:caption]Click image for full size.[/REGION]
+
+1. From the **Adjustments** panel, click any of the icons to add a layer to the Layers panel. You can use these to make specific color corrections to the final image. For this example, 
+**Vibrance** and **Brightness/Contrasts** have been added and adjusted. Once these are added to your Layers panel, select them from the list to expose their properties, or if these aren't 
+visible, you can right-click and choose **Edit Adjustment** to expose them. 
+
+	![](LUT3.png)(w:850)
+
+1. Any adjustments you make will be visible in the viewport. Adjust the properties of the Adjustment layers until you get the color correction you desire. In the example below, the adjustments 
+for Vibrance and Brightness/Contrast have been adjusted to give the captured screenshot more color saturation and definition for visual appeal. 
+	 
+	[OBJECT:ComparisonSlider]
+		[PARAM:before]
+		![Before](LUT1.png)(w:850)
+		[/PARAM]
+		[PARAM:after]
+		![After](LUT4.png)(w:850)
+		[/PARAM]
+	[/OBJECT]
+
+1. Now that you've achieved your desired look with the Adjustment layers, you'll need to open a copy of the **Neutral Color LUT** on its own. The adjustment layers that were created for your screenshot can be dragged onto the neutral color LUT. 
+
+	[REGION:note]
+	If you do not have a Neutral Color LUT available to use, you can use the RGBTable16x1.png image shown in the introduction to this page. 
+	* [LUT Texture Example](RGBTable16x1.png) - right-click and **Save as**
+	[/REGION]
+
+	With your Screenshot with the adjustment layers and the Neutral Color LUT opened, **drag and drop** the Adjustment layers from your Screenshot Layers panel onto the Neutral Color LUT image. You should now see them populate in the LUT's Layer panel. 
+
+	[OBJECT:EmbeddedVideo]
+		[PARAMLITERAL:videoid]
+		IFFgPsML1yc
+		[/PARAMLITERAL]
+		[PARAMLITERAL:width]
+		100
+		[/PARAMLITERAL]
+		[PARAMLITERAL:height]
+		56.25
+		[/PARAMLITERAL]
+		[PARAMLITERAL:units]
+		%
+		[/PARAMLITERAL]
+		[PARAMLITERAL:type]
+		simple
+		[/PARAMLITERAL]
+	[/OBJECT]
+	
+1. Now, navigate to **File** > **Save As** and give your LUT with the color corrections a name.
+1. In UE4, import your color corrected LUT into the Editor using the **drag and drop** method into the Content Browser or by using the **Import** button that is also available in the Content Browser.
+	
+	![](LUT5.png)(w:850)
+
+1. For the best results using your LUT in a Post Process Volume, open the LUT image in the Texture Editor and set the following properties:
+
+	[REGION:imagetable]
+	| ![](LUT6.png)(w:450) | [INCLUDE:#textureproperties] |
+	| --- | --- |
+	|||
+	[/REGION]
+
+1. Next, in the Post Process Volume, use the **Details** panel and under the **Color Grading** > **Global** section, enable **Color Grading LUT** and apply your LUT. 
+	
+	[REGION:imagetable]
+	| ![](LUT8.png)(h:400) | [INCLUDE:#finalcomparison] |
+	| --- | --- |
+	|||
+	[/REGION]
+
+<!---
+[INCLUDE:#textureproperties]
+[EXCERPT:textureproperties]
+* Mip Gen Settings: **NoMipMaps**
+* Texture Group: **Color Lookup Table**
+[/EXCERPT:textureproperties]
+
+[INCLUDE:#finalcomparison]
+
+[EXCERPT:finalcomparison]
+[OBJECT:ComparisonSlider]
+	[PARAM:before]
+	![Default](LUT1.png)(h:400)
+	[/PARAM]
+	[PARAM:after]
+	![With Created LUT](LUT7.png)(h:400)
+	[/PARAM]
+[/OBJECT]
+[/EXCERPT:finalcomparison]
+
+--->
+
+## Additional Notes
+
+The following are some things you should be aware of when using LUTs in your projects:
+
+* Since LUTs happen in Low Dynamic Range (LDR) and on the final image color that was output to the monitor in sRGB space, this means that it's just a snapshot in time for what display support is and will not necessarily look the same on any display its output to, unlike [High Dynamic Range](Engine/Rendering/HDRDisplayOutput) (HDR) which is processed so that any color correction happens before tone mapping. For this reason, it is recommended to use the latest [Color Grading](Engine/Rendering/PostProcessEffects/ColorGrading) color correction tools available in Unreal Engine 4.16 and later in place of using LUTs. 
+
+* LUTs can be a good way to quickly target a particular look for your project on your current display but then you should make final adjustments using the color correction properties found under the **Color Grading** section of the Post Process Volume to guarantee the look is consistent across different types of displays. 
+
+* Consistency is lost when using a LUT because any color corrections that have taken place will not work for an HDR display. The color grading properties in the Post Process Volume all operate in what is referred to as **Scene Referred Linear Space**. It means that the colors are consistent because they happen before tone mapping and before they are transformed into a particular display's color space. For example, if you wanted to output to a tv that uses Rec709 instead of sRGB, the color corrections would not carry over if using a LUT since color correction happens after the image has already been output to the display. For this reason, LUTs are being moved away from as a way of doing color correction image processing. 
+
diff --git a/Engine/Documentation/Source/Engine/Rendering/RenderTargets/BlueprintRenderTargets/Reference/BlueprintRenderTargetsReference.INT.udn b/Engine/Documentation/Source/Engine/Rendering/RenderTargets/BlueprintRenderTargets/Reference/BlueprintRenderTargetsReference.INT.udn
--- a/Engine/Documentation/Source/Engine/Rendering/RenderTargets/BlueprintRenderTargets/Reference/BlueprintRenderTargetsReference.INT.udn
+++ b/Engine/Documentation/Source/Engine/Rendering/RenderTargets/BlueprintRenderTargets/Reference/BlueprintRenderTargetsReference.INT.udn
@@ -1,12 +1,11 @@
 availability:Public
 title: Blueprint and Render Targets Reference
 description: A listing of properties and descriptions for using Blueprint & Render Targets.
-crumbs:%ROOT%, Engine, Engine/Rendering, Engine/Rendering/RenderTargets, Engine/Rendering/RenderTargets/BlueprintRenderTargets
+crumbs:%ROOT%,
 version:4.13
 Type:reference
 parent:Engine/Rendering/RenderTargets/BlueprintRenderTargets
 order:2
-Navigation:Reference
 tags: Render Targets
 tags: Blueprints
 topic-image:Engine/Rendering/RenderTargets/Render_Targets_Reference.png
@@ -27,10 +26,10 @@ Please note that DrawMaterialtoRenderTarget is the simpler version and should be
 
 |Property| Description|
 |----------|-----------------|
-|Texture Render Target| Texture render target 2D reference.|
-|Canvas| Reference to canvas.|
-|Size| Size of the render target.|
-|Context| Draw to render target Context Structure.|
+|**Texture Render Target**| Texture render target 2D reference.|
+|**Canvas**| Reference to canvas.|
+|**Size**| Size of the render target.|
+|**Context**| Draw to render target Context Structure.|
 
 ## End Draw Canvas to Render Target
 
@@ -40,7 +39,7 @@ Must be paired with a **BeginDrawCanvasToRenderTarget** to complete rendering to
 
 |Property| Description|
 |----------|-----------------|
-|Context| Draw to render target Context Structure.|
+|**Context**| Draw to render target Context Structure.|
 
 
 ## Clear Render Target 2D
@@ -51,8 +50,8 @@ Clears the specified render target with the given Clear Color.
 
 |Property| Description|
 |----------|-----------------|
-|Texture Render Target| Texture render target 2D reference.|
-|Clear Color| Clears the specified render target with the given Clear Color.|
+|**Texture Render Target**| Texture render target 2D reference.|
+|**Clear Color**| Clears the specified render target with the given Clear Color.|
 
 
 ## Create Render Target 2D
@@ -63,9 +62,9 @@ Creates a new render target and initializes it to the specified dimensions.
 
 |Property| Description|
 |----------|-----------------|
-|Width| Render targets width size.|
-|Height| Render targets height size.|
-|Return Value| Texture render target 2D reference.|
+|**Width**| Render targets width size.|
+|**Height**| Render targets height size.|
+|**Return Value**| Texture render target 2D reference.|
 
 ## Draw Material to Render Target
 
@@ -76,8 +75,8 @@ If you are rendering multiple primitives to the same target area, use **BeginDra
 
 |Property| Description|
 |----------|-----------------|
-|Texture Render Target| Texture render target 2D reference.|
-|Material| Material that the render target results will be saved to.|
+|**Texture Render Target**| Texture render target 2D reference.|
+|**Material**| Material that the render target results will be saved to.|
 
 ## Export Render Target
 
@@ -87,9 +86,9 @@ Exports a Render Target as an HDR image onto the disk.
 
 |Property| Description|
 |----------|-----------------|
-|Texture Render Target| Texture render target 2D reference.|
-|File Path| Path to where the render target will be saved.|
-|File Name| Name to give the saved render target.|
+|**Texture Render Target**| Texture render target 2D reference.|
+|**File Path**| Path to where the render target will be saved.|
+|**File Name**| Name to give the saved render target.|
 
 
 
diff --git a/Engine/Documentation/Source/Engine/Replay/Replay.INT.udn b/Engine/Documentation/Source/Engine/Replay/Replay.INT.udn
--- a/Engine/Documentation/Source/Engine/Replay/Replay.INT.udn
+++ b/Engine/Documentation/Source/Engine/Replay/Replay.INT.udn
@@ -29,7 +29,7 @@ The basic functions of the replay system are to start or stop recording a demo,
 | `UDemoNetDriver::GotoTimeInSeconds` | "demoscrub (ReplayTime)" | Jumps to the provided time (in seconds) in the replay. |
 | `UWorld::Pauser` (variable) | "demopause" | Setting this variable to the PlayerController held by the World's `DemoNetDriver->ServerConnection->OwningActor` variable will pause the replay. Setting it back to `NULL` will resume the replay. The "demopause" console command will automatically set or clear that variable. |
 | `WorldSettings->DemoPlayTimeDilation` (variable) | "demospeed (ReplaySpeed)" | Speeds up or slows down playback by setting the time-scaling factor. A value of 1.0 is normal speed. |
-| `AGameMode::bHandleDedicatedServerReplays` (variable) | N/A | Setting this variable to can be set to true to cause dedicated servers to record games automatically. |
+| `AGameMode::bHandleDedicatedServerReplays` (variable) | N/A | This variable can be set to `true` to cause dedicated servers to record games automatically. |
 
 ## DemoNetDriver and Streamers
 
diff --git a/Engine/Documentation/Source/Engine/Replay/Streamers/HTTPRESTAPI/HTTPRESTAPI.INT.udn b/Engine/Documentation/Source/Engine/Replay/Streamers/HTTPRESTAPI/HTTPRESTAPI.INT.udn
--- a/Engine/Documentation/Source/Engine/Replay/Streamers/HTTPRESTAPI/HTTPRESTAPI.INT.udn
+++ b/Engine/Documentation/Source/Engine/Replay/Streamers/HTTPRESTAPI/HTTPRESTAPI.INT.udn
@@ -47,7 +47,7 @@ Sends a request to begin uploading a replay stream.
 
 #### Response
 
-Expects a response of type `FNetworkReplayStartUploadingResponse`, which contains JSON data with the string "sessionId" defined. This value will be used as the **Session Name** in future communication from the HTTP Streamer. If the optional <Session Name Override> was provided, the returned "sessionId", and thus the ultimate <Session Name>, should be the same value. This is the game's way of demanding the session name.
+Expects a response of type `FNetworkReplayStartUploadingResponse`, which contains JSON data with the string "sessionId" defined. This value will be used as the **Session Name** in future communication from the HTTP Streamer. If the optional `<Session Name Override>` was provided, the returned "sessionId", and thus the ultimate `<Session Name>`, should be the same value. This is the game's way of demanding the session name.
 
 | Name | Type | Content |
 | --- | --- | --- |
@@ -444,7 +444,7 @@ Expects a response with content in the form of a JSON-formatted string.
 | replays | Array | Contains elements which consist of all the other fields in this table. |
 | AppName | String | The **Replay Version AppString** that was sent when the replay was uploaded. Should match the parameter you sent with the request. |
 | SessionName | String | The name (GUID) of the replay. |
-| FriendlyName | String | The <Platform Friendly Name> that was provided when the replay was uploaded. |
+| FriendlyName | String | The `<Platform Friendly Name>` that was provided when the replay was uploaded. |
 | Timestamp | DateTime | The time and date associated with the replay (i.e. when this game happened). |
 | SizeInBytes | Integer | Total size, in bytes, of the replay. |
 | DemoTimeInMs | Integer | Running length, in MS, of the replay. |
@@ -488,7 +488,7 @@ Expects a response with content as a JSON-formatted string containing a list of
 | replays | Array | Contains elements which consist of all the other fields in this table. |
 | AppName | String | The **Replay Version AppString** that was sent when the replay was uploaded. |
 | SessionName | String | The name (GUID) of the replay. |
-| FriendlyName | String | The <Platform Friendly Name> that was provided when the replay was uploaded. |
+| FriendlyName | String | The `<Platform Friendly Name>` that was provided when the replay was uploaded. |
 | Timestamp | DateTime | The time and date associated with the replay (i.e. when this game happened). |
 | SizeInBytes | Integer | Total size, in bytes, of the replay. |
 | DemoTimeInMs | Integer | Running length, in MS, of the replay. |
diff --git a/Engine/Documentation/Source/Engine/Sequencer/HowTo/ActorSequence/ActorSequence.INT.udn b/Engine/Documentation/Source/Engine/Sequencer/HowTo/ActorSequence/ActorSequence.INT.udn
--- a/Engine/Documentation/Source/Engine/Sequencer/HowTo/ActorSequence/ActorSequence.INT.udn
+++ b/Engine/Documentation/Source/Engine/Sequencer/HowTo/ActorSequence/ActorSequence.INT.udn
@@ -1,4 +1,4 @@
-Availability: Docs
+Availability: Public
 Crumbs: %ROOT%, Engine, Engine/Sequencer, Engine/Sequencer/HowTo
 Title:Using Embedded Sequencer Animations in Blueprints
 Description:Describes how you can embed Sequences within Actor Blueprints using an Actor Sequence Component.
diff --git a/Engine/Documentation/Source/Engine/Sequencer/HowTo/CharacterAnimation/CharacterAnimation.INT.udn b/Engine/Documentation/Source/Engine/Sequencer/HowTo/CharacterAnimation/CharacterAnimation.INT.udn
--- a/Engine/Documentation/Source/Engine/Sequencer/HowTo/CharacterAnimation/CharacterAnimation.INT.udn
+++ b/Engine/Documentation/Source/Engine/Sequencer/HowTo/CharacterAnimation/CharacterAnimation.INT.udn
@@ -1,21 +1,24 @@
 Availability: Public
-Title:Animation & Animation Blending
+Title:Blending Animations & Properties
 Crumbs: %ROOT%, Engine, Engine/Sequencer, Engine/Sequencer/HowTo
-Description:Illustrates how you can apply animation to a Skeletal Mesh Actor in Sequencer and blend animations. 
-Version: 4.15
+Description:Illustrates how you can apply animation to a Skeletal Mesh Actor in Sequencer and blend animations as well as blend property values. 
+Type: how-to
 SkillLevel: Intermediate
-parent:Engine/Sequencer/HowTo
-order:1
-checkpoint: editorqs
+Version: 4.17
+Parent: Engine/Sequencer/HowTo
+Order: 1
+Tags: Sequencer
 Related: Engine/Sequencer/Overview
 Related: Engine\Sequencer\ReferenceEditor\
-tags:Sequencer
 
 When creating cinematic sequences, there may be times where you want to have a character in your scene move around or perform some kind of animation.
-In **Sequencer**, this can be done by adding a Skeletal Mesh Actor to a Level Sequence then adding an Animation sub-track and specifying the animation that you want the Actor to perform.
-Repeating this process and adding different kinds of animations for your Actor to perform is what can bring your characters to life in your cinematic sequences!
+In **Sequencer**, this is done by adding a Skeletal Mesh Actor to a Level Sequence then adding an Animation sub-track and specifying the animation that you want the Actor to perform.
+Sequencer enables you to crossfade and blend animations (as well as property values) automatically by dragging and dropping a property tracksonto an existing property track. 
+When doing so, an automatic blend is applied which you can tune and tweak with easing properties or by defining your own custom curve asset. 
 
-Here in this guide, we will create a small scene that has a character start out walking and blend from a walk into a jump animation. 
+In this guide, we create a sample sequence of animations where a character blends from a walk into a jump and back to a walk again. 
+
+![](AnimationBlending.png)(w:720 h:480)
 
 ## Steps
 
@@ -31,7 +34,7 @@ For this how-to, we are using the **Blueprint Third Person Template** project wi
 
 	![](CharacterAnim_02.png)
 
-1.	With the **SK_Mannequin** asset still selected, click the **Add** button in Sequencer to add it to your Level Sequence. 
+1.	With the **SK_Mannequin** asset still selected, click the **Track** button in Sequencer to add it to your Level Sequence. 
 
 	![](CharacterAnim_03.png)
 
@@ -39,63 +42,102 @@ For this how-to, we are using the **Blueprint Third Person Template** project wi
 
 	![](CharacterAnim_04.png)
 
-1.	Expand **Animation**, then drag the **Timeline** to the **30** frame mark. 
+1.	Scrub the timeline to **30** and click the **Animation** button and add the **ThirdPersonJump_Start** animation. 
 
 	![](CharacterAnim_05.png)
 
-1.	Click the **Animation** button again and add the **ThirdPersonJump_Start** animation. 
+1.	Left-click and drag the **ThirdPersonJump_Start** animation in the tracks window onto the **ThirdPersonWalk** back to frame **25**. 
 
 	![](CharacterAnim_06.png)
 
-	If you play the animation now, you may notice a pop between the two animations. 
+1.	Add the **ThirdPersonJump_Loop** at frame **39** then drag it onto **ThirPersonJump_start** and crossfade it starting at frame **34**. 
 
-	[OBJECT:EmbeddedVideo]
-	[PARAMLITERAL:width]
-	640
-	[/PARAMLITERAL]
-	[PARAMLITERAL:height]
-	360
-	[/PARAMLITERAL]
-	[PARAMLITERAL:videoid]
-	oHjDpYPAdCo
-	[/PARAMLITERAL]
-	[/OBJECT]
+	![](CharacterAnim_07.png)
 
-	When adding animations, you can use the weight value to blend between multiple animations, which is how we will address the pop between animations. 
+1.	Add the **ThirdPersonJump_End** animation and crossfade it **2** frames. 
 
-1.	Drag the **ThirdPersonJump_Start** animation below and under the walk animation, back to frame **25** of the sequence. 
+	![](CharacterAnim_08.png)
 
-	![](CharacterAnim_07.png)
+1.	Add the **ThirdPersonWalk** and crossfrade it for **3** frames. 
 
-1.	Expand both sub-animation tracks and click the **add key** button on the walk track with a value of **1.0** and **0.0** for the jump start animation. 
+	![](CharacterAnim_09.png)
 
-	![](CharacterAnim_08.png)
+1.	At the start of one of the crossfades, right-click on the up triangle icon to bring up the **Easing Curve** options. 
 
-	This will effectively turn off the playback of the jump start animation since its weighted value is 0.0. 
+	![](CharacterAnim_10.png)
 
-1.	Scrub to frame **30** and add keys for both animations of **0.0** and **1.0** for the walk and jump animations, respectively. 
+1.	Under the **Options** section, select any of the desired **Easing** options to define how the curve is performed. 
 
-	![](CharacterAnim_09.png)
+	![](CharacterAnim_11.png)
+
+	Alternatively, you can change the **Method** to **Curve Asset**. 
+
+	![](CharacterAnim_12.png)
 
-	This will effectively blend in the jump start animation and blend out the walk animation. 
+	When doing so, the **Option** section will enable you to define which custom **Curve** asset you want to use. 
 
-1.	Click the **Preview Play** button to see the animation blending. 
+	![](CharacterAnim_13.png)
+
+1.	Scrub the timeline through the sequence to see the blending or click the **Preview Play** button. 
 
 
 ## End Result
 
-Your sequence now includes a Skeletal Mesh with multiple animations which blend together smoothly instead of popping between animations. 
+Your sequence now includes a Skeletal Mesh with multiple animations that smoothly blend together with easing options available that you can use to tune the blend. 
 
 [OBJECT:EmbeddedVideo]
+[PARAMLITERAL:videoid]
+ZV_jKfrmVM0
+[/PARAMLITERAL]
 [PARAMLITERAL:width]
-640
+100
 [/PARAMLITERAL]
 [PARAMLITERAL:height]
-360
+56.25
+[/PARAMLITERAL]
+[PARAMLITERAL:units]
+%
 [/PARAMLITERAL]
+[PARAMLITERAL:type]
+simple
+[/PARAMLITERAL]
+[/OBJECT]
+
+Another method of animation blending can be performed using the **Weight** sub-tracks on an animation track. 
+
+Below, we have one animation with a Weight value of **1.0** and another animation with a Weight value of **0.0**. 
+
+![](CharacterAnim_14.png)
+
+We then scrub the timeline ahead and reverse the Weight values for each which produces an interpolation between the two values causing the animations to blend.
+
+![](CharacterAnim_15.png)
+
+### More on Blending
+
+Blending is not restricted to animations, any 2D/3D transforms, vectors, colors, or floats can also be automatically blended. 
+
+By default, transform, vector, colors and similar tracks are set to **Infinite** which means the property value is persistant throughout the sequence. 
+By turning this option off and using property **Sections**, you can blend those Sections together. In the example video below, we blend the color of a light over time within Sequencer using a similar approach to this guide. 
+
+[OBJECT:EmbeddedVideo]
 [PARAMLITERAL:videoid]
-kVotxriT1Y4
+w9j6RaMG9Nc
+[/PARAMLITERAL]
+[PARAMLITERAL:width]
+100
+[/PARAMLITERAL]
+[PARAMLITERAL:height]
+56.25
+[/PARAMLITERAL]
+[PARAMLITERAL:units]
+%
+[/PARAMLITERAL]
+[PARAMLITERAL:type]
+simple
 [/PARAMLITERAL]
 [/OBJECT]
 
-You could continue with this example by adding in the jump loop and jump end animations and blend them in so the character jumps and lands. 
+
+
+
diff --git a/Engine/Documentation/Source/Engine/Sequencer/HowTo/SequencerHowTo.INT.udn b/Engine/Documentation/Source/Engine/Sequencer/HowTo/SequencerHowTo.INT.udn
--- a/Engine/Documentation/Source/Engine/Sequencer/HowTo/SequencerHowTo.INT.udn
+++ b/Engine/Documentation/Source/Engine/Sequencer/HowTo/SequencerHowTo.INT.udn
@@ -33,6 +33,7 @@ Refer to the table below for more information and click each link to view its co
 [](FANCY:Engine\Sequencer\HowTo\CineCameraActors "%Engine\Sequencer\HowTo\CineCameraActors:description%")
 [](FANCY:Engine\Sequencer\HowTo\TriggeringSequences "%Engine\Sequencer\HowTo\TriggeringSequences:description%")
 [](FANCY:Engine\Sequencer\HowTo\BurnIns "%Engine\Sequencer\HowTo\BurnIns:description%")
+[](FANCY:Engine\Sequencer\HowTo\WhenFinished "%Engine\Sequencer\HowTo\WhenFinished:description%")
 [/REGION]
 
 
diff --git a/Engine/Documentation/Source/Engine/Sequencer/HowTo/WhenFinished/WhenFinished.INT.udn b/Engine/Documentation/Source/Engine/Sequencer/HowTo/WhenFinished/WhenFinished.INT.udn
new file mode 100644
--- /dev/null
+++ b/Engine/Documentation/Source/Engine/Sequencer/HowTo/WhenFinished/WhenFinished.INT.udn
@@ -0,0 +1,149 @@
+Availability: Docs
+Crumbs: %ROOT%, Engine, Engine/Sequencer, Engine/Sequencer/HowTo
+Title: Keeping or Restoring Changes Made by Sequencer
+Description: Describes how to keep any changes made through Sequencer when it has completed or restore changes back to their original state.
+Type: how-to
+SkillLevel: intermediate
+Version: 4.16
+Parent: Engine/Sequencer/HowTo
+Order: 1
+Tags: Sequencer
+Tags: how-to
+topic-image:WhenFinishedTopicImage.png
+
+When working in Sequencer, you may make changes to the properties of your Actors that you want to persist after the cinematic finishes. For example, you might create a cinematic where a player unlocks and opens a door, and the door remains in the open state afterward. In a film production environment, however, you may want animated values in a shot to return to their pre-animated state so they don't bleed into the next shot. With the **When Finished** property, you can determine on a per-track basis whether tracks should return to their pre-animated state or keep changes when the sequence finishes.
+
+In this how-to guide, you will use the When Finished property to restore the color change of a light but keep the changes made to the movement of a door. 
+
+![](HeroImage.png)
+
+
+## Steps
+
+[REGION:note]
+For this how-to guide, we are using the **Blueprint Third Person Template** with **Starter Content** enabled. 
+[/REGION]
+
+1.	From the **Content/StarterContent/Props** folder, drag an **SM_Door** and **SM_DoorFrame** into the **Level Viewport** and position the door inside the frame. 
+
+	![](WhenFinished_01.png)
+
+	Press the **W** key to enter translation mode and, if needed, disable grid snapping (pictured below) to fine tune the position of the door in the frame. 
+
+	![](WhenFinished_02.png)
+
+1.	From the **Modes** panel under **Lights**, drag a **Point Light** into the **Level Viewport** and position it in front of the door. 
+
+	![](WhenFinished_03.png)
+
+1.	Position the **ThirdPersonCharacter** in the **Level Viewport** in front of the door. 
+
+	![](WhenFinished_04.png)
+
+1.	From the **Main Toolbar**, click the **Cinematics** button and select **Add Level Sequence**. 
+
+	![](WhenFinished_05.png)
+
+1.	In the **Save Asset As** window, select a **Name** and save location, then click **Save**. 
+
+	![](WhenFinished_06.png)
+
+1.	In the **Sequencer** window, click the **Add** button and select **Shot Track**. 
+
+	![](WhenFinished_07.png)
+
+	This will create a [Shot Track](Engine/Sequencer/HowTo/TracksShot) which enables you to section off parts of your cinematic into self-contained "shots" that you can arrange and re-arrange. 
+
+1.	On the **Shot Track**, click the **+ Shot** button and **Insert Shot** and use the default naming convention when prompted. 
+
+	![](WhenFinished_08.png)
+
+1.	Grab the end of the shot and pull it back to **Frame 75** so the shot is half the length of the entire sequence. 
+
+	![](WhenFinished_09.png)
+
+	We are going to use the shot to affect the light in our scene, but restore those changes once the shot is over.
+
+1.	With the **Timeline** at **Frame 0** and the **SM_Door** selected in the level, add the **Actor To Sequencer** from the **+ Add** button. 
+
+	![](WhenFinished_10.png)
+
+	We will affect the door on the master level and keep the state of the changes we are going to apply. 
+
+1.	Click the **+** button on the **Transform** track of **SM_Door**. 
+
+	![](WhenFinished_11.png)
+
+1.	Move the **Timeline** marker to **Frame 150**, then rotate the door (**E** key) in the level so it is open and press the **S** key to add its transform while open. 
+
+	![](WhenFinished_12.png)
+
+	This will cause the door to go from closed to open once we reach the end of the cinematic. 
+
+1.	Double-click on the shot in the **Shot Track** to open it, then in the level select the **Point Light** and add it to the shot. 
+
+	![](WhenFinished_13.png)
+
+1.	On the **Light Color** track, click the **+** sign to assign the default color to use at **Frame 0** of the shot. 
+
+	![](WhenFinished_14.png)
+
+1.	Move the **Timeline** marker to **Frame 75** and change the **Light Color** to **Green** (or any color you desire) adding keys for the new values. 
+
+	![](WhenFinished_15.png)
+
+1.	Right-click on the **Light Color** track in the tracks window and select **Properties** to view the **When Finished** state. 
+
+	![](WhenFinished_16.png)
+
+	By default, **When Finished** states that this track will **Restore State** (return to its original state) when the section is completed. 
+
+1.	Exit the **Shot** and return to the **Master Sequence** level by clicking the breadcrumb buttons in the upper-right of the **Sequencer** window. 
+
+	![](WhenFinished_17.png)
+
+1.	Right-click on the **Transform** track of the **SM_Door** and select **Properties**, then change the **When Finished** property to **Keep State**. 
+
+	![](WhenFinished_18.png)
+
+	We are now stating that we want to **Keep State**, meaning we want to keep the changes that we apply through Sequencer when the cinematic is complete. 
+
+1.	Select the **Level Sequence** Actor in the level, then, in the **Details** panel, set the **Auto Play** option to **True**. 
+
+	![](WhenFinished_19.png)
+
+1.	From the **Main Toolbar**, click the **Build** button and select **Build Lighting Only**, then (after lighting has built) click the **Play** button. 
+
+	![](WhenFinished_20.png)
+
+
+## End Result
+
+When you play in the level, the light will change colors before returning to white. Our door, however, will remain in the opened state. 
+
+[OBJECT:EmbeddedVideo]
+	[PARAMLITERAL:videoid]
+	VcM2Wf_Iycc
+	[/PARAMLITERAL]
+	[PARAMLITERAL:width]
+	100
+	[/PARAMLITERAL]
+	[PARAMLITERAL:height]
+	56.25
+	[/PARAMLITERAL]
+	[PARAMLITERAL:units]
+	%
+	[/PARAMLITERAL]
+	[PARAMLITERAL:type]
+	simple
+	[/PARAMLITERAL]
+[/OBJECT]
+
+For the light inside our shot, we used the default **When Finished** property value of **Restore State**, which will restore it back to its original state when the shot is finished. 
+For the door, however, we changed the When Finished property to **Keep State** preserving our changes and keeping the door in the open state when the cinematic completes. 
+
+One thing to consider is that, while Keep State has a local effect, there is a global **Restore State** option contained within the **Level Sequence** actor: 
+
+![](WhenFinished_21.png)
+
+The Restore State option in the **Details** panel of a Level Sequence actor will restore all states when the Level Sequence stops. If you want to keep the state of a track after the sequence has completely finished, you will need to make sure both that your section is set to Keep State, and the Level Sequence actor is **not** set to Restore State, otherwise the global value will override your local settings.
diff --git a/Engine/Documentation/Source/Engine/Sequencer/Overview/SequencerOverview.INT.udn b/Engine/Documentation/Source/Engine/Sequencer/Overview/SequencerOverview.INT.udn
--- a/Engine/Documentation/Source/Engine/Sequencer/Overview/SequencerOverview.INT.udn
+++ b/Engine/Documentation/Source/Engine/Sequencer/Overview/SequencerOverview.INT.udn
@@ -1,36 +1,30 @@
 Availability:Public
+Crumbs: %ROOT%
 Title: Sequencer Overview
-Crumbs:%ROOT%
 Description:Sequencer is Unreal Engine 4's multi-track editor, used for creating and previewing cinematic sequences in real-time. 
-Navigation:topic
-parent:Engine/Sequencer
-order:2
-Version: 4.15
-Related: Engine/Sequencer/ReferenceEditor
-Related: Engine/Sequencer/HowTo
-tags:Sequencer
-type:overview
-tags:Getting Started
-SkillLevel:Intermediate
+Type: Overview
+SkillLevel:Beginner
+Version: 4.17
+Parent:Engine/Sequencer
+Order: 1
+Tags:Sequencer
+topic-image:Engine/Physics/physics_topic.png
 
 [TOC(start:2 end:2)]
 
-The **Sequencer** Editor gives users the ability to create in-game cinematics with a specialized multi-track editor (similar to [Matinee](Engine\Matinee\)). 
-By creating **Level Sequences** and adding **Tracks**, users can define the make-up of each Track, which will determine the content for the scene.
-Tracks can consist of things like Animations (for animating a character), Transformations (moving things around in the scene), Audio (for including music or sound effects), and several other Track types.  
+The **Sequencer** Editor gives users the ability to create in-game cinematics with a specialized multi-track editor (similar to [Matinee](Engine\Matinee\)). By creating **Level Sequences** and adding **Tracks**, users can define the make-up of each Track, which will determine the content for the scene. Tracks can consist of things like Animations (for animating a character), Transformations (moving things around in the scene), Audio (for including music or sound effects), and several other Track types.
 
 ## Creating Level Sequences
 
-The **Level Sequence** is the "container" for your cinematic scenes, and must be created in order to begin working inside of the Sequencer Editor. You can create a Level Sequence and directly place it inside of your level by choosing to create one from the **Toolbar** under **Cinematics** (as seen below). 
+The **Level Sequence** is the "container" for your cinematic scenes, and must be created in order to begin working inside of the Sequencer Editor. You can create a Level Sequence directly in your level from the **Toolbar** under **Cinematics** (as seen below).
 
 ![](LevelSequence.png)
 
-This will add it to the Level, at which point it can be selected, and its properties can be manipulated in the **Details** panel (similar to a **Matinee Actor**). In the Details panel (below), you can define whether the Level Sequence will automatically play upon Level start, the Play Rate for the sequence, whether the sequence should loop, and other settings. 
+This will add it to the Level, at which point it can be selected, and its properties can be manipulated in the **Details** panel (similar to a **Matinee Actor**). In the Details panel (below), you can define whether the Level Sequence will automatically play upon Level start, whether the sequence should loop, the Play Rate for the sequence, and other settings.
 
 ![](DetailsPanel.png)
 
-How this differs from Matinee, is that Level Sequences are self-contained assets, and you can embed a Level Sequence within another Level Sequence.
-For example, you can create a Level Sequence that has animated characters and cameras as one scene that is part of a larger cinematic sequence.
+Unlike Matinee, Level Sequences are self-contained assets, and you can embed a Level Sequence within another Level Sequence. For example, you can create a Level Sequence that has animated characters and cameras as one scene that is part of a larger cinematic sequence.
 
 An alternate method of creating Level Sequences can be performed in the **Content Browser** by clicking the **Add New** button, and selecting **Level Sequence** from the **Animation** menu. When doing this, you are creating the Level Sequence asset before placing it in a Level.
 
@@ -38,7 +32,7 @@ An alternate method of creating Level Sequences can be performed in the **Conten
 
 ## Adding Tracks to Sequencer
 
-After creating a Level Sequence, double-click on it to open the Sequencer Editor so that you can begin creating your cinematic. 
+After creating a Level Sequence, double-click on it to open the **Sequencer Editor** so that you can begin creating your cinematic. 
 
 ![](BlankScene.png)
 
@@ -46,14 +40,14 @@ After creating a Level Sequence, double-click on it to open the Sequencer Editor
 Above, we have a newly created, empty Level Sequence.
 [/REGION]
 
-The first thing you will need to do is add a **Track** type, which you can do from the **Add** button's drop-down menu. 
+The first thing you will need to do is add a **Track** type, which you can do from the **Track** button's drop-down menu. 
 
 ![](AddToSequencer.png)
 
-From the drop-down menu, you will see several Track types you can select from, as the ability to add an **Actor To Sequencer**. The Actor To Sequencer option will enable you to add any Actors that you have selected in your Level (at the time) to Sequencer so that you can manipulate them during your scene.
+From the drop-down menu, you will see several Track types you can select from, as well as the ability to add an **Actor To Sequencer**. The Actor To Sequencer option will enable you to add any Actors that you have selected in your Level (at the time) to Sequencer so that you can manipulate them during your scene.
 
-Typically, if you are creating a cinematic that has characters, animals, creatures (or anything along those lines) that will animate and move, you will have a [Skeletal Mesh](Engine/Content/Types/SkeletalMeshes) for it, which you will need to add to Sequencer.
-For example, we have a Skeletal Mesh of a bear that we have placed in our Level (below). With the bear selected, we can then click the **Add** button in Sequencer and choose **Actor To Sequencer** so that we can add it and control it in the Sequencer Editor.
+Typically, if you are creating a cinematic that has characters, animals, creatures, or anything along those lines that will animate and move, you will have a [Skeletal Mesh](Engine/Content/Types/SkeletalMeshes) for it, which you will need to add to Sequencer.
+For example, we have a Skeletal Mesh of a bear that we have placed in our Level (below). With the bear selected, we can then click the **Track** button in Sequencer and choose **Actor To Sequencer** so that we can add it and control it in the Sequencer Editor.
 
 ![](AddBear.png)(w:640)
 
@@ -69,7 +63,7 @@ Below, we select the Animation Sub-Track and assign an Animation for our Skeleta
 
 [OBJECT:EmbeddedVideo]
 	[PARAMLITERAL:videoid]
-	KoksbI7ZOTM
+	RTAbKVuwgHM
 	[/PARAMLITERAL]
 	[PARAMLITERAL:width]
 	100
@@ -85,21 +79,23 @@ Below, we select the Animation Sub-Track and assign an Animation for our Skeleta
 	[/PARAMLITERAL]
 [/OBJECT]
 
-## Possessables vs. Spawnables
-
-Sequencer is very similar to Matinee, in that it uses the concept of "possessables", meaning that an Actor exists in the Level, and Sequencer will take possession of it in order to apply changes you make to it via Sequencer. 
-As illustrated in our bear example (above), you can place a Skeletal Mesh in the Level and add it to Sequencer, then inside Sequencer, you can add an Animation Track associated with the Skeletal Mesh and assign different animations for it to play during your scene. 
-In this example, we are "possessing" the Skeletal Mesh that exists in the Level in order to tell it what to do. 
+### Sequencer World Outliner Integration
 
-Sequencer also includes another form of manipulation, which is referred to as "spawnables", meaning that the object we are possessing does not yet exist, so Sequencer will "spawn" the object (when told to), and will have authority over it to determine its lifecycle. Since Sequencer spawns and owns the object, the object is not bound to a particular Level. Anything marked as "spawnable" can be used in any Level, so you could create a scene and have it take place in any environment, and reuse it in new Levels without having to recreate it. 
+When adding Actors to a Level Sequence, the [World Outliner](Engine/UI/LevelEditor/SceneOutliner/Interface) will update to show which Level Sequences the Actor is bound to. 
 
-[REGION:note]
-Please see [Creating Spawnables](Engine\Sequencer\HowTo\Spawnables\) for more information. 
+[REGION:raw]
+![](Outliner.png)
 [/REGION]
 
+The World Outliner will also indicate which Actors have been spawned by Sequencer indicated by the orange lighting bolt icon.  
+
+![](SpawnableWorldOutliner.png)
+
+
+
 ## Keyframing Scenes
 
-If you are familiar with most animation tools, the concept of using Keyframes to drive content may be familiar to you. 
+If you are familiar with most animation tools, the concept of using **Keyframes** to drive content may be familiar to you. 
 Sequencer allows you to manipulate your content by adding **Keys** (referred to as "Keyframing") with defined properties at desired points along a timeline. 
 Upon reaching those Keys in the Timeline, the properties you have defined at each Keyframe are updated to reflect the values that you entered.
 
@@ -107,7 +103,7 @@ A sample scene (below) illustrates the process of adding movement by Keyframing
 
 [OBJECT:EmbeddedVideo]
 	[PARAMLITERAL:videoid]
-	dVDNJkQxlRk
+	1sve4WgQgK0
 	[/PARAMLITERAL]
 	[PARAMLITERAL:width]
 	100
@@ -124,14 +120,14 @@ A sample scene (below) illustrates the process of adding movement by Keyframing
 [/OBJECT]
 
 [REGION:caption]
-For the **Location** track (above), we added a Keyframe for the starting position, and another Keyframe for the ending position.
+For the **Location** track (above), we added a Keyframe for the starting position and another Keyframe for the ending position.
 [/REGION]
 
-You can add a Keyframe by selecting a track and pressing **Enter**, or by clicking the "add Keyframe" button on each track. 
+You can add a Keyframe by selecting a track and pressing **Enter**, or by clicking the **Add Keyframe** button on each track. 
 
 ![](KeyframeButton.png)
 
-Some properties in the **Details** panel of an Actor can be Keyframed directly by clicking the "add Keyframe" button next to each property. 
+Some properties in the **Details** panel of an Actor can be Keyframed directly by clicking the **Add Keyframe** button next to each property. 
 
 ![](KeyframeDetails.png)
 
@@ -149,11 +145,11 @@ With a Key (or with multiple Keys or sections) selected, you can use the **Trans
 
 You can also open the Transform Keys/Selection tools by pressing the **Ctrl+M** Key combination. 
 
-This will allow you to offset the Keys/selection by a specified amount, or the amount to scale the Keys/section by a specified time. 
+This will enable you to offset the Keys/selection by a specified amount, or the amount to scale the Keys/section by a specified time. 
 
 [OBJECT:EmbeddedVideo]
 	[PARAMLITERAL:videoid]
-	Kd3CwzMwTLs
+	T27ngPC83Iw
 	[/PARAMLITERAL]
 	[PARAMLITERAL:width]
 	100
@@ -169,7 +165,40 @@ This will allow you to offset the Keys/selection by a specified amount, or the a
 	[/PARAMLITERAL]
 [/OBJECT]
 
-### Playback Range
+### Sequencer Editing Options
+
+There are three editing options you can use when a Level Sequence is open that defines how property values are handle by the Sequencer editor. 
+
+![](SequenceEdits.png)
+
+* **Allow Any Edits** - allows any edits to occur, some of which may produce tracks/keys or modify default properites. 
+* **Allow Sequencer Edits Only** - any edits to properties will automatically create a track if it doesn't exist and/or a keyframe. 
+	* These changes will be isolated to the particular shot you are working on and won't leak into others. 
+* **Allow Level Edits Only** - when using this option, you will be editing the default properties objects directly. 
+	* Properties that already have tracks in Sequencer will be disabled in the Details panel when you want to make global changes that affect all shots. 
+
+Below is an example of each setting in action: 
+
+[OBJECT:EmbeddedVideo]
+	[PARAMLITERAL:videoid]
+	RL94KEjiddE
+	[/PARAMLITERAL]
+	[PARAMLITERAL:width]
+	100
+	[/PARAMLITERAL]
+	[PARAMLITERAL:height]
+	56.25
+	[/PARAMLITERAL]
+	[PARAMLITERAL:units]
+	%
+	[/PARAMLITERAL]
+	[PARAMLITERAL:type]
+	simple
+	[/PARAMLITERAL]
+[/OBJECT]
+
+
+## Playback Range
 
 Sequencer will play back your scenes based on the defined Start and Stop playback ranges (green/red markers below). 
 
@@ -189,16 +218,45 @@ What is happening in this example is that our `Shot0020_001` is set to use 50 fr
 
 ![](Evaluation_1.png)
 
-If we dive into our Shot, the first red marker denotes the end of the Shot at frame 50, while the second red marker denotes the end of the Shot at the master Level. 
+If we dive into our Shot, the first red marker denotes the end of the Shot at frame 50, while the second red marker denotes the end of the Shot at the Master Level. 
 
 ![](Evaluation_2.png)
 
 Frames 50 through 100 are grayed out to indicate that they are not being considered for evaluation. 
 To address the difference in length, we could increase the amount of frames being evaluated at the Shot Level to 100, or at the Master Sequence Level, we could reduce the length of the Shot to 50 frames. 
 
+## Hierarchical Bias
+
+By default, tracks at a lower level in the Level Sequence hierarchy take precedence. This allows filmmakers to build a pipeline they are accustomed to, where adjustments at the shot level override the sequence they are contained in. 
+For example, in the video below, we have a sample scene with three camera cuts and a light pointed at, and highlighting, our actor. 
+Our light is set to white at the Master level. However, when we dive in and add the light to our second shot, we can make changes to it, overriding the Master level setting.
+
+[OBJECT:EmbeddedVideo]
+	[PARAMLITERAL:videoid]
+	294aoFPGPIs
+	[/PARAMLITERAL]
+	[PARAMLITERAL:width]
+	100
+	[/PARAMLITERAL]
+	[PARAMLITERAL:height]
+	56.25
+	[/PARAMLITERAL]
+	[PARAMLITERAL:units]
+	%
+	[/PARAMLITERAL]
+	[PARAMLITERAL:type]
+	simple
+	[/PARAMLITERAL]
+[/OBJECT]
+
+You can also set the amount of influence by right-clicking on a shot, then increasing the **Hierarchical Bias** value (higher values will take precedence) under **Properties**.
+
+![](HierarchicalBias.png)
+
+
 ##Specialized Viewports
 
-Sequencer allows you to use specialized viewports to make the editing process easier. 
+Sequencer uses specialized viewports to make the editing process easier.
 
 ![](SpecialViewports.png)
 
@@ -218,7 +276,7 @@ The **Camera Rig Crane** Actor can be used to simulate crane-like camera movemen
 
 ![](Crane.png)(w:480)
 
-You can attach a Camera to the Camera Rig Crane Actor, which enables you to manipulate the crane's movement through the **Details** panel and **Crane Controls** values. Those values can affect the Crane Pitch, Crane Yaw, and Crane Arm Length. All of these values can be Keyframed within Sequencer, which will allow you to adjust them during your cinematics (as desired). 
+You can attach a Camera to the Camera Rig Crane Actor, which enables you to manipulate the crane's movement through the **Details** panel and **Crane Controls** values. Those values can affect the Crane Pitch, Crane Yaw, and Crane Arm Length. All of these values can be Keyframed within Sequencer, which will enable you to adjust them as desired during your cinematics.
 
 [REGION:note]
 Please see [Shooting from a Camera Rig Crane](Engine\Sequencer\HowTo\CameraRigCrane) for more information. 
@@ -274,8 +332,6 @@ One tool you can use to narrow down your content creation time is to use the **S
 	[/PARAMLITERAL]
 [/OBJECT]
 
-
-
 To add a new recording, pick an Actor to record, and click Record to start capturing. When stopping the recording, new Level Sequence assets will be created. You can then incorporate them into your existing cinematic, or you can use pieces from the recorded Level Sequence to create new sequences.
 You can also record audio to go along with your recorded sequence from an external microphone, which will allow you to narrate your sequence (or provide instructions while recording Level Editor actions).
 
@@ -318,7 +374,7 @@ Please see [](Engine/Sequencer/HowTo/ImportExportEDLs) for a step-by-step guide
 ### Custom Render Passes
 
 If you are looking to export your cinematic in different render passes, you can do so from the **Render Movie Settings** window. 
-This will allow you to specify the render pass(es) to use when exporting your sequence. You can also export HDR (High Dynamic-Range) data as .exr files, and define compression and color gamut settings. 
+These settings specify the render pass(es) to use when exporting your sequence. You can also export HDR (High Dynamic-Range) data as .exr files, and define compression and color gamut settings. 
 
 
 [REGION:note]
@@ -359,6 +415,17 @@ Please see [](Engine\Sequencer\HowTo\ActorSequence) for a step-by-step guide.
 [/REGION]
 
 
+## Possessables vs. Spawnables
+
+Sequencer is very similar to Matinee, in that it uses the concept of "possessables", meaning that an Actor exists in the Level, and Sequencer will take possession of it in order to apply your changes to it. For example, you can place a Skeletal Mesh in the Level and add it to Sequencer, then inside Sequencer, you can add an Animation Track associated with the Skeletal Mesh and assign different animations for it to play during your scene. In this example, we are possessing the Skeletal Mesh that exists in the Level in order to tell it what to do. 
+
+Sequencer also includes another form of manipulation, which is referred to as "spawnables", meaning that the object we are possessing does not yet exist and will be spawned by Sequencer. 
+In this case, Sequencer will have authority over the spawned Actor and can determine its lifecycle. Since the object is spawned, it is not bound to a particular level.
+Anything marked as spawnable can be used in any Level, so you could create a scene and have it take place in any environment, and reuse that scene in new Levels without having to recreate it. 
+
+[REGION:note]
+Please see [Creating Spawnables](Engine\Sequencer\HowTo\Spawnables\) for more information. 
+[/REGION]
 
 
 ## Workflow Considerations
