Copying //UE4/Dev-Documentation to //UE4/Samples-Main (Source: //UE4/Dev-Documentation @ 3623665)

diff --git a/Engine/Documentation/Source/Engine/Sequencer/ReferenceEditor/EditorReference.INT.udn b/Engine/Documentation/Source/Engine/Sequencer/ReferenceEditor/EditorReference.INT.udn
--- a/Engine/Documentation/Source/Engine/Sequencer/ReferenceEditor/EditorReference.INT.udn
+++ b/Engine/Documentation/Source/Engine/Sequencer/ReferenceEditor/EditorReference.INT.udn
@@ -42,10 +42,15 @@ What follows is a description of each of the toolbar buttons, from left to right
 | ![](tb_1.png) | Saves the current Level Sequence. |
 | ![](tb_2.png) | Saves the current Level Sequence under a different name. |
 | ![](SearchIcon.png) | Find the currently viewed sequence asset in the Content Browser. |
+| ![](cameraicon.png) | Creates a new camera and set it as the current camera cut. |
 | ![](tb_4.png) | Opens the Render Movie Settings window which allows you to output your sequence as a movie or image sequence. |
+| ![](UndoIcon.png) | Restore any objects that have been animated by sequencer back to their original state (Ctrl+R). |
 | ![](tb_5.png) | Opens the General Options menu with settings to define functionality for Sequencer.  |
+| ![](PlaybackIcon.png) | Opens the Playback Options menu. |
+| ![](SelectEditOptions.png) | Opens the Select/Edit Options menu. |
 | ![](tb_6.png) | Enables/Disables the Key All option. With this and auto-key enabled, determines if all properties are keyed for a track when changes are made or if only the changed property is keyed. |
 | ![](tb_7.png) | Enables/Disables different Auto-Key options. By default, this option is disabled but can be set to Auto-key All or Auto-key Animated. |
+| ![](SequenceEditOptions.png) | Opens the Sequence Edit Options menu.  |
 | ![](tb_8.png) | Enables/Disables grid snapping for the Tracks Area as well as grid snapping options. |
 | ![](tb_9.png) | Allows for adjustment of the time snapping internval. |
 | ![](tb_10.png) | Opens the [**Curve Editor**](Engine/UI/CurveEditor) which allows you to fine tune the animation keys placed in the Tracks Area. |
@@ -55,21 +60,21 @@ What follows is a description of each of the toolbar buttons, from left to right
 Depending on your workflow, you may choose to create several **Level Sequences** as scenes and place each of those scenes in a new Level Sequence which is your "master" Level Sequence which plays those scenes back-to-back as part of a **Shots** track. 
 In these instances, you can jump from your master Level Sequence directly into each scene by **Double-clicking** on them inside the Tracks Area. 
 
-In the image below, we have a Level Sequence called _shot_06_ inside another Level Sequence called _Master_.
+In the image below, we have a Level Sequence inside another Level Sequence which is our "master" sequence.
 
 ![](SequenceBreadcrumbs1.png)
 
-We can click the _Master_ breadcrumb to return to our master Level Sequence containing our _shot_06_ Level Sequence and other sequences. 
+We can click the breadcrumbs at the top to return to our master Level Sequence containing all of our shots.
 
 ![](SequenceBreadcrumbs2.png)
 
 ## Add / Filter Tracks
 
-The **Add** button allows you to add **Tracks** or [**Actors**](Engine/Actors) from your level to **Sequencer** which you can then use to craft your cinematic.
+The **Track** button allows you to add **Tracks** or [**Actors**](Engine/Actors) from your level to **Sequencer** which you can then use to craft your cinematic.
 
 ![](AddOption.png)
 
-With an Actor in the level selected, clicking the **Add** button will allow you to select it quickly from the menu. If you do not have an Actor selected, you can select one from the **Choose Actor** section. 
+With an Actor in the level selected, clicking the **Track** button will allow you to select it quickly from the menu. If you do not have an Actor selected, you can select one from the **Choose Actor** section. 
 
 
 Below are the available Tracks you can add to Sequencer, click each link below to view a How-to page for an example using each Track.
@@ -83,9 +88,10 @@ Below are the available Tracks you can add to Sequencer, click each link below t
 | [**Camera Cut Track**](Engine\Sequencer\HowTo\TracksCameraCut) | Allows you to switch between cameras associated with Sequencer and provide a camera persepctive. |
 | [**Shot Track**](Engine\Sequencer\HowTo\TracksShot) | Allows you to add Level Sequences as "shots" which can have **Takes** that can be swapped in/out. |
 | [**Play Rate Track**](Engine\Sequencer\HowTo\TracksPlayRate) | Adds a track which can be used to manipulate the playback rate of the sequence. |
-| [**Sub Track**](Engine\Sequencer\HowTo\TracksSub) | Adds a sub track which can contain other Level Sequences. This is useful for creating multiple copies of a sequence that different users can work from without altering the original sequence. |
+| [**Subscenes Track**](Engine\Sequencer\HowTo\TracksSub) | Adds a sub track which can contain other Level Sequences. This is useful for creating multiple copies of a sequence that different users can work from without altering the original sequence. |
 | [**Fade Track**](Engine\Sequencer\HowTo\TracksFade) | Adds a track that will allow you to control the fading in/out of the sequence.  |
 | [**Level Visibility Track**](Engine\Sequencer\HowTo\TracksLevelVisibility) | This can be used to keyframe and control the visibilty of your level. |
+| [**Material Parameter Collection Track**](Engine/Sequencer/HowTo/TracksMaterialParameter) | This track can be used to drive Material Parameter Collections within Sequencer. |
 
 The **Filter** option allows you to filter the **Tree View** for only Actors that have been added to the sequence from the level.
 
@@ -102,7 +108,7 @@ You can **Right-click** on the Timeline to define new Start/Stop points from the
 ![](Timeline2.png)
 
 [REGION:tip]
-Your sequence will continue to play from the start marker all the way to the end marker even if your content does not full reach the end of the sequence (or exceeds past the end marker). Inside the **General Options**, there is an option to keep your content inside the section bounds which will ensure there is no dead space during your sequence. 
+Your sequence will continue to play from the start marker all the way to the end marker even if your content does not full reach the end of the sequence (or exceeds past the end marker). Inside the **Playback Options**, there is an option to keep your content inside the section bounds which will ensure there is no dead space during your sequence. 
 
 ![](KeepInSectionBounds.png)
 [/REGION]
diff --git a/Engine/Documentation/Source/Engine/Tools/CrashReport/Setup/SettingUpCrashReport.INT.udn b/Engine/Documentation/Source/Engine/Tools/CrashReport/Setup/SettingUpCrashReport.INT.udn
--- a/Engine/Documentation/Source/Engine/Tools/CrashReport/Setup/SettingUpCrashReport.INT.udn
+++ b/Engine/Documentation/Source/Engine/Tools/CrashReport/Setup/SettingUpCrashReport.INT.udn
@@ -80,10 +80,10 @@ There are 3 websites we will be creating for this service to work
 * **Right-click** on Sites -> Add Web Site
     * Site Name: AutoReportService
     * Application Pool: AutoReportService
-    * Physical Path: <Web Root>/AutoReportService
+    * Physical Path: `<Web Root>/AutoReportService`
 * In the File System, create a Folder for your Log file outside of the web root
     * e.g. C:\AutoReportLog\
-* Share <Web Root>/AutoReportService to the local network to allow for simple deployment
+* Share `<Web Root>/AutoReportService` to the local network to allow for simple deployment
     * From IIS, you can simply **Right-click** on the Website name -> Explore
     * **Right-click** on the Folder and click Properties
         * Set the Permissions as you see fit, allowing at least the person deploying the site access
@@ -121,7 +121,7 @@ Open The CrashReport solution in VS 2008 located at /UnrealEngine3/Development/T
 ### http://CrashReport.yourdomain.com
 
 * Edit Web.config
-    * Find the line under <connectionStrings> starting with `<add name="CrashReportConnectionString"`
+    * Find the line under `<connectionStrings>` starting with `<add name="CrashReportConnectionString"`
     * Edit Data Source, User ID and Password to your Connection Settings
 * Edit Settings.settings ( or Properties->settings)
     * Note the url you made for `http://AutoReportFiles.yourdomain.com` and enter that url here.
diff --git a/Engine/Documentation/Source/Engine/Tools/CrashReporter/CrashReporter.INT.udn b/Engine/Documentation/Source/Engine/Tools/CrashReporter/CrashReporter.INT.udn
--- a/Engine/Documentation/Source/Engine/Tools/CrashReporter/CrashReporter.INT.udn
+++ b/Engine/Documentation/Source/Engine/Tools/CrashReporter/CrashReporter.INT.udn
@@ -43,9 +43,9 @@ server. This new style of Windows Error Reporting was added in Windows Vista, so
 
 
 This application searches for all relevant crashes in the past 'UploadHistory' days, and uploads them the the 'CrashReportUploadURL'. The application is stateless in that it does not require any direct input from the engine; it can be run as a standalone application and
-will discover and upload just the same. If the upload fails for any reason (such as the URL being unavailable), the next launch of the uploader will retry. No parameters required, but '-unattended' and '-logfolder=<folder name>' as commandline options are respected. 
+will discover and upload just the same. If the upload fails for any reason (such as the URL being unavailable), the next launch of the uploader will retry. No parameters required, but `-unattended` and `-logfolder=<folder name>` as commandline options are respected. 
 
-Crash reports are discovered in by searching for any folders matching the pattern 'SpecialFolder.LocalApplicationData\Microsoft\Windows\WER\ReportQueue\AppCrash_UE4-*'. These folders have a Guid that is used to see if the report is already uploaded.
+Crash reports are discovered in by searching for any folders matching the pattern `SpecialFolder.LocalApplicationData\Microsoft\Windows\WER\ReportQueue\AppCrash_UE4-*`. These folders have a Guid that is used to see if the report is already uploaded.
 
 To account for locally compiled binaries, the uploader checks to see if it has a pdb, and is running on the local domain. If both these conditions are met, MinidumpDiagnostics.exe will be run to generate a report (which includes a callstack). This is then uploaded as part
 of the report.
diff --git a/Engine/Documentation/Source/Engine/UI/ClassViewer/ClassViewer.INT.udn b/Engine/Documentation/Source/Engine/UI/ClassViewer/ClassViewer.INT.udn
--- a/Engine/Documentation/Source/Engine/UI/ClassViewer/ClassViewer.INT.udn
+++ b/Engine/Documentation/Source/Engine/UI/ClassViewer/ClassViewer.INT.udn
@@ -105,9 +105,39 @@ The context menu is opened by clicking on ![Context Menu Button](ContextMenuButt
 
 |Class Type|Options|Description|
 |----------|-------|-----------|
-|**Blueprint Class**|**Create Blueprint...**<br>**Edit Blueprint...**<br>**Find In Content Browser**|Creates a new Blueprint parented to the selected Blueprint.<br>Opens the selected Blueprint in the [Blueprint Editor](Engine/Blueprints).<br>Finds the Blueprint Actor in the [Content Browser](Engine/Content/Browser).|
-|**Code Class**|**Create Blueprint...**<br>**Open C++ Header File...**<br>**Create New C++ Class...**|Creates a new Blueprint parented to the selected Blueprint.<br>Opens up the class header file in Visual Studio.<br>Opens the [](Programming\Development\ManagingGameCode\CppClassWizard) to create a new class with the selected class as the Parent class.|
-
+|**Blueprint Class**|[INCLUDE:#bpclassoptions]|[INCLUDE:#bpclassoptiondesc]|
+|**Code Class**|[INCLUDE:#codeclassoptions]|[INCLUDE:#codeclassoptiondesc]|
+
+<!--
+[EXCERPT:BPClassOptions]
+[REGION:plain_list]
+* **Create Blueprint...**
+* **Edit Blueprint...**
+* **Find In Content Browser**
+[/REGION]
+[/EXCERPT:BPClassOptions]
+[EXCERPT:BPClassOptionDesc]
+[REGION:plain_list]
+* Creates a new Blueprint parented to the selected Blueprint.
+* Opens the selected Blueprint in the [Blueprint Editor](Engine/Blueprints).
+* Finds the Blueprint Actor in the [Content Browser](Engine/Content/Browser).
+[/REGION]
+[/EXCERPT:BPClassOptionDesc]
+[EXCERPT:CodeClassOptions]
+[REGION:plain_list]
+* **Create Blueprint...**
+* **Open C++ Header File...**
+* **Create New C++ Class...**
+[/REGION]
+[/EXCERPT:CodeClassOptions]
+[EXCERPT:CodeClassOptionDesc]
+[REGION:plain_list]
+* Creates a new Blueprint parented to the selected Blueprint.
+* Opens up the class header file in Visual Studio.
+* Opens the [](Programming\Development\ManagingGameCode\CppClassWizard) to create a new class with the selected class as the Parent class.
+[/REGION]
+[/EXCERPT:CodeClassOptionDesc]
+-->
 
 ## Drag & Dropping
 Dropping a Blueprint class Actor into a viewport will add the Actor to the scene. When successfully dragging a Class, the mouse will show what class is in hand. Only one class can be dragged at a time.
diff --git a/Engine/Documentation/Source/Engine/UI/SourceControl/SourceControl.INT.udn b/Engine/Documentation/Source/Engine/UI/SourceControl/SourceControl.INT.udn
--- a/Engine/Documentation/Source/Engine/UI/SourceControl/SourceControl.INT.udn
+++ b/Engine/Documentation/Source/Engine/UI/SourceControl/SourceControl.INT.udn
@@ -57,10 +57,9 @@ The **Content Browser** will display special icons on the upper-right corner of
 
 
 [REGION:imagetable]
-| ![CheckedOutByYou.png](CheckedOutByYou.png)(w:178) | ![CheckedOut.png](CheckedOut.png)(w:178) | ![MarkedForAdd.png](MarkedForAdd.png)(w:178) | ![NotInDepot.png](NotInDepot.png)(w:178) | ![NotHeadRevision.png](NotHeadRevision.png)(w:178) 
+| ![CheckedOutByYou.png](CheckedOutByYou.png)(w:178) | ![CheckedOut.png](CheckedOut.png)(w:178) | ![MarkedForAdd.png](MarkedForAdd.png)(w:178) | ![NotInTheDepot.png](NotInTheDepot.png)(w:178) | ![NotHeadRevision.png](NotHeadRevision.png)(w:178) 
 | --- | --- | --- | --- | --- |
 | Checked out by you | Checked out by another user | Marked for add | Not in depot | Newer version of file exists in source control
-
 [/REGION]
 
 ## Source Control Actions
diff --git a/Engine/Documentation/Source/Engine/UMG/UserGuide/Anchors/Anchors.INT.udn b/Engine/Documentation/Source/Engine/UMG/UserGuide/Anchors/Anchors.INT.udn
--- a/Engine/Documentation/Source/Engine/UMG/UserGuide/Anchors/Anchors.INT.udn
+++ b/Engine/Documentation/Source/Engine/UMG/UserGuide/Anchors/Anchors.INT.udn
@@ -1,12 +1,17 @@
-Availability:Public
+Availability: Public
+Crumbs: %ROOT%
 Title: Anchors
-Crumbs:%ROOT%, Engine, Engine/UMG, Engine/UMG/UserGuide
-Description:Using Anchors to keep UI Widgets in fixed locations based on different aspect ratios. 
+Description: Using Anchors to keep UI Widgets in fixed locations based on different aspect ratios. 
+Type: Overview
+SkillLevel:Beginner
+Version:4.17
+Parent:Engine/UMG/UserGuide
 Related: Engine/UMG
 Related: Engine/UMG/UserGuide/DPIScaling
 Related: Engine/Blueprints
-version: 4.9
-tags:UMG UI Designer
+Order:1
+Tags:UMG UI Designer
+Topic-Image:Engine/UMG/UserGuide/AnchorsTopic.png
 
 
 [TOC(start:2 end:2)]
diff --git a/Engine/Documentation/Source/Engine/UMG/UserGuide/Animation/Animation.INT.udn b/Engine/Documentation/Source/Engine/UMG/UserGuide/Animation/Animation.INT.udn
--- a/Engine/Documentation/Source/Engine/UMG/UserGuide/Animation/Animation.INT.udn
+++ b/Engine/Documentation/Source/Engine/UMG/UserGuide/Animation/Animation.INT.udn
@@ -1,9 +1,14 @@
-Availability:Public
+Availability: Public
+Crumbs: %ROOT%
 Title: Animation
-Crumbs:%ROOT%, Engine, Engine/UMG, Engine/UMG/UserGuide
-Description:How to create animated UI elements in UMG. 
-version: 4.9
-tags:UMG UI Designer
+Description: How to create aniamted UI elements in UMG. 
+Type: Overview
+SkillLevel:Beginner
+Version:4.17
+Parent:Engine/UMG/UserGuide
+Order:2
+Tags:UMG UI Designer
+Topic-Image:Engine/UMG/UserGuide/AnimationTopic.png
 
 
 [REGION:fullwidth raw]
diff --git a/Engine/Documentation/Source/Engine/UMG/UserGuide/BestPractices/UMGBestPractices.INT.udn b/Engine/Documentation/Source/Engine/UMG/UserGuide/BestPractices/UMGBestPractices.INT.udn
--- a/Engine/Documentation/Source/Engine/UMG/UserGuide/BestPractices/UMGBestPractices.INT.udn
+++ b/Engine/Documentation/Source/Engine/UMG/UserGuide/BestPractices/UMGBestPractices.INT.udn
@@ -1,46 +1,17 @@
-Availability:Public
-Title:UMG Best Practices
-Crumbs:%ROOT%, Engine\UMG\UserGuide
-Description:Tips and Tricks for working with Unreal Motion Graphics.
-version: 4.10
-parent:Engine/UMG	
-order:2
-type:how-to
-tags:UMG UI Designer
-
-
-[VAR:Topic]
-[OBJECT:Topic]
-	[PARAM:image]
-	![](Engine\UMG\UserGuide\BestPractices\BestPractices_topic.png)
-	[/PARAM]
-	[PARAM:title]
-		%Engine\UMG\UserGuide\BestPractices:title%
-	[/PARAM]
-	[PARAM:description]
-		%Engine\UMG\UserGuide\BestPractices:description%
-	[/PARAM]
-	[PARAM:path]
-	[RELATIVE:Engine\UMG\UserGuide\BestPractices]
-	[/PARAM]
-[/OBJECT]
-[/VAR]
-[VAR:TopicCompact]
-[OBJECT:TopicCompact]
-	[PARAM:image]
-	![](Engine\UMG\UserGuide\BestPractices\BestPractices_topic.png)
-	[/PARAM]
-	[PARAM:title]
-		%Engine\UMG\UserGuide\BestPractices:title%
-	[/PARAM]
-	[PARAM:description]
-		%Engine\UMG\UserGuide\BestPractices:description%
-	[/PARAM]
-	[PARAM:path]
-	[RELATIVE:Engine\UMG\UserGuide\BestPractices]
-	[/PARAM]
-[/OBJECT]
-[/VAR]
+Availability: Public
+Crumbs: %ROOT%
+Title: UMG Best Practices
+Description:Tips and Tricks for working with Unreal Motion Graphics. 
+Type:how-to
+SkillLevel:Beginner
+Version:4.17
+Parent:Engine/UMG/UserGuide
+Order:4
+Tags:UMG UI Designer
+Tags:Getting Started
+Topic-Image:BestPractices_topic.png
+
+
 
 [TOC(start:2 end:2)]
 
diff --git a/Engine/Documentation/Source/Engine/UMG/UserGuide/Clipping/Clipping.INT.udn b/Engine/Documentation/Source/Engine/UMG/UserGuide/Clipping/Clipping.INT.udn
new file mode 100644
--- /dev/null
+++ b/Engine/Documentation/Source/Engine/UMG/UserGuide/Clipping/Clipping.INT.udn
@@ -0,0 +1,95 @@
+Availability:Public
+Crumbs:%ROOT%
+Title: Clipping
+Description: An overview of using the Clipping properties within the UMG UI Designer.
+Type: Reference
+Version:4.17
+Parent:Engine/UMG/UserGuide
+Related:Programming\Slate\ClippingSystem
+Order:3
+Tags:UMG UI Designer
+Topic-image:Clipping_Topic.png
+
+
+[VAR:spacer2]![](%ROOT%/spacer.png)(w:30 convert:false)[/VAR]
+
+
+[REGION:warning]
+With Unreal Engine 4.17, the clipping system has been overhauled and in some cases may break backward compatibility. For additional information about these changes see the 
+[Slate Clipping System](Programming\Slate\ClippingSystem) details here. 
+[/REGION]
+
+
+The clipping system in UMG uses [Slate's Clipping System](Programming/Slate/ClippingSystem) as a framework to control how text, images, or content is shown for Widgets (as well as the rest of the
+Editor). **Clipping** works by restricting rendered objects (graphics and text) to a region using a bounding box so that anything outside of it is not shown. The clipping system is now axis-aligned
+meaning that it can clip any rotation, which was not possible before because of the way transforms were handled. 
+
+![](Programming\Slate\ClippingSystem\ClippingExamples.png)(w:500)
+
+In this example, each of the buttons is a parent to the displayed text. These examples demonstrate whether the button or the text is responsible for clipping.  
+
+* Left - No clipping is enabled on the button or the text.
+* Center - Clipping is enabled on the text
+* Right - Clipping is enabled on the button.
+
+The **Canvas** panel (or clipping zone), outlined in blue, represents your game screen and will clip (or not draw) anything outside of this area for your game. 
+
+[REGION:fullwidth]
+![](UMGDesigner.png)
+[/REGION]
+[REGION:caption] In the UMG Designer Graph, the Canvas panel (blue) represents the clipping zone for your game screen. [/REGION]
+
+In Unreal Engine 4.16 and earlier, clipping for Widgets was handled using layout-space which prevented anything from rendering outside of the Canvas panel. So if part of a Widget's bounding 
+box was outside of the Canvas panel it would no longer rendered, even if the Widget were rotated, its bounding box would not be leaving parts of the graphic or text cut off even if it's inside
+the Canvas panel.
+
+Take for instance these examples showing a comparison before and after the changes:
+
+[REGION:imagetable]
+| --- | --- | 
+| ![](Clipping1.png)(w:500) | [INCLUDE:#comparison] |
+[/REGION]
+
+
+
+<!----
+[INCLUDE:#comparison]
+
+[EXCERPT:comparison]
+[OBJECT:ComparisonSlider]
+	[PARAM:before]
+	![Render Transform clipping | in 4.16 and earlier.](Clipping2.png)(w:500)
+	[/PARAM]
+	[PARAM:after]
+	![Render Transform clipping | in 4.17 and later.](Clipping3.png)(w:500)
+	[/PARAM]
+[/OBJECT]
+[/EXCERPT:comparison]
+
+--->
+
+The Image Widget (left) has been placed at the edge of the Canvas panel so that part of its bounding box is outside. Now that the clipping system is axis-aligned, instead of using layout-space,
+clipping which resolves artifacts and issues like the comparison (right).  
+
+## Clipping Properties
+
+You can change how clipping is handled based on the Widget selected in the UMG **Details** panel under **Clipping**. 
+
+![](ClippingProperties.png)
+
+[REGION:raw]
+| Property | %spacer2% | Description |
+| --- | --- | --- |
+| **Inherit** || This widget does not clip children and obeys whatever clipping / culling it was passed in from a parent widget. |
+| **Clip to Bounds** || This widget clips content to the bounds of the widget. It intersects those bounds with any previous clipping area. |
+| **Clip to Bounds - Without Intersecting** || This widget clips to its bounds. It does not intersect with any existing clipping geometry, it pushes a new clipping state. Effectively allowing it to render outside the bounds of hierarchy that does clip. [REGION:note] This will not allow you to ignore the clipping zone that is set to **Clip to Bounds - Always** [/REGION] |
+| **Clip to Bounds - Always** || This widget clips to its bounds. It intersects those bounds with any previous clipping area. [REGION:note] This clipping area cannot be ignored, it will always clip children. This is useful for hard barriers in the UI where you never want animations or other effects to break this region. [/REGION] |
+| **On Demand** || This widget clips to its bound when it's Desired Size is larger than the allocated geometry the widget is given. If that occurs, it will work like **Clip to Bounds**.  [REGION:note] This mode was primarily added for **Text**, which is often placed into containers that eventually are resized to not be able to support the length of the text. So rather than needing to tag every container that could contain text with [YES], which would result in almost no batching, this mode was added to dynamically adjust the clipping if needed. The reason not every is not set to **On Demand** is because not every panel returns a Desired Size that matches what it plans to render at. [/REGION] |
+[/REGION]
+
+
+## Additonal Considerations
+
+* In most cases, you should not need to adjust the clipping method except in instances where you can't control the length of the text and need to clip it. An example of this would be Scroll Box and Editable Text Widgets, which are set to **Clip to Bounds** instead of Inherit.
+* Elements in different clipping spaces cannot be batched together, so there is a performance cost that comes with clipping. For this reason, do not enable clipping unless a panel actually needs to prevent content from showing up outside of its bounds. 
+
diff --git a/Engine/Documentation/Source/Engine/UMG/UserGuide/CreatingWidgets/CreatingWidgets.INT.udn b/Engine/Documentation/Source/Engine/UMG/UserGuide/CreatingWidgets/CreatingWidgets.INT.udn
--- a/Engine/Documentation/Source/Engine/UMG/UserGuide/CreatingWidgets/CreatingWidgets.INT.udn
+++ b/Engine/Documentation/Source/Engine/UMG/UserGuide/CreatingWidgets/CreatingWidgets.INT.udn
@@ -1,9 +1,16 @@
-Availability:Public
+Availability: Public
+Crumbs: %ROOT%
 Title: Creating Widgets
-Crumbs:%ROOT%, Engine, Engine/UMG, Engine/UMG/UserGuide
 Description:How to create and display a Widget Blueprint in game. 
-version: 4.9
-tags:UMG UI Designer
+Type: Overview
+SkillLevel:Beginner
+Version:4.17
+Parent:Engine/UMG/UserGuide
+Order:2
+Tags:UMG UI Designer
+Tags:Getting Started
+Topic-image:Engine/UMG/UserGuide/CreateWidgetTopic.png
+
 
 Once you have created and laid out your **Widget Blueprint**, in order for it to be displayed in-game, you will need call it by using the **Create Widget** and **Add to Viewport** nodes inside another Blueprint (**Level Blueprint** or a **Character Blueprint** for example). 
 
diff --git a/Engine/Documentation/Source/Engine/UMG/UserGuide/DPIScaling/Scaling.INT.udn b/Engine/Documentation/Source/Engine/UMG/UserGuide/DPIScaling/Scaling.INT.udn
--- a/Engine/Documentation/Source/Engine/UMG/UserGuide/DPIScaling/Scaling.INT.udn
+++ b/Engine/Documentation/Source/Engine/UMG/UserGuide/DPIScaling/Scaling.INT.udn
@@ -1,9 +1,15 @@
-Availability:Public
+Availability: Public
+Crumbs: %ROOT%
 Title: DPI Scaling
-Crumbs:%ROOT%, Engine, Engine/UMG, Engine/UMG/UserGuide
-Description:How to access and adjust DPI Scaling Rules. 
-version: 4.9
-tags:UMG UI Designer
+Description: How to access and adjust DPI Scaling Rules.
+Type: Overview
+SkillLevel:Beginner
+Version:4.17
+Parent:Engine/UMG/UserGuide
+Order:4
+Tags:UMG UI Designer
+Topic-Image:Engine/UMG/UserGuide/ScalingTopic.png
+
 
 [REGION:banner]
 ![](DPIscaling.png)(convert:false)
diff --git a/Engine/Documentation/Source/Engine/UMG/UserGuide/Events/Events.INT.udn b/Engine/Documentation/Source/Engine/UMG/UserGuide/Events/Events.INT.udn
--- a/Engine/Documentation/Source/Engine/UMG/UserGuide/Events/Events.INT.udn
+++ b/Engine/Documentation/Source/Engine/UMG/UserGuide/Events/Events.INT.udn
@@ -1,9 +1,16 @@
-Availability:Public
+Availability: Public
+Crumbs: %ROOT%
 Title: Events
-Crumbs:%ROOT%, Engine, Engine/UMG, Engine/UMG/UserGuide
-Description:A guide to using Events in UMG. 
-version: 4.9
-tags:UMG UI Designer
+Description: A guide to using Events in UMG.
+Type: Overview
+SkillLevel:Beginner
+Version:4.17
+Parent:Engine/UMG/UserGuide
+Order:5
+Tags:UMG UI Designer
+Topic-Image:Engine/UMG/UserGuide/EventsTopic.png
+
+
 
 There are a couple methods in which you can call **Events** in **UMG** which are outlined below. 
 
diff --git a/Engine/Documentation/Source/Engine/UMG/UserGuide/Fonts/Fonts.INT.udn b/Engine/Documentation/Source/Engine/UMG/UserGuide/Fonts/Fonts.INT.udn
--- a/Engine/Documentation/Source/Engine/UMG/UserGuide/Fonts/Fonts.INT.udn
+++ b/Engine/Documentation/Source/Engine/UMG/UserGuide/Fonts/Fonts.INT.udn
@@ -3,14 +3,15 @@ Crumbs:%ROOT%, Engine, Engine/UMG, Engine/UMG/UserGuide
 Title: Fonts 
 Description: Landing page for all the documentation relating to Fonts.
 Type: Landing
-Version:4.15
+Version:4.17
 Parent: Engine/UMG/UserGuide
 Related: Engine/UMG
 Related: Engine/UMG/UserGuide/WidgetBlueprints
 Related: Engine/UMG/UserGuide/WidgetTypeReference
-Order:1
+Order:6
 Tags: UMG UI Designer
 Tags: Fonts
+Topic-image:Engine/UMG/UserGuide/FontsTopic.png
 
 [REGION:fullwidth raw]
 ![](FontLandingBanner.png)
diff --git a/Engine/Documentation/Source/Engine/UMG/UserGuide/PropertyBinding/PropertyBinding.INT.udn b/Engine/Documentation/Source/Engine/UMG/UserGuide/PropertyBinding/PropertyBinding.INT.udn
--- a/Engine/Documentation/Source/Engine/UMG/UserGuide/PropertyBinding/PropertyBinding.INT.udn
+++ b/Engine/Documentation/Source/Engine/UMG/UserGuide/PropertyBinding/PropertyBinding.INT.udn
@@ -1,9 +1,17 @@
-Availability:Public
+Availability: Public
+Crumbs: %ROOT%
 Title: Property Binding
-Crumbs:%ROOT%, Engine, Engine/UMG, Engine/UMG/UserGuide
-Description:A guide to binding properties inside UMG. 
-Version: 4.9
-tags:UMG UI Designer
+Description: A guide to binding properties inside UMG. 
+Type: Overview
+SkillLevel:Beginner
+Version:4.17
+Parent:Engine/UMG/UserGuide
+Order:7
+Tags:UMG UI Designer
+Topic-Image:Engine/UMG/UserGuide/PropertyBindingTopic.png
+
+
+
 
 [TOC(start:2 end:2)]
 
diff --git a/Engine/Documentation/Source/Engine/UMG/UserGuide/Slots/Slots.INT.udn b/Engine/Documentation/Source/Engine/UMG/UserGuide/Slots/Slots.INT.udn
--- a/Engine/Documentation/Source/Engine/UMG/UserGuide/Slots/Slots.INT.udn
+++ b/Engine/Documentation/Source/Engine/UMG/UserGuide/Slots/Slots.INT.udn
@@ -1,9 +1,17 @@
-Availability:Public
+Availability: Public
+Crumbs: %ROOT%
 Title: Slots
-Crumbs:%ROOT%, Engine, Engine/UMG, Engine/UMG/UserGuide
-Description:An overview of how slots are used in UMG.
-Version: 4.9
-tags:UMG UI Designer
+Description:An overview of how slots are used in UMG.  
+Type: Overview
+SkillLevel:Beginner
+Version:4.17
+Parent:Engine/UMG/UserGuide
+Order:8
+Tags:UMG UI Designer
+Topic-Image:Engine/UMG/UserGuide/SlotTopic.png
+
+
+
 
 [TOC(start:2 end:2)]
 
diff --git a/Engine/Documentation/Source/Engine/UMG/UserGuide/Styling/Styling.INT.udn b/Engine/Documentation/Source/Engine/UMG/UserGuide/Styling/Styling.INT.udn
--- a/Engine/Documentation/Source/Engine/UMG/UserGuide/Styling/Styling.INT.udn
+++ b/Engine/Documentation/Source/Engine/UMG/UserGuide/Styling/Styling.INT.udn
@@ -1,10 +1,18 @@
-Availability:Public
+Availability: Public
+Crumbs: %ROOT%
 Title: Styling
-Crumbs:%ROOT%, Engine, Engine/UMG, Engine/UMG/UserGuide
-Description:A guide to stylizing Widget Blueprints in UMG. 
+Description: A guide to stylizing Widget Blueprints in UMG. 
+Type: Overview
+SkillLevel:Beginner
+Version:4.17
+Parent:Engine/UMG/UserGuide
 Related: Engine\UMG\UserGuide\Slots
 Related: Engine\UMG\UserGuide\WidgetTypeReference
-tags:UMG UI Designer
+Order:9
+Tags:UMG UI Designer
+Topic-Image:Engine/UMG/UserGuide/StylingTopic.png
+
+
 
 
 [REGION:fullwidth raw]
diff --git a/Engine/Documentation/Source/Engine/UMG/UserGuide/UserGuide.INT.udn b/Engine/Documentation/Source/Engine/UMG/UserGuide/UserGuide.INT.udn
--- a/Engine/Documentation/Source/Engine/UMG/UserGuide/UserGuide.INT.udn
+++ b/Engine/Documentation/Source/Engine/UMG/UserGuide/UserGuide.INT.udn
@@ -1,243 +1,30 @@
-Availability:Public
+Availability: Public
+Crumbs: %ROOT%
 Title: UMG UI Designer User Guide
-Crumbs:%ROOT%, Engine, Engine/UMG
-Description:A guide to using Unreal Motion Graphics to create UI elements. 
-version: 4.9
-tags:UMG UI Designer
+Description: A guide to using Unreal Motion Graphics to create UI elements.
+Type: Landing
+Version:4.17
+Parent:Engine/UMG
+Order:
+Tags:UMG UI Designer
 
 
-[VAR:Topic]
-[OBJECT:Topic]
-	[PARAM:image]
-		![%Engine/UMG:title%](UMG_Topic.png)
-	[/PARAM]
-	[PARAM:icon]
-		![](%ROOT%/start_icon.png)(convert:false)
-	[/PARAM]
-	[PARAM:title]
-		%Engine/UMG:title%
-	[/PARAM]
-	[PARAM:description]
-		%Engine/UMG:description%
-	[/PARAM]
-	[PARAM:path]
-		[RELATIVE:Engine/UMG]
-	[/PARAM]
-[/OBJECT]
-[/VAR]
-[VAR:TopicCompact]
-[OBJECT:TopicCompact]
-	[PARAM:image]
-		![%Engine/UMG:title%](UMG_Topic.png)
-	[/PARAM]
-	[PARAM:icon]
-		![](%ROOT%/start_icon.png)(convert:false)
-	[/PARAM]
-	[PARAM:title]
-		%Engine/UMG:title%
-	[/PARAM]
-	[PARAM:description]
-		%Engine/UMG:description%
-	[/PARAM]
-	[PARAM:path]
-		[RELATIVE:Engine/UMG]
-	[/PARAM]
-[/OBJECT]
-[/VAR]
 
 [REGION:fullwidth raw]
 ![](UserGuideHeader.png)(w:940 h:160 convert:false)
 [/REGION]
 
-
 [EXCERPT:umgheader]
 **Unreal Motion Graphics UI Designer (UMG)**  is a visual UI authoring tool which can be used to create UI elements such as in-game HUDs, menus or other interface related graphics you wish to present to your users. At the core of UMG are Widgets, which are a series of pre-made functions that can be used to construct your interface (things like buttons, checkboxes, sliders, progress bars, etc.). These Widgets are edited in a specialized Widget Blueprint, which uses two tabs for construction: the Designer tab allows for the visual layout of the interface and basic functions while the Graph tab provides the functionality behind the Widgets used. 
 [/EXCERPT:umgheader]
 
-[OBJECT:TopicButtonList]
-	[PARAM:icon]
-		![](%ROOT%/start_icon.png)(convert:false)
-	[/PARAM]
-	[PARAM:title]
-		Getting Started
-	[/PARAM]
-	[PARAM:links]
-		[REGION:topics fourth]
-		[INCLUDE:#topics]
-		[/REGION]
-	[/PARAM]
-[/OBJECT]
 
-<!--
-[EXCERPT:Topics]
-[OBJECT:Topic]
-	[PARAM:image]
-		![%Engine/UMG/UserGuide/WidgetBlueprints:title%](WidgetBlueprintTopic.png)(w:291 h:120)
-	[/PARAM]
-	[PARAM:icon]
-		![](%ROOT%/blueprint_icon.png)(convert:false)
-	[/PARAM]
-	[PARAM:title]
-		%Engine/UMG/UserGuide/WidgetBlueprints:title%
-	[/PARAM]
-	[PARAM:description]
-		%Engine/UMG/UserGuide/WidgetBlueprints:description%
-	[/PARAM]
-	[PARAM:path]
-		[RELATIVE:Engine/UMG/UserGuide/WidgetBlueprints]
-	[/PARAM]
-[/OBJECT]
-[OBJECT:Topic]
-	[PARAM:image]
-		![%Engine/UMG/UserGuide/CreatingWidgets:title%](CreateWidgetTopic.png)(w:291 h:120)
-	[/PARAM]
-	[PARAM:icon]
-		![](%ROOT%/blueprint_icon.png)(convert:false)
-	[/PARAM]
-	[PARAM:title]
-		%Engine/UMG/UserGuide/CreatingWidgets:title%
-	[/PARAM]
-	[PARAM:description]
-		%Engine/UMG/UserGuide/CreatingWidgets:description%
-	[/PARAM]
-	[PARAM:path]
-		[RELATIVE:Engine/UMG/UserGuide/CreatingWidgets]
-	[/PARAM]
-[/OBJECT]
-[OBJECT:Topic]
-	[PARAM:image]
-		![%Engine/UMG/UserGuide/WidgetTypeReference:title%](WidgetTypeRefTopic.png)(w:291 h:120)
-	[/PARAM]
-	[PARAM:icon]
-		![](%ROOT%/blueprint_icon.png)(convert:false)
-	[/PARAM]
-	[PARAM:title]
-		%Engine/UMG/UserGuide/WidgetTypeReference:title%
-	[/PARAM]
-	[PARAM:description]
-		%Engine/UMG/UserGuide/WidgetTypeReference:description%
-	[/PARAM]
-	[PARAM:path]
-		[RELATIVE:Engine/UMG/UserGuide/WidgetTypeReference]
-	[/PARAM]
-[/OBJECT]
-[OBJECT:Topic]
-	[PARAM:image]
-		![%Engine\UMG\UserGuide\BestPractices:title%](Engine\UMG\UserGuide\BestPractices\BestPractices.png)(w:291 h:120)
-	[/PARAM]
-	[PARAM:icon]
-		![](%ROOT%/blueprint_icon.png)(convert:false)
-	[/PARAM]
-	[PARAM:title]
-		%Engine\UMG\UserGuide\BestPractices:title%
-	[/PARAM]
-	[PARAM:description]
-		%Engine\UMG\UserGuide\BestPractices:description%
-	[/PARAM]
-	[PARAM:path]
-		[RELATIVE:Engine\UMG\UserGuide\BestPractices]
-	[/PARAM]
-[/OBJECT]
-[/EXCERPT:Topics]
--->
+## Getting Started
+
+[DIR(output:"topiccompact" parent:"Engine/UMG/UserGuide" org:"hierarchy" tags:"Getting Started" end:"1")]
+
+
+## Further Reading
+
+[DIR(output:"topiccompact" parent:"Engine/UMG/UserGuide" org:"hierarchy" tags:"!Getting Started" end:"1")]
 
-[OBJECT:TopicButtonList]
-	[PARAM:icon]
-		![](%ROOT%/related_icon.png)(convert:false)
-	[/PARAM]
-	[PARAM:title]
-		Further Reading
-	[/PARAM]
-	[PARAM:links]
-	[REGION:topics third]
-	[OBJECT:TopicCompact]
-		[PARAM:image]
-			![](AnchorsTopic.png)
-		[/PARAM]
-		[PARAM:title]
-			Anchors
-		[/PARAM]
-		[PARAM:path]
-		[RELATIVE:Engine/UMG/UserGuide/Anchors]
-		[/PARAM]
-	[/OBJECT]
-	[OBJECT:TopicCompact]
-		[PARAM:image]
-			![](AnimationTopic.png)
-		[/PARAM]
-		[PARAM:title]
-			Animation
-		[/PARAM]
-		[PARAM:path]
-		[RELATIVE:Engine/UMG/UserGuide/Animation]
-		[/PARAM]
-	[/OBJECT]
-	[OBJECT:TopicCompact]
-		[PARAM:image]
-			![](EventsTopic.png)
-		[/PARAM]
-		[PARAM:title]
-			Events
-		[/PARAM]
-		[PARAM:path]
-		[RELATIVE:Engine/UMG/UserGuide/Events]
-		[/PARAM]
-	[/OBJECT]
-	[OBJECT:TopicCompact]
-		[PARAM:image]
-			![](PropertyBindingTopic.png)
-		[/PARAM]
-		[PARAM:title]
-			Property Binding
-		[/PARAM]
-		[PARAM:path]
-		[RELATIVE:Engine/UMG/UserGuide/PropertyBinding]
-		[/PARAM]
-	[/OBJECT]
-	[OBJECT:TopicCompact]
-		[PARAM:image]
-			![](ScalingTopic.png)
-		[/PARAM]
-		[PARAM:title]
-			DPI Scaling
-		[/PARAM]
-		[PARAM:path]
-		[RELATIVE:Engine/UMG/UserGuide/DPIScaling]
-		[/PARAM]
-	[/OBJECT]
-	[OBJECT:TopicCompact]
-		[PARAM:image]
-			![](FontsTopic.png)
-		[/PARAM]
-		[PARAM:title]
-			Fonts
-		[/PARAM]
-		[PARAM:path]
-		[RELATIVE:Engine/UMG/UserGuide/Fonts]
-		[/PARAM]
-	[/OBJECT]
-	[OBJECT:TopicCompact]
-		[PARAM:image]
-			![](StylingTopic.png)
-		[/PARAM]
-		[PARAM:title]
-			Styling
-		[/PARAM]
-		[PARAM:path]
-		[RELATIVE:Engine/UMG/UserGuide/Styling]
-		[/PARAM]
-	[/OBJECT]
-	[OBJECT:TopicCompact]
-		[PARAM:image]
-			![](SlotTopic.png)
-		[/PARAM]
-		[PARAM:title]
-			Slots
-		[/PARAM]
-		[PARAM:path]
-		[RELATIVE:Engine/UMG/UserGuide/Slots]
-		[/PARAM]
-	[/OBJECT]
-	[/REGION]
-	[/PARAM]
-[/OBJECT]
diff --git a/Engine/Documentation/Source/Engine/UMG/UserGuide/WidgetBlueprints/WidgetBlueprints.INT.udn b/Engine/Documentation/Source/Engine/UMG/UserGuide/WidgetBlueprints/WidgetBlueprints.INT.udn
--- a/Engine/Documentation/Source/Engine/UMG/UserGuide/WidgetBlueprints/WidgetBlueprints.INT.udn
+++ b/Engine/Documentation/Source/Engine/UMG/UserGuide/WidgetBlueprints/WidgetBlueprints.INT.udn
@@ -1,9 +1,19 @@
-Availability:Public
+Availability: Public
+Crumbs: %ROOT%
 Title: Widget Blueprints
-Crumbs:%ROOT%, Engine, Engine/UMG, Engine/UMG/UserGuide
-Description:How to create a Widget Blueprint and Overview of the Widget Blueprint Interface.
-Version: 4.9
-tags:UMG UI Designer
+Description: How to create a Widget Blueprint and Overview of the Widget Blueprint Interface
+Type: Overview
+SkillLevel:Beginner
+Version:4.17
+Parent:Engine/UMG/UserGuide
+Order:1
+Tags:UMG UI Designer
+Tags:Getting Started
+Topic-image:Engine/UMG/UserGuide/WidgetBlueprintTopic.png
+
+
+
+
 
 In order to start working with Unreal Motion Graphics, you will first need to create a **Widget Blueprint** as seen below. 
 
diff --git a/Engine/Documentation/Source/Engine/UMG/UserGuide/WidgetInteraction/WidgetInteraction.INT.udn b/Engine/Documentation/Source/Engine/UMG/UserGuide/WidgetInteraction/WidgetInteraction.INT.udn
--- a/Engine/Documentation/Source/Engine/UMG/UserGuide/WidgetInteraction/WidgetInteraction.INT.udn
+++ b/Engine/Documentation/Source/Engine/UMG/UserGuide/WidgetInteraction/WidgetInteraction.INT.udn
@@ -1,12 +1,19 @@
-Availability:Public
+Availability: Public
+Crumbs: %ROOT%
 Title: Widget Interaction Component
-Crumbs:%ROOT%, Engine, Engine/UMG, Engine/UMG/UserGuide
-Description:Outlines how to use the Widget Interaction Component to interact with 3D Widgets.
-Version: 4.13
+Description: Outlines how to use the Widget Interaction Component to interact with 3D Widgets. 
+Type: Overview
+SkillLevel:Beginner
+Version:4.17
+Parent:Engine/UMG/UserGuide
 Related: Engine/UMG/UserGuide
 Related: Engine/Blueprints
 Related: Engine/Components/Widget
-tags:UMG UI Designer
+Order:10
+Tags:UMG UI Designer
+Topic-Image:WidgetInteractionBanner.png
+
+
 
 [REGION:warning]
 This feature is considered experimental and some aspects may not work as expected and may change in future updates. 
diff --git a/Engine/Documentation/Source/Engine/UMG/UserGuide/WidgetTypeReference/WidgetTypeReference.INT.udn b/Engine/Documentation/Source/Engine/UMG/UserGuide/WidgetTypeReference/WidgetTypeReference.INT.udn
--- a/Engine/Documentation/Source/Engine/UMG/UserGuide/WidgetTypeReference/WidgetTypeReference.INT.udn
+++ b/Engine/Documentation/Source/Engine/UMG/UserGuide/WidgetTypeReference/WidgetTypeReference.INT.udn
@@ -1,9 +1,14 @@
-Availability:Public
+Availability: Public
+Crumbs: %ROOT%
 Title: Widget Type Reference
-Crumbs:%ROOT%, Engine, Engine/UMG, Engine/UMG/UserGuide
-Description:Reference for the types of Widgets available in Unreal Motion Graphics (UMG).
-version: 4.15
-tags:UMG UI Designer
+Description: Reference for the types of Widgets available in Unreal Motion Graphics (UMG).
+Type: Reference
+Version: 4.17
+Parent:Engine/UMG/UserGuide
+Order:3
+Tags:UMG UI Designer
+Tags:Getting Started
+Topic-image:Engine/UMG/UserGuide/WidgetTypeRefTopic.png
 
 [TOC(start:2 end:2)]
 
diff --git a/Engine/Documentation/Source/Gameplay/Gameplay.INT.udn b/Engine/Documentation/Source/Gameplay/Gameplay.INT.udn
--- a/Engine/Documentation/Source/Gameplay/Gameplay.INT.udn
+++ b/Engine/Documentation/Source/Gameplay/Gameplay.INT.udn
@@ -37,6 +37,7 @@ From input to saving, find overviews of these topics and example setups in C++ a
 %Gameplay/AI:role%
 %Gameplay/Localization:role%
 %Gameplay/Analytics:role%
+%Gameplay/Tags:role%
 
 ## Gameplay Tools
 	
diff --git a/Engine/Documentation/Source/Gameplay/HowTo/SpawnAndDestroyActors/Blueprints/GHT_1B.INT.udn b/Engine/Documentation/Source/Gameplay/HowTo/SpawnAndDestroyActors/Blueprints/GHT_1B.INT.udn
--- a/Engine/Documentation/Source/Gameplay/HowTo/SpawnAndDestroyActors/Blueprints/GHT_1B.INT.udn
+++ b/Engine/Documentation/Source/Gameplay/HowTo/SpawnAndDestroyActors/Blueprints/GHT_1B.INT.udn
@@ -1,221 +1,184 @@
 Availability: Public
-Title:Spawning/Destroying an Actor in Blueprints
 Crumbs: %ROOT%, Gameplay, Gameplay/HowTo/SpawnAndDestroyActors
-Description:A How To Guide for Spawning/Destroying Actors through Blueprints in Unreal Engine 4.
-Version: 4.9
-SkillLevel:Intermediate
-type:how-to
+Title: Spawn and Destroy an Actor with Blueprint
+Description: In this guide, we use Blueprint to spawn an object when a button is pressed and destroy it when the button is pressed a second time. 
+Type: how-to
+SkillLevel: Beginner
+Version: 4.16	
+Parent: Gameplay/HowTo/SpawnAndDestroyActors
+Order: 1
+Tags:Blueprints
+Tags:Basics
+Tags:Getting Started
+Related: Engine/Blueprints
 
-[TOC (start:2)]
+Regardless of the game type you are making, more often than not there are going to be instances where you want something to appear at runtime that was not originally placed in the level. 
+The concept of "spawning" an object in the world can be achieved with [Blueprint Visual Scripting](Engine/Blueprints) as well as the ability to "destroy" or despawn objects that have been (or have not been) spawned in during gameplay. 
 
- Here you will learn how to spawn an Actor in Blueprints and define its spawn location. You will also learn how to destroy spawned Actors. 
+ In this guide, we will spawn in an [Actor](Engine/Actors) with the press of a button, attach it to a character, then destroy it when the button is pressed a second time using a [variable](Engine/Blueprints/UserGuide/Variables) reference. 
 
-## Spawning an Actor
+![](HeroImage.png)
 
-Below you will find steps to spawn an Actor in the level by pressing a key. 
+## Steps
 
-1.  Open up any project, once inside the editor, click the **Blueprints** button and **Open Level Blueprint**.  
+[REGION:note]
+For this how-to guide, we are using the **Blueprint Third Person Template** with **Starter Content** enabled. 
+[/REGION]
 
-	[REGION:raw]
-	![](GHT1_1.png)
-	[/REGION]
+1.	Inside your project on the **Main Toolbar**, click the **Blueprints** button then select **Open Level Blueprint**. 
 
-1.  This will open the **Blueprint Editor** window (shown below). 
-	
-	[REGION:lightbox]
-	[![](Engine\Blueprints\UserGuide\Types\LevelBlueprint\level_blueprint_editor.png)](Engine\Blueprints\UserGuide\Types\LevelBlueprint\level_blueprint_editor.png)
-	[/REGION]
+	![](SpawnDestroy_01.png)
 
-	_Click image for full view._
-
-	In order to spawn an Actor, we will need to use the Spawn Actor from Class node. In this example, we are using the Level Blueprint, however you can use the Spawn Actor from Class node in most Blueprints where it can be found in the Palette window as shown below.
+	The [](Engine/Blueprints/UserGuide/Types/LevelBlueprint) is a specialized type of Blueprint that is unique to each level and acts as a level-wide global graph. 
 
-1.  **Right-click** anywhere in the graph and in the context menu, search for and add the **Spawn Actor from Class** node. 
+1.	Right-click in the graph and search for and add the **Spawn Actor from Class** node. 
 
-	![](GHT1_5.png)
+	![](SpawnDestroy_02.png)
 
-1.  The **Spawn Actor from Class** node will be added to the graph (**SpawnActor** for short). 
+	This will enable us to define the thing that we want to spawn by specifying the class (or type of object) we want to spawn. 
 
-	![](GHT1_4.png)
+1.	In the **Class** section of the **SpawnActor** node, select the **Blueprint_Effect_Fire** class. 
 
-	The Spawn Actor from Class node takes an Actor Class (specified via the Class input) and attempts to spawn an instance of that class in the world. The Spawn Transform input is used to define the Actor's spawn location in the world (and its initial orientation). If that location is blocked by a collider, then the Actor will fail to spawn unless the Spawn Even if Colliding input is set to True. 
+	![](SpawnDestroy_03.png)
 
-	Each of the pins on the Spawn Actor node are outlined below.
+	This is a Blueprint Class that was created as part of the **Starter Content** pack enabled when creating the project. 
 
-	![](GHT1_6.png)
+1.	Right-click in the graph and search for and add the **F** Keyboard Event. 
 
-	| Number | Description |
-	| ------- | ------------ |
-	| **Execution Input (1)** | This is an execution input and is what triggers the node to spawn an Actor. |
-	| **Class (2)** | This is the Actor Class you want to spawn an instance of in the world (this is required). |
-	| **Spawn Transform (3)** | This takes in a Transform used to position and orientate the Actor in the world. |
-	| **Spawn Even If Colliding (4)** | Determines if the Actor should be spawned when something is blocking the spawn location. |
-	| **Execution Output (5)** | This is an execution output and leads to script that should occur following the Spawn Actor. |
-	| **Return Value (6)** | This outputs the new Actor instance that is spawned in the world. |
+	![](SpawnDestroy_04.png)
 
-1.  Click the **Class** drop-down menu and then select an Actor from the list of Actors you want to spawn. 
+	This will be the key that we will use to spawn and destroy our fire effect. 
 
-	![](GHT1_7.png)
+1.	Right-click in the graph again and search for and add the **Flip Flop** node. 
 
-1.  After selecting the Actor you want to spawn, the node will update and a new pin will be added. 
+	![](SpawnDestroy_05.png)
 
-	![](GHT1_8.png)
+	This node acts as a switch where the first time we enter it we will execute the **A** pin and the second we enter this node will execute the **B** pin. 
+	Once executing the B pin, the next time we enter we will execute the A pin again and "flip-flop" back-and-forth between the A and B pins. 
 
-	The Instigator is optional and is the Pawn responsible for damage done by the spawned Actor.
+1.	Left-click and drag off the **Pressed** pin of the **F** node and connect it to the **Flip Flop** and **Spawn Actor** nodes. 
 
-1.  **Left-click** and drag off the **Spawn Transform** pin to the left, then in the context menu search for and add **Get Actor Transform**. 
+	![](SpawnDestroy_06.png)
 
-	![](GHT1_9.png)
+	This will enable us to execute (when F is pressed) the call to the Spawn Actor node to spawn in our Blueprint Fire effect. 
 
-	The **Get Actor Transform** node is telling the **Spawn Actor** where to spawn the new Actor we want to spawn.
+1.	Right-click in the graph and search for and add the **Destroy Actor** node. 
 
-1.  **Left-click** and drag off the **Get Actor Transform** pin, then in the context menu search for and add **Get Player Character**. 
+	![](SpawnDestroy_07.png)
 
-	![](GHT1_10.png)
+	With this node we can destroy any Actor we specify. 
 
-	This is now saying, Get the Player Character then Get their Transform and assign that to the **Spawn Transform**. 
+1.	Connect **B** pin of the **Flip Flop** node to the **Destroy Actor** node. 
 
-1.  **Right-click** anywhere in the graph, and in the context menu, search for and add the **F** Key Event. 
+	![](SpawnDestroy_08.png)
 
-	![](GHT1_11.png)	
+	The second time we press **F**, the **B** pin on the **Flip Flop** will be called and will execute the **Destroy Actor** node. 
+	The Destroy Actor has a **Target** input pin, instructing the node what it is to destroy. 
+	By default, it is set to **Self** and in this case is a null pointer, meaning it is unclear what it is to destroy as the Level Blueprint cannot be destroyed. 
+	If this Destroy Actor node were sitting inside our Character Blueprint and we instructed it to target Self, then our Character would be destroyed. 
 
-1.  Connect the **Pressed** pin to the **Execute Input** pin of the **SpawnActor** node. 
+	In this case, we need to tell the Destroy Actor node what it should be destroying (our Blueprint Fire effect). 
+	To do this we will need to create a reference to the Blueprint Fire effect that we spawn. 
 
-	![](GHT1_12.png)
+1.	Right-click on the **Return Value** of the **Spawn Actor** node and select **Promote to Variable**. 
 
-	This now says, execute the **SpawnActor** when **F** is pressed and spawn the Actor at the Player Character's Transform. 
+	![](SpawnDestroy_09.png)
 
-1.  Click the **Compile** button then close the **Blueprint**. 
+	The **Return Value**, in this case, is the "thing that we spawn" and what we want to store a reference to so that we can destroy it later. 
 
-	Before Compile: 
-	
-	![](GHT1_13.png)
-
-	After Successful Compile: 
-
-	![](GHT1_14.png)
-
-	If you do not see the green check mark, verify that all your nodes are connected and an **Actor Class** is specified on the **SpawnActor** node.
-
-1.  From the main window, click the **Play** Button to play in the editor. 
-
-	[REGION:raw]
-	![](GHT1_15.png)
-	[/REGION]
+1.	In the **Details** panel for the variable, give it a name such as **CurrentFlames**. 
 
-	Whenever you press **F**, the Actor you selected should spawn where the player character is in the world.
+	![](SpawnDestroy_10.png)
 
-This example illustrates spawning an Actor through a key press, however you could use any number of Events to execute the SpawnActor node. 
+1.	In the **MyBlueprint** panel under **Variables**, hold **Ctrl** then drag-and-drop the **CurrentFlames** variable into the graph and connect it to the **Target** of **Destroy Actor**. 
 
-Refer to the [](Engine/Blueprints/UserGuide) for more on working with **Blueprints**. 
+	![](SpawnDestroy_11.png)
 
-## Setting a Spawn Point
+	We are now instructing the Destroy Actor node to target our variable reference to the Blueprint Fire effect we created. 
 
-If you want to spawn an Actor at a specified location, you can do this by setting the Transform manually or by referencing another Actor's Transform.
-
-The examples below will show you how to manually set a Transform as well as reference an Actor's Transform to use as a spawn location.
-
-1.  Open up any project, once inside the editor, click the **Blueprints** button and **Open Level Blueprint**.  
-
-	[REGION:raw]
-	![](GHT1_1.png)
-	[/REGION]
-
-1.	In the graph window, **Right-click** and search for then add the **Spawn Actor from Class** node. 
-
-	![](GHT1_5.png)
-
-1.	On the **SpawnActor** node, click the **Class** drop-down menu and select an Actor to spawn. 
-
-	![](GHT1_7.png)
-
-1.	Minimize the **Level Blueprint** window. 
-
-1.  From the Main Editor window, under **Modes**, search for **Target** and drag a **Target Point** into your level. 
-
-	[REGION:raw]
-	![](GHT1_16.png)
-	[/REGION]
+1.	Right-click in the graph and search for and add the **Get Player Character** node. 
 
-1.  Return to the **Level Blueprint** window, then **Right-click** in the graph and choose **Add Reference to TargetPoint**.
+	![](SpawnDestroy_12.png)
 
-	![](GHT1_17.png)
+	This enables us to locate the designated player character so that we can access and obtain information about the character. 
 
-1.  **Left-click** drag off the out pin of the newly added node and search for and add **Get Actor Transform**. 
+1.	Off the **Return Value** of **Get Player Character**, search for and add the **Get Actor Transform** node. 
 
-	![](GHT1_18.png)
+	![](SpawnDestroy_13.png)
 
-1.  Connect the **Return Value** of the **Get Actor Transform** node to the **Spawn Transform** in pin. 
+	This will tell us the **Location**, **Rotation** and **Scale** of our player character and we will use this to define the spawn location of our Blueprint Fire effect. 
 
-1.	**Right-click** in the graph and search for and add the **F** Key Event, then off **Pressed** connect it to the **SpawnActor**. 
+1.	Connect the **Return Value** of the **Get Actor Transform** to the **Spawn Transform** on the **Spawn Actor** node. 
 
-	![](GHT1_19.png)
-
-	We used a Key Press Event to execute the **SpawnActor** node, however you can use any method you wish.
-
-1.	**Compile** then close the **Blueprint** and play in the editor, you should see your Actor spawn at the Target Point. 
-
-In addition to specifying a spawn point by referencing an Actor, you can manually enter the Spawn Transform for the SpawnActor as well, to do this: 
-
-1.	Drag off the **Spawn Transform** pin of a **SpawnActor** node, then select **Make Transform**. 
-
-	![](GHT1_21.png)
-
-1.	This will create a **Make Transform** node where you can enter the **Location**, **Rotation**, and **Scale** manually. 
-
-	![](GHT1_22.png)
-
-## Destroying a Spawned Actor
-
-Once you have spawned an Actor, suppose you wanted to destroy it (for example, maybe the player killed the Actor or picked the Actor up as a pick-up, etc.). 
-
-Below are the steps you can take to destroy a Spawned Actor in the Level Blueprint. 
-
-1.  Open up any project, once inside the editor, click the **Blueprints** button and **Open Level Blueprint**.  
-
-	[REGION:raw]
-	![](GHT1_1.png)
+	[REGION:lightbox]
+	[![](SpawnDestroy_14.png)](SpawnDestroy_14.png)
 	[/REGION]
 
-1.	In the graph window, **Right-click** and search for then add the **Spawn Actor from Class** node. 
-
-	![](GHT1_5.png)
-
-1.	On the **SpawnActor** node, click the **Class** drop-down menu and select an Actor to spawn. 
-
-	![](GHT1_7.png)
-
-1.	**Left-click** and drag off the **Spawn Transform** pin and add a **Get Actor Transform** node. 
+	_Click image for full view._
 
-1.	**Left-click** and drag off the **Get Actor Transform** node and add a **Get Player Character** node. 
+	Our graph at this point is almost complete and we designated that when we press the **F** key that we will spawn in our Blueprint Fire at the location of our player and store it as a variable. 
+	When we press **F** a second time, we will then destroy that variable that we created (killing off the fire effect). The last thing that we will do is attach the fire effect to our player. 
+	Currently when we spawn in our fire effect it will spawn at the location of our player, however, if the player moves around the fire will stay in the location in which it was spawned. 
 
-1.	**Right-click** in an empty space in the graph, then search for the **F** Key Event node and connect it to the **SpawnActor** node. 
+1.	Left-click and drag off the reference pin on the **Set** node and search for and add the **AttachToActor** node. 
 
-	![](GHT1_12.png)
+	![](SpawnDestroy_15.png)
 
-1.	**Left-click** and drag off the **Execute Out** pin of the **SpawnActor** and add a **Delay** node. 
+	With this node we can attach the fire effect to another Actor, in this case, we will want to attach it to our player character. 
 
-1. On the **Delay** node, set the duration to **3** (this is the amount of seconds to wait before we destroy the Actor). 
+1.	Use a **Get Player Character** node and attach as shown below, setting the **Location Rule** and **Rotation Rule** to **Snape to Target**. 
 
-1. **Left-click** and drag off the **Return Value** of the **SpawnActor** and add a **Destroy Actor** node, then connect as shown below. 
-	
 	[REGION:lightbox]
-	[![](GHT1_20.png)(w:740)](GHT1_20.png)
+	[![](SpawnDestroy_16.png)](SpawnDestroy_16.png)
 	[/REGION]
 
 	_Click image for full view._
 
-1.	**Compile** then close the **Blueprint** and play in the editor, you should see the Actor spawn at the player location then despawn after 3 seconds. 
-
-As with executing the SpawnActor node to spawn an Actor, destroying the Actor can be triggered by several different types of Events based on your needs. 
-
-Refer to the [](Engine/Blueprints/UserGuide) for more on working with Blueprints. 
-
-
+	Here we are stating that we want to attach the reference variable for our fire effect to the player character, snapping its location and rotation to match the player character. 
 
+1.	Click the **Compile** button to compile the Blueprint. 
 
+	![](SpawnDestroy_17.png)
 
+1.	Return to the main Level Editor window, then from the **Main Toolbar** click the **Play** button to play inside the editor. 
 
+	![](SpawnDestroy_18.png)
+	
 
+## End Result
+
+While playing in the editor, when you press the **F** key the fire effect will spawn in and become attached to the player character. 
+When pressing the **F** key while the fire is attached to the player, this will cause the effect to be destroyed where you can press **F** again to re-spawn it. 
+
+[OBJECT:EmbeddedVideo]
+	[PARAMLITERAL:videoid]
+	uEzFpMt6mxo
+	[/PARAMLITERAL]
+	[PARAMLITERAL:width]
+	100
+	[/PARAMLITERAL]
+	[PARAMLITERAL:height]
+	56.25
+	[/PARAMLITERAL]
+	[PARAMLITERAL:units]
+	%
+	[/PARAMLITERAL]
+	[PARAMLITERAL:type]
+	simple
+	[/PARAMLITERAL]
+[/OBJECT]
+
+Another way you could approach this is to **Activate** and **Deactivate** the particle effect contained within the Blueprint Fire effect. 
+
+[REGION:lightbox]
+[![](AlternativeMethod.png)](AlternativeMethod.png)
+[/REGION]
+
+_Click image for full view._
+
+In the example above, we spawn in the Blueprint Fire at the start of the game using the **Event BeginPlay** node and attach it to our player character. 
+We also deactivate the particle effect and fade out the associated sound at the start of the game. 
+When we press **F**, since the Blueprint Fire effect has already been spawned in (just not activated), we instead activate it and fade in the audio. 
+When we press **F** a second time, instead of destroying the effect, we deactivate it and fade out the associated sounds. 
 
 
diff --git a/Engine/Documentation/Source/Gameplay/Networking/Blueprints/Blueprints.INT.udn b/Engine/Documentation/Source/Gameplay/Networking/Blueprints/Blueprints.INT.udn
--- a/Engine/Documentation/Source/Gameplay/Networking/Blueprints/Blueprints.INT.udn
+++ b/Engine/Documentation/Source/Gameplay/Networking/Blueprints/Blueprints.INT.udn
@@ -55,7 +55,7 @@ In the details panel of variables on your actors, there is a **Replication** dro
 | ------ | ----------- |
 | **None** | This is the default for new variables, and means the value will not be sent over the network to clients. |
 | **Replicated** | When the server replicates this actor, it will send this variable to clients. The value of the variable on the receiving client will update automatically, so that the next time it's accessed, it will reflect what the value was on the server. Of course, when playing over a real-world network, the updates will be delayed by an amount of time dependent on the network's latency. Remember that replicated variables only go in one direction, from the server to the client! To send data from a client to the server, see the "Events" section. |
-| **RepNotify** | The variable will be replicated as in the Replicated option, but in addition, an OnRep_<variable name> function will be created in your blueprint. This function will be called by the engine automatically on the client and the server whenever the value of this variable changes. You're free to implement this function however you'd like, as needed by your game. ![](onrep.png) |
+| **RepNotify** | The variable will be replicated as in the Replicated option, but in addition, an `OnRep_<variable name>` function will be created in your blueprint. This function will be called by the engine automatically on the client and the server whenever the value of this variable changes. You're free to implement this function however you'd like, as needed by your game. ![](onrep.png) |
 
 [region:note]
 Many of the variables in the engine's built-in classes already have replication enabled, so that many features work automatically in a multiplayer context.
diff --git a/Engine/Documentation/Source/Gameplay/Tags/Tags.INT.udn b/Engine/Documentation/Source/Gameplay/Tags/Tags.INT.udn
--- a/Engine/Documentation/Source/Gameplay/Tags/Tags.INT.udn
+++ b/Engine/Documentation/Source/Gameplay/Tags/Tags.INT.udn
@@ -4,7 +4,7 @@ Crumbs: %ROOT%
 Description: Gameplay Tags can be used to identify, categorize, match, and filter objects.
 version: 4.15
 skilllevel:Intermediate
-Parent:Engine/Gameplay
+Parent:Gameplay
 Order:
 type:overview
 tags:Gameplay
diff --git a/Engine/Documentation/Source/Gameplay/Tools/NetworkProfiler/NetworkProfiler.INT.udn b/Engine/Documentation/Source/Gameplay/Tools/NetworkProfiler/NetworkProfiler.INT.udn
--- a/Engine/Documentation/Source/Gameplay/Tools/NetworkProfiler/NetworkProfiler.INT.udn
+++ b/Engine/Documentation/Source/Gameplay/Tools/NetworkProfiler/NetworkProfiler.INT.udn
@@ -17,11 +17,11 @@ You can pass the "networkprofiler=true" command-line argument to record as soon
 
 * netprofile disable: stop recording if currently recording
 
-The data file will be saved to <project directory>/Saved/Profiling/<project name>-<timestamp>.nprof. If you see a file named NetworkProfiling.tmp in the same folder, this is just the temporary file the engine uses while collecting data. When the current profiling session is stopped, this file will be renamed according to the scheme above and is then ready to be opened in the standalone tool.
+The data file will be saved to `<project directory>/Saved/Profiling/<project name>-<timestamp>.nprof`. If you see a file named NetworkProfiling.tmp in the same folder, this is just the temporary file the engine uses while collecting data. When the current profiling session is stopped, this file will be renamed according to the scheme above and is then ready to be opened in the standalone tool.
 
 ## Viewing a profile session in the Network Profiler application
 
-The network profiler is a standalone application located at <UE install directory>/Engine/Binaries/DotNET/NetworkProfiler.exe.
+The network profiler is a standalone application located at `<UE install directory>/Engine/Binaries/DotNET/NetworkProfiler.exe`.
 
 At the top of the window, click the "Open File" button to select a .nprof file and view the data contained within.
 
diff --git a/Engine/Documentation/Source/GettingStarted/HowTo/FindingHelp/FindingHelp.INT.udn b/Engine/Documentation/Source/GettingStarted/HowTo/FindingHelp/FindingHelp.INT.udn
--- a/Engine/Documentation/Source/GettingStarted/HowTo/FindingHelp/FindingHelp.INT.udn
+++ b/Engine/Documentation/Source/GettingStarted/HowTo/FindingHelp/FindingHelp.INT.udn
@@ -125,7 +125,7 @@ When clicking the **Help** button, several options are available in the dropdown
 | Control | Action |
 | ------- | ------ |
 | **Documentation** | This will take you to the main [Official Documentation](#externalhelp:documentation) page. |
-| **API Reference** | This will open the **UE4 API Documentation** and is a low level index of Engine classes and functions.  |
+| **API Reference** | This will open the [UE4 API Documentation](https://docs.unrealengine.com/latest/INT/API/index.html) and is a low level index of Engine classes and functions.  |
 | **Tutorials** | Selecting this option will open the [In-Editor Tutorials](#in-editortutorials). |
 | **Forums** | This will take you to the [Official Unreal Engine 4 Forums](#externalhelp:forums) page. |
 | **Answer Hub** | Selecting this will take you to the main [Answer Hub](#externalhelp:answerhub) page. |
diff --git a/Engine/Documentation/Source/GettingStarted/RunningUnrealEngine/RunningUnrealEngine.INT.udn b/Engine/Documentation/Source/GettingStarted/RunningUnrealEngine/RunningUnrealEngine.INT.udn
--- a/Engine/Documentation/Source/GettingStarted/RunningUnrealEngine/RunningUnrealEngine.INT.udn
+++ b/Engine/Documentation/Source/GettingStarted/RunningUnrealEngine/RunningUnrealEngine.INT.udn
@@ -320,9 +320,7 @@ When running from the executable, you must specify the game you want to run as a
 
 ## Running a Cooked Game
 
-[INCLUDE:Engine/Basics/Projects/Packaging#runcooked]
-
-See [](Engine/Basics/Projects/Packaging) for more information on how to package and run cooked game builds.
+See [](Engine/Basics/Projects/Packaging) for information on how to package and run cooked game builds.
 
 ## Useful In-Game Commands
 
diff --git a/Engine/Documentation/Source/Platforms/GearVR/BestPractices/GearVRBestPractices.INT.udn b/Engine/Documentation/Source/Platforms/GearVR/BestPractices/GearVRBestPractices.INT.udn
--- a/Engine/Documentation/Source/Platforms/GearVR/BestPractices/GearVRBestPractices.INT.udn
+++ b/Engine/Documentation/Source/Platforms/GearVR/BestPractices/GearVRBestPractices.INT.udn
@@ -1,9 +1,13 @@
-Availability:Public
+Availability: Public
+Crumbs: %ROOT%
 Title:Samsung Gear VR Best Practices
-Crumbs:%ROOT%, Platforms
 Description:Information on how to get the most out of your Gear VR headset.
-Navigation:platform
-Platform:Gear VR
+Type: Reference
+Version:4.17
+Parent:Platforms/GearVR
+Order:1
+tags:VR
+tags:GearVR
 Related: Platforms/Mobile/Performance
 Related: Engine/Performance
 Related: Engine/Performance/StatCommands
@@ -11,45 +15,34 @@ Related: Platforms/Mobile/Textures
 Related: Platforms/Mobile/Materials
 Related: Engine/Rendering/Materials/MaterialProperties/BlendModes
 Related: Engine/Rendering/Materials/CustomizedUVs
-Skilllevel: Intermediate
-Version: 4.9
-type:overview
-parent:Platforms/GearVR
-tags:VR
-tags:GearVR
+
 
 [TOC(start:2 end:3)]
 
-When developing content for Virtual Reality (VR) projects, certain steps must be taken to ensure that the user will have the best VR experience possible.
-In the following Best Practices guide we will cover all you need to know to help guarantee that the user has the best possible experience when viewing your project 
-using the Samsung Gear Virtual Reality(VR) Headset and Samsung Galaxy Note 4 Smartphone.
+When developing content for Virtual Reality (VR) projects, certain steps must be taken to ensure that the user will have the best VR experience possible. In the following Best Practices guide we will cover all you need to know to help guarantee that the user has the best possible experience when viewing your project using the Samsung Gear Virtual Reality(VR) Headset and Samsung Galaxy Smartphone.
 
 
 ## Gear VR Rendering Options
-While combining the Gear VR Headset with the Samsung Galaxy Note 4 Smartphone can produce some very immersive VR experiences, 
+While combining the Gear VR Headset with the Samsung Galaxy Smartphone can produce some very immersive VR experiences, 
 there are some rendering features that UE4 can not perform while using the Gear VR Headset due to hardware limitations. 
-Many of the restrictions that you will encounter are similar to restrictions for [mobile devices](Platforms/Mobile/Performance).
-However because of the unique rendering requirements for the Gear VR Headset there are few extra things that must be considered.
-In the following sections we will cover what these requirements are and how they can effect your Gear VR project. 
+Many of the restrictions that you will encounter are similar to limitations for [mobile devices](Platforms/Mobile/Performance), and
+in the following sections, we will cover what these requirements are and how they can affect your Gear VR project. 
 
 ### Mobile HDR
-While you may be tempted to enable Mobile HDR so that you enable more advanced rendering features for your projects but do not do this.
-Mobile HDR will not work with the Gear VR Headset and should be left off. 
-Enabling this feature will cause your project to not be rendered correctly when placed into the Gear VR Headset.
+Gear VR projects do not support Mobile HDR, and it should not be enabled. 
 
 ### Post Processing
-Due to the demanding rendering requirements for the Gear VR Headset, there are many Post Processing features that unfortunately will not work with the Gear VR Headset. 
-Below you will find a list of all the Post Process features that will work with the Gear VR Headset.
+Due to the demanding rendering requirements for Gear VR, many Post Processing features will not work. Below you will find a list of all the Post Process features that will work with projects that are being developed for Gear VR Headset.
 
 |Post Processing Section |Feature Name| Description |
 | -------- | ---------- |------|
-| Auto Exposure|Exposure Bias |Can be used to increase or decrease the scenes brightness. |
+| Auto Exposure|Exposure Bias |Can be used to increase or decrease the brightness of the scene. |
 
 ### Materials
 Creating and interacting with Materials for Gear VR projects is very similar to creating and interacting with Materials for other projects. 
-You still use the Material Editor to create Materials and you can still use many of the Material nodes you find in the Material Editor. 
-However there are a few restrictions / considerations that must be adhered to when making Materials for Gear VR to ensure that what you make will work with the Gear VR Headset. 
-For example only the following Material inputs will work with the Gear VR Headset.
+You still use the Material Editor to create Materials, and you can still use many of the Material nodes you find in the Material Editor. 
+However, there are a few restrictions/considerations that must be adhered to when making Materials for Gear VR to ensure that what you make will work with the Gear VR Headset. 
+For example, only the following Material inputs will work with the Gear VR Headset.
 
 | Material Input Name | Description |
 | -------- | ---------- |
@@ -57,18 +50,17 @@ For example only the following Material inputs will work with the Gear VR Headse
 | **Roughness** |Roughness defines how rough or smooth the Material is. |
 | **Metallic** |Metallic defines how metal-like the Material is. |
 | **Specular** |Specular defines the intensity of the reflections. |
-| **Normal**  |Normal helps to add detail that would otherwise be to costly to model in. |
+| **Normal**  |Normal helps to add detail that would otherwise be too costly to model in. |
 | **Emissive**  |Emissive helps to make parts of Materials glow like they are emitting light. |
-| **Opacity**  |Opacity helps to define how see through something is. |
+| **Opacity**  |Opacity helps to define how to see through something is. |
 | **Opacity Mask** |Opacity Mask helps to define if a certain pixel is see through or not. |
 
 [REGION:note]
-Please keep in mind that the both Translucent and Masked Materials can be very expensive to render, especially on hardware like Smartphone's.
-Whenever possible you should use [Opaque Materials](Engine/Rendering/Materials/MaterialProperties/BlendModes) as they are much cheaper to render.
+Please keep in mind that the both Translucent and Masked Materials can be very expensive to render, especially on hardware like Smartphones. Whenever possible, you should use [Opaque Materials](Engine/Rendering/Materials/MaterialProperties/BlendModes) as they are much cheaper to render.
 [/region]
 
 ### Material Shading Models
-Due to hardware limitations only the following **Material Shading Models** are available for use in projects that use the Gear VR Headset.
+Due to hardware limitations only the following **Material Shading Models** are available for use in Gear VR projects.
 | Shading Model Name | Description |
 | -------- | ---------- |
 | **Default Lit** |Default Lit is, as the name suggests, the default shading model and will likely be the one that you will use most often. This shading model makes use of direct and indirect lighting, as well as specularity for reflections.. |
@@ -80,14 +72,14 @@ Try to use the Unlit Shading Model as much as possible with your Gear VR project
 
 ### Material Blend Modes
 Material Blend Modes describe how the output of the current Material will blend over what is already being drawn in the background. 
-When creating content for the Gear VR Headset the following Material Blend Modes are available for use.
+When creating content for the Gear VR Headset, the following Material Blend Modes are available for use.
 
 | Blend Mode | Description |
 | ---- | ----- |
 | **Opaque** | For a Solid object, Opaque is the cheapest material. But, as its name implies, it has no support for translucency. That means it will never create overdraw, and rendering one pixel of the material will always take the same time depending on the complexity of the **Node Graph**. |
-| **Masked** |  If you need some pixels not drawn on your material, this is the way to go. The material only creates overdraw for the area that is masked out. But the material is 100\% opaque or 100\% see through. The bigger the area that is masked out the more overdraw is created. That means more expensive. |
+| **Masked** |  If you need some pixels not drawn on your material, this is the way to go. The material only creates overdraw for the area that is masked out. But the material is 100\% opaque, or 100\% see through. The bigger the area that is masked out, the more overdraw is created. That means more expensive. |
 | **Translucent** | This is the most expensive **BlendMode**. It can look like a Masked and Opaque but will have an additional effect on performance regardless of actually being transparent or not. |
-| **Additive** | Is a cheaper alternative to Translucent. It does not need to sort the pixels, they just get added to the scene. But that means it can only brighten. Black pixels on the material have no effect on the scene.|
+| **Additive** | Is a cheaper alternative to Translucent. It does not need to sort the pixels, they just get added to the scene. But that means it can only brighten. Black pixels on the material do not affect the scene.|
 | **Modulated** | Is the second cheapest alternative to Translucent, but it multiplies its Pixels with the scene. That means it can only darken the scene.|
 [region:note]
 When using the Modulate blend mode make sure to disable **Separate Translucency** or your Material will not compile correctly.
@@ -99,31 +91,30 @@ On the **Main Material Node** under the **Details Panel** you will find a sectio
 The settings contained here are for helping to lessen Material rendering costs on mobile devices.  
 | Property | Description |
 | -------- | ---------- |
-| **Fully Rough** | When enabled, saves a number of shader ALU instructions and one sampler.  This effectively disables Specular while maintaining overall reflective energy. |
+| **Fully Rough** | When enabled, saves some shader ALU instructions and one sampler.  Enabling this will effectively disable Specular while maintaining overall reflective energy. |
 | **Use Lightmap Directionality** | When disabled, lighting from light maps will be flat but cheaper. |
 
 ### Customized UVs
 [Customized UVs](Engine/Rendering/Materials/CustomizedUVs) allow you to perform certain shader operations in vertex shaders instead of pixel shaders to help avoid dependent texture fetches. 
-This offers an excellent way to claw back much needed performance when doing simple things like UV manipulation to make a texture title more than it currently does. 
-Whenever possible make sure to do all UV manipulation using Customized UVs especially when developing content for the Gear VR Headset.
+Using Customized UVs offers an excellent way to claw back much-needed performance when doing simple things like UV manipulation to make a texture title more than it currently does. Whenever possible make sure to do all UV manipulation using Customized UVs especially when developing content for the Gear VR Headset.
 
 [region:note]
 Customized UVs are currently only implemented for certain component types such as Static Meshes, Skeletal Meshes, BSP, Landscape, and Mesh Particles.
 [/region]
 
 ### Reflections
-The Gear VR does not support Screen Space Reflections(SSR) which means you will only be able generate reflections that come from [Reflection Capture Actors](Resources/ContentExamples/Reflections/1_4/) that are placed in levels. 
-Also keep in mind that the reflections that are generated from the reflections actors will not be of very high fidelity due to hardware limitations. 
+The Gear VR does not support Screen Space Reflections(SSR) which means you will only be able to generate reflections that come from [Reflection Capture Actors](Resources/ContentExamples/Reflections/1_4/) that are placed in levels. 
+Also, keep in mind that the reflections that are generated from the reflections actors will not be of very high fidelity due to hardware limitations. 
 
 ## Gear VR & Content Setup
-When creating art content for your Gear VR projects there are few things that you need to keep in mind to ensure that your project can run as smooth as possible on the Gear VR Headset. 
-In the following section we will cover what you need to be concerned with when making art content for the Gear VR Headset.
+When creating art content for your Gear VR projects, there are few things that you need to keep in mind to ensure that your project can run as smooth as possible on the Gear VR Headset. 
+In the following section, we will cover what you need to be concerned with when making art content for the Gear VR Headset.
 
 ### Content Hard limits
-Due to certain hardware restrictions on mobile devices the following restrictions must be adhered to when creating content. 
-Failing to follow these restrictions will result in your project not working.
+Due to certain hardware restrictions on mobile devices, the following restrictions must be adhered to when creating content. 
+Failing to follow these restrictions will result in your Gear VR project not working.
 
-* All mesh types can have up to but not over 65,000 vertices due to the lack of 32 bit index support on mobile hardware.
+* All mesh types can have up to but not over 65,000 vertices due to the lack of 32-bit index support on mobile hardware.
 * Skeletal Meshes can have up to but not over 75 bones.
 
 ### Mesh Considerations
@@ -133,30 +124,29 @@ When creating Meshes to use with your Gear VR Project the following information
 * Make sure to make use of [Static Mesh LODs](Engine/Content/FBX/StaticMeshes) whenever possible.
 * Keep an eye on your 3D meshes polygon count and try and keep it as [Low Poly](https://en.wikipedia.org/wiki/Low_poly) as possible.
 * Make sure that all static models have a 2nd UV set so they can make use of [Light Mapping](Engine/Content/Types/StaticMeshes/LightmapUnwrapping).
-* Whenever possible try and make use of [Imposter Sprites](Engine/Content/Tools/RenderToTextureTools/3) to replace far away Static Meshes as they are much cheaper to render.
+* Whenever possible try and take advantage of [Imposter Sprites](Engine/Content/Tools/RenderToTextureTools/3) to replace far away Static Meshes as they are much cheaper to render.
 
 ### Texture Considerations
 When creating Textures to use with your Gear VR Project the following information needs to be kept in mind to ensure that what you are creating will not be detrimental to your projects overall performance. 
 
-* Textures use the [ECT2](http://en.wikipedia.org/wiki/Ericsson_Texture_Compression) Texture compression format.
 * Textures can not be bigger than 2048 pixels in either the X or Y dimension.
 * Textures sizes must be in a power of 2 (For example, 2, 4, 8, 16, 32, 64, 128, 256, 512, 1024, 2048)
-* Loading up to 512 MB of Textures into memory is possible with the Gear VR however it is recommended that you do not load more that 128 MB of Textures into memory.
-* Use square Textures whenever possible as them make a more efficient use of memory. 
+* Loading up to 512 MB of Textures into memory is possible with the Gear VR however it is recommended that you do not load more than 128 MB of Textures into memory.
+* Use square Textures whenever possible as they make a more efficient use of memory. 
 * Wherever possible try and use [Texture Atlasing](http://en.wikipedia.org/wiki/Texture_atlas) to reduce the number of Textures needed in memory.
-* [RGB Texture Masking / Packing](Engine/Rendering/Materials/HowTo/Masking) is a great way to cut down on unique Textures in memory.
+* [RGB Texture Masking / Packing](Engine/Rendering/Materials/HowTo/Masking) is an excellent way to cut down on unique Textures in memory.
 * Only **TC_Default** and **TC_NormalMap** compression settings should be used as they are the only two that will work on mobile platforms. 
 
 ### Level Considerations
 Getting the maximum performance out of your levels on the Gear VR is a delicate balancing act that will require time to get just right. 
-The following suggestions / things to look out for can help you quickly identify areas of your levels that might cause performance issues when viewed on the Gear VR Headset.
+The following suggestions/things to look out for can help you quickly identify areas of your levels that might cause performance issues when viewed on the Gear VR Headset.
 
-* Always make sure that the level lighting has been [built](Engine/Rendering/LightingAndShadows/Lightmass) to ensure you get the best performance possible. If you see a red **Lighting needs to be rebuilt** error message on the screen, this means that your lighting needs to be built.
-* Make sure that all the lights in you level are set to [Static](Engine/Rendering/LightingAndShadows/LightMobility/StaticLights) as the other lighting options can be very expensive to render.
+* Always make sure that the level lighting has been [built](Engine/Rendering/LightingAndShadows/Lightmass) to ensure you get the best performance possible. If you see a red **Lighting needs to be rebuilt** error message on the screen, this means that your light needs to be rebuilt.
+* Make sure that all the lights in you level are set to [Static](Engine/Rendering/LightingAndShadows/LightMobility/StaticLights) as the other lighting options can be costly to render.
 * Check to see if your level is making use of [Precomputed Visibility](Engine/Actors/Volumes) and if not make sure to add one to the level and re-build lighting to enable it in the level's World Settings. 
 * Draw calls for the entire scene should not be greater than 100 for any view. This can be seen with [stat openglrhi](Engine/Performance/StatCommands) on the device, [stat d3d11rhi](Engine/Performance/StatCommands) in ES2 previewer on PC.
-* The total triangle count of the level should not be greater than 100,000 for any view. This can be seen with [stat openglrhi](Engine/Performance/StatCommands) on the device, [stat d3d11rhi](Engine/Performance/StatCommands) in ES2 previewer on PC.
-* Be very careful in how Translucent and Opacity Masked Materials are used in the level as they are very expensive to render.
+* The total triangle count of the level should not be greater than 100,000 from any view. This can be seen with [stat openglrhi](Engine/Performance/StatCommands) on the device, [stat d3d11rhi](Engine/Performance/StatCommands) in ES2 previewer on PC.
+* Be very careful in how Translucent and Opacity Masked Materials are used in the level as they are costly to render.
 * Make sure to [Performance Profile](Engine/Performance) your level often to make sure you catch any performance issues sooner rather than later.
 
 
