Copying //UE4/Dev-Documentation to //UE4/Samples-Main (Source: //UE4/Dev-Documentation @ 3623665)

diff --git a/Engine/Documentation/Source/Platforms/GearVR/Debugging/GearVRDebugging.INT.udn b/Engine/Documentation/Source/Platforms/GearVR/Debugging/GearVRDebugging.INT.udn
--- a/Engine/Documentation/Source/Platforms/GearVR/Debugging/GearVRDebugging.INT.udn
+++ b/Engine/Documentation/Source/Platforms/GearVR/Debugging/GearVRDebugging.INT.udn
@@ -1,9 +1,14 @@
 Availability:Public
+Crumbs: %ROOT%
 Title:Samsung Gear VR Debugging
 Description:Information on how to debug applications for the Gear VR headset.
-Navigation:platform
-Crumbs:%ROOT%, Platforms, Platforms/GearVR
-Platform:Gear VR
+Type: Overview
+SkillLevel: Advanced
+Version: 4.17
+Parent: Platforms/GearVR
+Order:3
+tags:VR
+tags:GearVR
 Related: Platforms/Mobile/Performance
 Related: Engine/Performance
 Related: Engine/Performance/StatCommands
@@ -11,77 +16,70 @@ Related: Platforms/Mobile/Textures
 Related: Platforms/Mobile/Materials
 Related: Engine/Rendering/Materials/MaterialProperties/BlendModes
 Related: Engine/Rendering/Materials/CustomizedUVs
-Skilllevel: Intermediate
-Version: 4.9
-type:overview
-parent:Platforms/GearVR
-tags:VR
-tags:GearVR
-
 [TOC(start:2 end:3)]
 
 Being able to debug your project on its target device is critical to tracking down performance issues and other possible project breaking errors. 
 In the following section, you will learn what tools are available to assist you in debugging your Gear VR projects and how to use them.
 
 ### Enabling Gear VR Developer Mode
-Having the ability to view your Gear VR Project without having to use the Gear VR Headset can greatly speed up the time it takes to debug an issue by allowing you quick and easy access to things like Unreal's console. 
-In the following sections we will cover how you can set up your Gear VR projects so that they can be viewed without the Gear VR Headset. 
+Having the ability to view your Gear VR Project without having to use the Gear VR Headset can significantly speed up the time it takes to debug an issue by allowing you quick and easy access to things like Unreal's console. 
+In the following sections, we will cover how you can set up your Gear VR projects so that they can be viewed without the Gear VR Headset. 
 
 [region:note]
-Before you can do this, you must make sure that you have some type of Gear VR application installed to your phone that has your device signature in it. 
-This can be any kind of application (download or compiled).
+Before you can do this, you must make sure that you have some Gear VR application installed to your phone that has your device signature in it. 
+This can be any application (download or compiled).
 [/region]
 
 1. First, on your Samsung Galaxy Note 4 Smartphone go to **Settings** > **Application Manager** > **Gear VR Service**.
-	
-	[REGION:lightbox]
-	[![](T_GearVR_Dev_Mode_01.png)(w:600 h:355)](T_GearVR_Dev_Mode_01.png)
-	[/REGION]
-	
-	[REGION:caption]
-	Click for full image.
-	[/REGION]
+    
+    [REGION:lightbox]
+    [![](T_GearVR_Dev_Mode_01.png)(w:600 h:355)](T_GearVR_Dev_Mode_01.png)
+    [/REGION]
+    
+    [REGION:caption]
+    Click for full image.
+    [/REGION]
 
 1. Next, tap on the **Manage Storage** option.
   
-	[REGION:lightbox]
-	[![](T_GearVR_Dev_Mode_02.png)(w:250 h:444)](T_GearVR_Dev_Mode_02.png)
-	[/REGION]
-	
-	[REGION:caption]
-	Click for full image.
-	[/REGION]
+    [REGION:lightbox]
+    [![](T_GearVR_Dev_Mode_02.png)(w:250 h:444)](T_GearVR_Dev_Mode_02.png)
+    [/REGION]
+    
+    [REGION:caption]
+    Click for full image.
+    [/REGION]
 
 1. Then, tap on the **VR Service Version** **6** times.
 
-	[REGION:lightbox]
-	[![](T_GearVR_Dev_Mode_03.png)(w:250 h:225)](T_GearVR_Dev_Mode_03.png)
-	[/REGION]
-	
-	[REGION:caption]
-	Click for full image.
-	[/REGION]
+    [REGION:lightbox]
+    [![](T_GearVR_Dev_Mode_03.png)(w:250 h:225)](T_GearVR_Dev_Mode_03.png)
+    [/REGION]
+    
+    [REGION:caption]
+    Click for full image.
+    [/REGION]
 
 1. You should now see a scan process happen and when completed you should see the **Developer Mode** toggle on the screen.
 
-	[REGION:lightbox]
-	[![](T_GearVR_Dev_Mode_04.png)(w:250 h:225)](T_GearVR_Dev_Mode_04.png)
-	[/REGION]
-	
-	[REGION:caption]
-	Click for full image.
-	[/REGION]
+    [REGION:lightbox]
+    [![](T_GearVR_Dev_Mode_04.png)(w:250 h:225)](T_GearVR_Dev_Mode_04.png)
+    [/REGION]
+    
+    [REGION:caption]
+    Click for full image.
+    [/REGION]
 
 1. Now, you should be able to launch your Gear VR Project and have it display without having to use the Gear VR Headset.
 
-	[region:tip]
-	Be careful when enabling VR Developer mode (it can drastically drain your phone's battery) as you are overriding any settings, to turn off the Gear VR Headset when not in use.
-	[/region]
+    [region:tip]
+    Be careful when enabling VR Developer mode (it can drastically drain your phone's battery) as you are overriding any settings, to turn off the Gear VR Headset when not in use.
+    [/region]
 
 1. If for any reason you need to disable VR Developer mode, just slide the **Developer Mode** slider from **On** to **Off** and VR Developer mode will be disabled.
 
 ### Enabling the UE4 Console
-The UE4 Console is a very powerful tool that allows you to look at, and in some cases, affect data that is currently being displayed in UE4. 
+The UE4 Console is a potent tool that allows you to look at, and in some cases, affect data that is currently being displayed in UE4. 
 Accessing the UE4 console on Mobile devices requires a slightly different approach than what you might be used to on PC and is explained below.
 
 [region:tip]
@@ -89,129 +87,129 @@ You can get a full list of available console commands by entering **DumpConsoleC
 [/region]
 
 1. First, make sure that you have enabled Gear VR Developer mode on your phone. 
-	This way you can launch your project and see it without having to use the Gear VR Headset which, if used, would prevent you from being able to touch your phone's screen.
-	
-	[REGION:lightbox]
-	[![](T_GearVR_Dev_Mode_04.png)(w:250 h:225)](T_GearVR_Dev_Mode_04.png)
-	[/REGION]
-	
-	[REGION:caption]
-	Click for full image.
-	[/REGION]
-
-1. Next, launch your Gear VR project by tapping on the project's icon to launch it. 
-	When the project has been launched, you should see something looking like this on your phone.
-	
-	[REGION:lightbox]
-	[![](T_App_On_Phone.png)(w:500 h:281)](T_App_On_Phone.png)
-	[/REGION]
-	
-	[REGION:caption]
-	Click for full image.
-	[/REGION]
+    This way you can launch your project and see it without having to use the Gear VR Headset which, if used, would prevent you from being able to touch your phone's screen.
+    
+    [REGION:lightbox]
+    [![](T_GearVR_Dev_Mode_04.png)(w:250 h:225)](T_GearVR_Dev_Mode_04.png)
+    [/REGION]
+    
+    [REGION:caption]
+    Click for full image.
+    [/REGION]
+
+1. Next, launch your Gear VR project by tapping on the project's icon to start it. 
+    When the project has been started, you should see something looking like this on your phone.
+    
+    [REGION:lightbox]
+    [![](T_App_On_Phone.png)(w:500 h:281)](T_App_On_Phone.png)
+    [/REGION]
+    
+    [REGION:caption]
+    Click for full image.
+    [/REGION]
 
 1. With your Gear VR project now running, tap the phone's screen with **Four (4)** fingers at once to open up the Unreal Engine Console. 
-	If done correctly, you should see something like this on your phone's screen. 
+    If done correctly, you should see something like this on your phone's screen. 
 
-	[REGION:lightbox]
-	[![](T_GearR_Open_Console.png)(w:500 h:281)](T_GearR_Open_Console.png)
-	[/REGION]
-	
-	[REGION:caption]
-	Click for full image.
-	[/REGION]
+    [REGION:lightbox]
+    [![](T_GearR_Open_Console.png)(w:500 h:281)](T_GearR_Open_Console.png)
+    [/REGION]
+    
+    [REGION:caption]
+    Click for full image.
+    [/REGION]
 
 1. Now input a console command you wish to use in the input box and then press the **Ok** button. 
-	For this example we will use the **Stat Memory** command to see how much memory the project is using. 
+    For this example, we will use the **Stat Memory** command to see how much memory the project is using. 
 
-	[REGION:lightbox]
-	[![](T_GearR_Input_Command.png)(w:500 h:281)](T_GearR_Input_Command.png)
-	[/REGION]
-	
-	[REGION:caption]
-	Click for full image.
-	[/REGION]
+    [REGION:lightbox]
+    [![](T_GearR_Input_Command.png)(w:500 h:281)](T_GearR_Input_Command.png)
+    [/REGION]
+    
+    [REGION:caption]
+    Click for full image.
+    [/REGION]
 
 1. Once the **Ok** button has been pressed, your project should now be displaying how much memory is being used and by what.
 
-	[REGION:lightbox]
-	[![](T_GearR_Displaying_Memory.png)(w:500 h:281)](T_GearR_Displaying_Memory.png)
-	[/REGION]
-	
-	[REGION:caption]
-	Click for full image.
-	[/REGION]
+    [REGION:lightbox]
+    [![](T_GearR_Displaying_Memory.png)(w:500 h:281)](T_GearR_Displaying_Memory.png)
+    [/REGION]
+    
+    [REGION:caption]
+    Click for full image.
+    [/REGION]
 
-1. When you are done debugging, or if you just wish to clear the screen of debug information, you can input **Stat None** to clear the screen of all debug information.
+1. When you are done debugging, or if you just wish to clear the screen of debugging information, you can input **Stat None** to clear the screen of all of the debug information.
 
-	[region:tip]
-	The best option for console debugging with the Gear VR Headset is to use a Bluetooth keyboard. This way, you can use the console without having to remove the phone from the headset.
-	[/region]
+    [region:tip]
+    The best option for console debugging with the Gear VR Headset is to use a Bluetooth keyboard. This way, you can use the console without having to remove the phone from the headset.
+    [/region]
 
 ### Viewing Gear VR Debug Logs
-Having the ability to view debug logs is of critical importance when developing any kind of 3D project. 
+Having the ability to view debug logs is of critical importance when developing any 3D project. 
 Debug logs can provide vital information and insight into why your project is not working and, in some cases, even make suggestions as to what you can do to fix things. 
-In the following section, we will go over what you need to do in order to get debug logs from your Smartphone.
+In the following section, we will go over what you need to do to get debug logs from your Smartphone.
 
 [region:note]
-Check out the official Android documentation for more detailed information on all of various things you can do with [Logcat](http://developer.android.com/tools/help/logcat.html).
+Check out the official Android documentation for more detailed information on all of the various things you can do with [Logcat](http://developer.android.com/tools/help/logcat.html).
 [/region]
 
 1. First, make sure that your Smartphone with the application you want to get debug information from is plugged into your PC via USB.
 
 1. Next, open up a Windows Command Shell by first pressing the **Windows Key + R** and then entering **CMD** in the **Run** dialogue box that is opened.
-	Then press either the **Enter** key on your keyboard or just click on the **OK** button with your mouse. 
-	 
-	 ![](T_Open_CMD_Prompt.png)
+    Then press either the **Enter** key on your keyboard or just click on the **OK** button with your mouse. 
+     
+     ![](T_Open_CMD_Prompt.png)
 
 1. Once the **OK** button has been toggled, a Windows Command Line window like the following should now be displayed.
-	
-	[REGION:lightbox]
-	[![](T_Windows_CMD.png)(w:422 h:269)](T_Windows_CMD.png)
-	[/REGION]
-	
-	[REGION:caption]
-	Click for full image.
-	[/REGION]
+    
+    [REGION:lightbox]
+    [![](T_Windows_CMD.png)(w:422 h:269)](T_Windows_CMD.png)
+    [/REGION]
+    
+    [REGION:caption]
+    Click for full image.
+    [/REGION]
 
 1. Now, check to make sure that your PC can see your Smartphone by entering **ADB Devices** into the command line window and then press the **Enter** key to display devices attached to your PC.
-	
-	[REGION:lightbox]
-	[![](T_Windows_ADB_Devices.png)(w:422 h:269)](T_Windows_ADB_Devices.png)
-	[/REGION]
-	
-	[REGION:caption]
-	Click for full image.
-	[/REGION]
-
-	[region:tip]
-	If you enter **ADB Devices** but do not see anything show up, make sure that you have gone through the [](Platforms/GearVR) setup section and installed / setup everything that is listed there and then try this step again. 
-	[/region]
+    
+    [REGION:lightbox]
+    [![](T_Windows_ADB_Devices.png)(w:422 h:269)](T_Windows_ADB_Devices.png)
+    [/REGION]
+    
+    [REGION:caption]
+    Click for full image.
+    [/REGION]
+
+    [region:tip]
+    If you enter **ADB Devices** but do not see anything show up, make sure that you have gone through the [](Platforms/GearVR) setup section and installed / setup everything that is listed there and then attempt this step again. 
+    [/region]
 
 1. Now that you have verified that you can see your Smartphone in ADB Devices, it is now time to get the log information from your phone. 
-	But firs,t we are going to clear the log cache to ensure that we are getting only current debug information and not old debug information. 
-	To clear the log, enter the following command.
-
-		adb logcat -c
-	[REGION:lightbox]
-	[![](T_ADB_Clear_Log.png)(w:422 h:269)](T_ADB_Clear_Log.png)
-	[/REGION]
-	
-	[REGION:caption]
-	Click for full image.
-	[/REGION]
+    But firs,t we are going to clear the log cache to ensure that we are getting the only current debug information and not old debug information. 
+    To clear the log, enter the following command.
+
+        adb logcat -c
+    [REGION:lightbox]
+    [![](T_ADB_Clear_Log.png)(w:422 h:269)](T_ADB_Clear_Log.png)
+    [/REGION]
+    
+    [REGION:caption]
+    Click for full image.
+    [/REGION]
 
 1. With the log now cleared of old data, launch the application on your Smartphone and enter the following command to get the log information.
-	
-		adb logcat -s UE4 OVR VrApi VrLib VRSVC HmdInfo >logoutput.txt
+    
+        adb logcat -s UE4 OVR VrApi VrLib VRSVC HmdInfo >logoutput.txt
 
-	[REGION:lightbox]
-	[![](T_ADB_Get_Log.png)(w:422 h:269)](T_ADB_Get_Log.png)
-	[/REGION]
-	
-	[REGION:caption]
-	Click for full image.
-	[/REGION]
+    [REGION:lightbox]
+    [![](T_ADB_Get_Log.png)(w:422 h:269)](T_ADB_Get_Log.png)
+    [/REGION]
+    
+    [REGION:caption]
+    Click for full image.
+    [/REGION]
 
  [region:note]
  To see the live output of the log, omit the **>logoutput.txt** part and all log updates will happen in the command line window.
@@ -219,31 +217,31 @@ Check out the official Android documentation for more detailed information on al
 
 1. To complete the log capture, press **CTRL + C** on the keyboard and when completed, you should now have a log file named,** logoutput.txt**, that is saved to your **Users** folder in Windows.
 
-	[REGION:lightbox]
-	[![](T_ADB_Log_Location.png)(w:381 h:338)](T_ADB_Log_Location.png)
-	[/REGION]
-	
-	[REGION:caption]
-	Click for full image.
-	[/REGION]
+    [REGION:lightbox]
+    [![](T_ADB_Log_Location.png)(w:381 h:338)](T_ADB_Log_Location.png)
+    [/REGION]
+    
+    [REGION:caption]
+    Click for full image.
+    [/REGION]
 
-	[region:note]
-	Data will still continue to be output to the log file while the command line window is open and running the ADB logcat command. To stop this, all you have to do is press **CTRL + C** and data will stop being written to the log file.
-	[/region] 
+    [region:note]
+    Data will continue to be output to the log file while the command line window is open and running the ADB logcat command. To stop this, all you have to do is press **CTRL + C** and data will cease being written to the log file.
+    [/region] 
 
 1. The information that is contained within the log file was generated from adding **UE4**, **OVR**,  **VrApi**, **VrLib**,**VRSVC**, and **HmdInfo** to the logcat command. 
-	When you look at the log file that was created and search for any of those terms, you will find all debug information related to those terms. 
-	You can see an example of this in the image below, some of the terms used (OVR, UE4, Vrapi, and VrApi) show up in the log with information that is related to those terms.
-	
-	[REGION:lightbox]
-	[![](T_ADB_Log_Info.png)(w:516 h:349)](T_ADB_Log_Info.png)
-	[/REGION]
-	
-	[REGION:caption]
-	Click for full image.
-	[/REGION]
-
-If your Gear VR Project is not starting up for some reason, using the logcat command is a great way to see exactly where your project is failing and why. 
+    When you look at the log file that was created and search for any of those terms, you will find all debug information related to those terms. 
+    You can see an example of this in the image below, some of the terms used (OVR, UE4, Vrapi, and VrApi) show up in the log with information that is related to those terms.
+    
+    [REGION:lightbox]
+    [![](T_ADB_Log_Info.png)(w:516 h:349)](T_ADB_Log_Info.png)
+    [/REGION]
+    
+    [REGION:caption]
+    Click for full image.
+    [/REGION]
+
+If your Gear VR Project is not starting up for some reason, using the logcat command is an excellent way to see exactly where your project is failing and why. 
 It also provides you with detailed information about the problem that can then be posted to the UE4 Forums or Answer Hub so that you can get help with the issue.
 
 ### Wireless Deploying & Debugging
@@ -253,72 +251,72 @@ You can even deploy projects in this manner; however, it is extremely slow when
 In the following section, we will cover all you need to know about setting up your Smartphone so that it can send debug information and be deployed over wifi.
 
 [region:note]
-Please note that in order for this to work, you have to have your development PC and Galaxy Note 4 on the same network.
+Please note that for this to work, you have to have your development PC and Galaxy Note 4 on the same network.
 [/region]
 
 1. Next, open up a Windows Command Shell by first pressing the **Windows Key + R** and then entering **CMD** in the **Run** dialogue box that is opened.
-	Then press either the **Enter** key on your keyboard or just click on the **OK** button with your mouse. 
-	 
-	[REGION:lightbox]
-	[![](T_Open_CMD_Prompt.png)(w:427 h:269)](T_Open_CMD_Prompt.png)
-	[/REGION]
-	
-	[REGION:caption]
-	Click for full image.
-	[/REGION]
+    Then press either the **Enter** key on your keyboard or just click on the **OK** button with your mouse. 
+     
+    [REGION:lightbox]
+    [![](T_Open_CMD_Prompt.png)(w:427 h:269)](T_Open_CMD_Prompt.png)
+    [/REGION]
+    
+    [REGION:caption]
+    Click for full image.
+    [/REGION]
 
 1. Once the **OK** button has been toggled, a Windows Command Line window like the following should now be displayed.
-	
-	[REGION:lightbox]
-	[![](T_Windows_CMD.png)(w:427 h:269)](T_Windows_CMD.png)
-	[/REGION]
-	
-	[REGION:caption]
-	Click for full image.
-	[/REGION]
+    
+    [REGION:lightbox]
+    [![](T_Windows_CMD.png)(w:427 h:269)](T_Windows_CMD.png)
+    [/REGION]
+    
+    [REGION:caption]
+    Click for full image.
+    [/REGION]
 
 
 1. Now, enter the following command to see what your Smartphone's IP address is and make sure to take note of it. The IP address you want to use is the one that comes after **inet**.
-	
-		adb shell ip -f inet addr show wlan0
-	
-	[REGION:lightbox]
-	[![](T_ADB_Phone_IP.png)(w:427 h:269)](T_ADB_Phone_IP.png)
-	[/REGION]
-	
-	[REGION:caption]
-	Click for full image.
-	[/REGION]
-
-1. Next, enter the following command to change the network connection to **TCP** and set the port to listen to.
-
-		adb tcpip 5555
-	
-	[REGION:lightbox]
-	[![](T_ADB_Set_Port.png)(w:427 h:269)](T_ADB_Set_Port.png)
-	[/REGION]
-	
-	[REGION:caption]
-	Click for full image.
-	[/REGION]
+    
+        adb shell ip -f inet addr show wlan0
+    
+    [REGION:lightbox]
+    [![](T_ADB_Phone_IP.png)(w:427 h:269)](T_ADB_Phone_IP.png)
+    [/REGION]
+    
+    [REGION:caption]
+    Click for full image.
+    [/REGION]
+
+1. Next, enter the following command to change the network connection to **TCP** and set the port that should be listening to.
+
+        adb tcpip 5555
+    
+    [REGION:lightbox]
+    [![](T_ADB_Set_Port.png)(w:427 h:269)](T_ADB_Set_Port.png)
+    [/REGION]
+    
+    [REGION:caption]
+    Click for full image.
+    [/REGION]
 
 1. Once the port has been set, it is now time to connect to the Smartphone. To do this, all you need to do is enter the following command and press the Enter key.
-	
-		adb connect 10.3.31.133 - Replace with your devices IP Address.
-		
+    
+        adb connect 10.3.31.133 - Replace with your devices IP Address.
+        
 1. You should now see a message that says **connected to** with the IP address of your device.
-	
-	![](T_ADB_Connected.png)
+    
+    ![](T_ADB_Connected.png)
 
-1. To make sure you are connected over Wifi, enter **adb devices** into the command line window and press the Enter key and you will now see that in the **List of devices attached** the IP Address used above is now the only attached device.
+1. To make sure you are connected to Wifi, enter **adb devices** into the command line window and press the Enter key, and you will now see that in the **List of devices attached** the IP Address used above is now the only attached device.
 
-	[REGION:lightbox]
-	[![](T_ADB_Attached_IP.png)(w:427 h:269)](T_ADB_Attached_IP.png)
-	[/REGION]
-	
-	[REGION:caption]
-	Click for full image.
-	[/REGION]
+    [REGION:lightbox]
+    [![](T_ADB_Attached_IP.png)(w:427 h:269)](T_ADB_Attached_IP.png)
+    [/REGION]
+    
+    [REGION:caption]
+    Click for full image.
+    [/REGION]
 
 
 With your Smartphone and PC now connected to one another via the network, you can now begin to get debug information from your Smartphone or even deploy projects to your phone without having to use the USB cable. 
@@ -326,22 +324,16 @@ However, keep in mind that deploying to your Smartphone from your PC over the ne
 
 ### Return to Wired Deploying & Debugging
 1. To return to using the USB cable for debugging and deployment, input the following command into the command window.
-	
-		adb usb
-
-	[REGION:lightbox]
-	[![](T_ADB_Detached.png)(w:427 h:269)](T_ADB_Detached.png)
-	[/REGION]
-	
-	[REGION:caption]
-	Click for full image.
-	[/REGION]
-
-
-
-
-
-
+    
+        ADB USB
+
+    [REGION:lightbox]
+    [![](T_ADB_Detached.png)(w:427 h:269)](T_ADB_Detached.png)
+    [/REGION]
+    
+    [REGION:caption]
+    Click for full image.
+    [/REGION]
 
 
 
diff --git a/Engine/Documentation/Source/Platforms/GearVR/GearVR.INT.udn b/Engine/Documentation/Source/Platforms/GearVR/GearVR.INT.udn
--- a/Engine/Documentation/Source/Platforms/GearVR/GearVR.INT.udn
+++ b/Engine/Documentation/Source/Platforms/GearVR/GearVR.INT.udn
@@ -4,7 +4,7 @@ Availability:Public
 Navigation:platform
 Platform:Gear VR
 Skilllevel: Intermediate
-Version: 4.9
+Version: 4.17
 prereq: Engine\Blueprints\GettingStarted
 Crumbs:%ROOT%, Platforms
 type:landing
@@ -14,9 +14,35 @@ tags:GearVR
 tags:Platform
 Topic-image:SamsungGearVRDevelopment_topic.png
 
+
+## Prerequisites 
+
+[DIR(end:"1" output:"topic" type:"reference" tags:"Basics" parent:"Platforms/GearVR")]
+
+
+## Quick Start
+
+[DIR(end:"1" output:"topic" type:"quick start" parent:"Platforms/GearVR")]
+
+## Best Practices
+
+[DIR(end:"1" output:"topic" type:"reference" tags:"GearVR" parent:"Platforms/GearVR")]
+
+## How To
+
+[DIR(end:"1" output:"topic" type:"how-to" tags:"GearVR" parent:"Platforms/GearVR")]
+
+## Gear VR Command Reference
+
+
+
+## Gear VR Profiling & Performance
+
+
+<!--
 ## Starting Out
 
-[DIR(output:"listbutton" type:"quick start" parent:"Platforms/GearVR")]
+
 
 ## GearVR Topics
 
@@ -36,3 +62,4 @@ rNjDsNYvgfk
 [/PARAMLITERAL]
 [/OBJECT]
 
+-->
\ No newline at end of file
diff --git a/Engine/Documentation/Source/Platforms/GearVR/HowTo/CameraSetup/HTGearVRCameraSetup.INT.udn b/Engine/Documentation/Source/Platforms/GearVR/HowTo/CameraSetup/HTGearVRCameraSetup.INT.udn
new file mode 100644
--- /dev/null
+++ b/Engine/Documentation/Source/Platforms/GearVR/HowTo/CameraSetup/HTGearVRCameraSetup.INT.udn
@@ -0,0 +1,244 @@
+Availability: Public
+Crumbs: %ROOT%, Platforms
+Title: Set Up Cameras for Gear VR
+Description: Setting up the UE4 camera to work with the Gear VR.
+Type: how-to
+SkillLevel: Beginner
+Version: 4.17
+Parent: Platforms/GearVR/HowTo
+Order: 2
+Tags: Basics
+topic-image:HT_GearVR_Camera_Setup_Topic_Image.png
+prereq: Platforms\GearVR\Prerequisites
+prereq: Platforms\GearVR\QuickStart
+
+No matter if your Gear VR project is a sitting or standing experience, the manner in which the user's camera is set up will be the same. In the following How - To we will go over how up the UE4 camera for use with the Gear VR HMD.
+
+
+## Project & Basic Camera Setup for a Gear VR Experience 
+
+Below you will find instructions on project setup requirements as well as how to create the needed Pawn Blueprints for Gear VR.
+
+1. First, create a new UE4 project with the following settings and once the setting has been input, press the **Create Project** button to create the project.
+
+    * Blueprint 
+    * Blank Project
+    * Mobile / Tablet
+    * Scalable 3D or 2D
+    * No Starter Content
+    * Location for the project to be saved
+    * Name for the project
+
+    [REGION:lightbox]
+    [![](HT_GearVR_Camera_Setup_00.png)(w:555)](HT_GearVR_Camera_Setup_00.png)
+    [/REGION]
+
+    [REGION:caption]
+    Click for full image.
+    [/REGION]
+
+1. Once the project loads, tight-click in the **Content Browser** and from the displayed menu, go to the **Create Basic Asset** section and select the **Blueprint Class** option.
+
+    [REGION:lightbox]
+    [![](Platforms\Oculus\HowTo\StandingCamera\HT_Rift_Camera_Setup_01.png)(w:150)](Platforms\Oculus\HowTo\StandingCamera\HT_Rift_Camera_Setup_01.png)
+    [/REGION]
+
+    [REGION:caption]
+    Click for full image.
+    [/REGION]
+
+1. From the **Pick Parent Class** window, select the **Pawn** option and then give it a name of **GearVR_Pawn**.
+
+    [REGION:lightbox]
+    [![](HT_GearVR_Camera_Setup_01.png)(w:555)](HT_GearVR_Camera_Setup_01.png)
+    [/REGION]
+
+    [REGION:caption]
+    Click for full image.
+    [/REGION]
+
+1. Open up the GearVR_Pawn Blueprint and go to the **Component** section of the **Viewport** tab. Using the **Add Component** button add the following two components with the following names and also make sure that the VRCamera is a child of the VRCameraRoot.
+    
+    * Scene - VRCameraRoot
+    * Camera - VRCamera
+
+    ![](HT_GearVR_Camera_Setup_02.png)
+
+    [region:note]
+    No matter what VR HMD you are using, this is the Epic recommend way to set up a VR camera as it allows you to offset the position of the camera without having to move the actual camera.
+    [/region]
+
+Now that a basic Blueprint for the VR Player has been created, in the next section, we will go through setting up this Blueprint to be used with Gear VR.
+
+## Camera Setup for Gear VR Experience 
+
+Below you will find instructions on how to setup your camera for a Gear VR experience.
+
+
+1. First, open up the **GearVR_Pawn** Blueprint and in the **Event Graph** drag off the **Event Begin Play** node to show the Executable Actions list and search for the **Set Tracking Origin** node, clicking on it to add it to the Event Graph.
+
+    [REGION:lightbox]
+    [![](Platforms\Oculus\HowTo\StandingCamera\HT_Rift_Camera_Setup_09.png)(w:555)](Platforms\Oculus\HowTo\StandingCamera\HT_Rift_Camera_Setup_09.png)
+    [/REGION]
+
+    [REGION:caption]
+    Click for full image.
+    [/REGION]
+
+1. Make sure the **Origin** of the Set Tracking Origin node is set to **Eye Level**.
+    
+    [REGION:lightbox]
+    [![](Platforms\Oculus\HowTo\StandingCamera\HT_Rift_Camera_Setup_10.png)(w:555)](Platforms\Oculus\HowTo\StandingCamera\HT_Rift_Camera_Setup_10.png)
+    [/REGION]
+
+    [REGION:caption]
+    Click for full image.
+    [/REGION]
+
+1. In the Blueprint editor under the **Variables** section, create a new **Vector** variable called **CameraHeight** and set the **Z** value to **180**.
+
+    ![](HT_GearVR_Camera_Setup_03.png)
+
+    [region:note]
+    This number set in CameraHeight should be set to the real world standing height of the user in Centimeters.
+    [/region]
+
+1. Drag off the output of the Set Tracking Origin node and search for the **Set Relative Location** node, selecting the **SetRelativeLocation(VRCameraRoot)** option.
+
+    [REGION:lightbox]
+    [![](HT_GearVR_Camera_Setup_04.png)(w:342)](HT_GearVR_Camera_Setup_04.png)
+    [/REGION]
+
+    [REGION:caption]
+    Click for full image.
+    [/REGION]
+
+1. Connect the **CameraHeight** variable to the **New Location** input on the Set Relative Location node and then press the Compile button. When completed your Event Graph should look like the following image.
+
+    [REGION:fullwidth]
+    [INCLUDE:Platforms\GearVR\HowTo\CameraSetup\#GearVRCamera]
+    [/REGION]
+
+    [region:tip]
+    You can copy the completed Blueprint by clicking on the upper left-hand corner of above and pressing CRTL + C to copy the graph. Once copied, go to your Blueprints event and press CTRL + V to paste it in. 
+    [/region]
+
+1. Drag the **GearVR_Pawn** Blueprint from the Content Browser to a level, making sure that it is placed as close to the ground as possible.
+
+    [REGION:lightbox]
+    [![](HT_GearVR_Camera_Setup_06.png)(w:555)](HT_GearVR_Camera_Setup_06.png)
+    [/REGION]
+
+    [REGION:caption]
+    Click for full image.
+    [/REGION]
+
+1. Select the GearVR_Pawn Blueprint and in the **Details** panel under the **Pawn** setting, set the **Auto Possess Player** from **Disabled** to **Player 0**.
+
+    [REGION:lightbox]
+    [![](HT_GearVR_Camera_Setup_07.png)(w:181)](HT_GearVR_Camera_Setup_07.png)
+    [/REGION]
+
+    [REGION:caption]
+    Click for full image.
+    [/REGION]
+
+1. Finally, go to the **Main Toolbar** and click on the small white triangle next to **Launch Options** and then select your Android device from the list to start the packaging and deployment to that device.
+
+    ![](HT_GearVR_Camera_Setup_08.png)
+
+1. Once the UE4 project has deployed, put your Smartphone into your Gear VR HMD and when you move your head around it should look like the following video. 
+
+    [OBJECT:EmbeddedVideo]
+        [PARAMLITERAL:videoid]
+        JdeI2jTS8V8
+        [/PARAMLITERAL]
+        [PARAMLITERAL:width]
+        100
+        [/PARAMLITERAL]
+        [PARAMLITERAL:height]
+        56.25
+        [/PARAMLITERAL]
+        [PARAMLITERAL:units]
+        %
+        [/PARAMLITERAL]
+        [PARAMLITERAL:type]
+        simple
+        [/PARAMLITERAL]
+    [/OBJECT]
+
+## UE4 Project Downloads
+
+Below you will find a link to where you can download the UE4 project that was used to create this example. 
+
+* [**Gear VR Camera Example Project**](https://epicgames.box.com/s/wd0o3iakpi5dt8zt8rxm427148ogmx7w)
+
+<!--
+[EXCERPT:GearVRCamera]	
+[OBJECT:Snippet]
+	[PARAMLITERAL:id]
+		GearVRCamera
+	[/PARAMLITERAL]
+	[PARAMLITERAL:type]
+		blueprint
+	[/PARAMLITERAL]
+	[PARAM:image]
+		[REGION:lightbox]
+		[![](HT_GearVR_Camera_Setup_05.png)(w:1000)](HT_GearVR_Camera_Setup_05.png)
+		[/REGION]
+	[/PARAM]
+	[PARAM:snippet]
+			BEGIN OBJECT
+							Begin Object Class=/Script/BlueprintGraph.K2Node_Event Name="K2Node_Event_0"
+			   EventReference=(MemberParent=Class'/Script/Engine.Actor',MemberName="ReceiveBeginPlay")
+			   bOverrideFunction=True
+			   bCommentBubblePinned=True
+			   NodeGuid=C59E43604E42C7F4DAD48BA1137C7EB9
+			   CustomProperties Pin (PinId=FE39330745F60B4B61B9138BA2E683E1,PinName="OutputDelegate",Direction="EGPD_Output",PinType.PinCategory="delegate",PinType.PinSubCategory="",PinType.PinSubCategoryObject=None,PinType.PinSubCategoryMemberReference=(MemberParent=Class'/Script/Engine.Actor',MemberName="ReceiveBeginPlay"),PinType.PinValueType=(),PinType.bIsMap=False,PinType.bIsSet=False,PinType.bIsArray=False,PinType.bIsReference=False,PinType.bIsConst=False,PinType.bIsWeakPointer=False,PersistentGuid=00000000000000000000000000000000,bHidden=False,bNotConnectable=False,bDefaultValueIsReadOnly=False,bDefaultValueIsIgnored=False,bAdvancedView=False,)
+			   CustomProperties Pin (PinId=CFFCA1684C5D4AA0653641B0BB601A74,PinName="then",Direction="EGPD_Output",PinType.PinCategory="exec",PinType.PinSubCategory="",PinType.PinSubCategoryObject=None,PinType.PinSubCategoryMemberReference=(),PinType.PinValueType=(),PinType.bIsMap=False,PinType.bIsSet=False,PinType.bIsArray=False,PinType.bIsReference=False,PinType.bIsConst=False,PinType.bIsWeakPointer=False,LinkedTo=(K2Node_CallFunction_77 7227CDD94FDBCC1F02D6258FDC30999A,),PersistentGuid=00000000000000000000000000000000,bHidden=False,bNotConnectable=False,bDefaultValueIsReadOnly=False,bDefaultValueIsIgnored=False,bAdvancedView=False,)
+			End Object
+			Begin Object Class=/Script/BlueprintGraph.K2Node_CallFunction Name="K2Node_CallFunction_77"
+			   FunctionReference=(MemberParent=Class'/Script/Engine.HeadMountedDisplayFunctionLibrary',MemberName="SetTrackingOrigin")
+			   NodePosX=224
+			   NodeComment="PS4"
+			   NodeGuid=42D9B736496D2754CE4A38880F53497B
+			   CustomProperties Pin (PinId=7227CDD94FDBCC1F02D6258FDC30999A,PinName="execute",PinToolTip="\nExec",PinType.PinCategory="exec",PinType.PinSubCategory="",PinType.PinSubCategoryObject=None,PinType.PinSubCategoryMemberReference=(),PinType.PinValueType=(),PinType.bIsMap=False,PinType.bIsSet=False,PinType.bIsArray=False,PinType.bIsReference=False,PinType.bIsConst=False,PinType.bIsWeakPointer=False,LinkedTo=(K2Node_Event_0 CFFCA1684C5D4AA0653641B0BB601A74,),PersistentGuid=00000000000000000000000000000000,bHidden=False,bNotConnectable=False,bDefaultValueIsReadOnly=False,bDefaultValueIsIgnored=False,bAdvancedView=False,)
+			   CustomProperties Pin (PinId=6F56A02B45E8A67EEE79388CF61A44BC,PinName="then",PinToolTip="\nExec",Direction="EGPD_Output",PinType.PinCategory="exec",PinType.PinSubCategory="",PinType.PinSubCategoryObject=None,PinType.PinSubCategoryMemberReference=(),PinType.PinValueType=(),PinType.bIsMap=False,PinType.bIsSet=False,PinType.bIsArray=False,PinType.bIsReference=False,PinType.bIsConst=False,PinType.bIsWeakPointer=False,LinkedTo=(K2Node_CallFunction_46 FE5D38794B000B7752514CAFCC5316E4,),PersistentGuid=00000000000000000000000000000000,bHidden=False,bNotConnectable=False,bDefaultValueIsReadOnly=False,bDefaultValueIsIgnored=False,bAdvancedView=False,)
+			   CustomProperties Pin (PinId=5E2D303F41CB55A5DAC09B8A4032EA14,PinName="self",PinFriendlyName=NSLOCTEXT("K2Node", "Target", "Target"),PinToolTip="Target\nHead Mounted Display Function Library Reference",PinType.PinCategory="object",PinType.PinSubCategory="",PinType.PinSubCategoryObject=Class'/Script/Engine.HeadMountedDisplayFunctionLibrary',PinType.PinSubCategoryMemberReference=(),PinType.PinValueType=(),PinType.bIsMap=False,PinType.bIsSet=False,PinType.bIsArray=False,PinType.bIsReference=False,PinType.bIsConst=False,PinType.bIsWeakPointer=False,DefaultObject="/Script/Engine.Default__HeadMountedDisplayFunctionLibrary",PersistentGuid=00000000000000000000000000000000,bHidden=True,bNotConnectable=False,bDefaultValueIsReadOnly=False,bDefaultValueIsIgnored=False,bAdvancedView=False,)
+			   CustomProperties Pin (PinId=C82F35534E48B3A5FF212CB4CAAEFD61,PinName="Origin",PinToolTip="Origin\nEHMDTrackingOrigin Enum",PinType.PinCategory="byte",PinType.PinSubCategory="",PinType.PinSubCategoryObject=Enum'/Script/Engine.EHMDTrackingOrigin',PinType.PinSubCategoryMemberReference=(),PinType.PinValueType=(),PinType.bIsMap=False,PinType.bIsSet=False,PinType.bIsArray=False,PinType.bIsReference=False,PinType.bIsConst=False,PinType.bIsWeakPointer=False,DefaultValue="Eye",PersistentGuid=00000000000000000000000000000000,bHidden=False,bNotConnectable=False,bDefaultValueIsReadOnly=False,bDefaultValueIsIgnored=False,bAdvancedView=False,)
+			End Object
+			Begin Object Class=/Script/BlueprintGraph.K2Node_CallFunction Name="K2Node_CallFunction_46"
+			   FunctionReference=(MemberParent=Class'/Script/Engine.SceneComponent',MemberName="K2_SetRelativeLocation")
+			   NodePosX=560
+			   NodePosY=-16
+			   NodeGuid=74BC1E504ADBB3371F22AB8555888796
+			   CustomProperties Pin (PinId=FE5D38794B000B7752514CAFCC5316E4,PinName="execute",PinToolTip="\nExec",PinType.PinCategory="exec",PinType.PinSubCategory="",PinType.PinSubCategoryObject=None,PinType.PinSubCategoryMemberReference=(),PinType.PinValueType=(),PinType.bIsMap=False,PinType.bIsSet=False,PinType.bIsArray=False,PinType.bIsReference=False,PinType.bIsConst=False,PinType.bIsWeakPointer=False,LinkedTo=(K2Node_CallFunction_77 6F56A02B45E8A67EEE79388CF61A44BC,),PersistentGuid=00000000000000000000000000000000,bHidden=False,bNotConnectable=False,bDefaultValueIsReadOnly=False,bDefaultValueIsIgnored=False,bAdvancedView=False,)
+			   CustomProperties Pin (PinId=1C7A9A54494AE23B7164B7B8725F5D52,PinName="then",PinToolTip="\nExec",Direction="EGPD_Output",PinType.PinCategory="exec",PinType.PinSubCategory="",PinType.PinSubCategoryObject=None,PinType.PinSubCategoryMemberReference=(),PinType.PinValueType=(),PinType.bIsMap=False,PinType.bIsSet=False,PinType.bIsArray=False,PinType.bIsReference=False,PinType.bIsConst=False,PinType.bIsWeakPointer=False,PersistentGuid=00000000000000000000000000000000,bHidden=False,bNotConnectable=False,bDefaultValueIsReadOnly=False,bDefaultValueIsIgnored=False,bAdvancedView=False,)
+			   CustomProperties Pin (PinId=0811E84B4E91F1030D8FE3BAC770A247,PinName="self",PinFriendlyName=NSLOCTEXT("K2Node", "Target", "Target"),PinToolTip="Target\nScene Component Reference",PinType.PinCategory="object",PinType.PinSubCategory="",PinType.PinSubCategoryObject=Class'/Script/Engine.SceneComponent',PinType.PinSubCategoryMemberReference=(),PinType.PinValueType=(),PinType.bIsMap=False,PinType.bIsSet=False,PinType.bIsArray=False,PinType.bIsReference=False,PinType.bIsConst=False,PinType.bIsWeakPointer=False,LinkedTo=(K2Node_VariableGet_20 14EFDA2C47FF2DF1F58E7A9BA401135E,),PersistentGuid=00000000000000000000000000000000,bHidden=False,bNotConnectable=False,bDefaultValueIsReadOnly=False,bDefaultValueIsIgnored=False,bAdvancedView=False,)
+			   CustomProperties Pin (PinId=C2D0A283497566CAFA793CB992D07158,PinName="NewLocation",PinToolTip="New Location\nVector \n\nNew location of the component relative to its parent.",PinType.PinCategory="struct",PinType.PinSubCategory="",PinType.PinSubCategoryObject=ScriptStruct'/Script/CoreUObject.Vector',PinType.PinSubCategoryMemberReference=(),PinType.PinValueType=(),PinType.bIsMap=False,PinType.bIsSet=False,PinType.bIsArray=False,PinType.bIsReference=False,PinType.bIsConst=False,PinType.bIsWeakPointer=False,DefaultValue="0, 0, 0",AutogeneratedDefaultValue="0, 0, 0",LinkedTo=(K2Node_VariableGet_40 7FE7B76D4888DB151FF54EAB9D25F826,),PersistentGuid=00000000000000000000000000000000,bHidden=False,bNotConnectable=False,bDefaultValueIsReadOnly=False,bDefaultValueIsIgnored=False,bAdvancedView=False,)
+			   CustomProperties Pin (PinId=3D1069744CC4C94F298F3980C2A0F289,PinName="bSweep",PinToolTip="Sweep\nBoolean\n\nWhether we sweep to the destination location, triggering overlaps along the way and stopping short of the target if blocked by something. Only the root component is swept and checked for blocking collision, child components move without sweeping. If collision is off, this has no effect.",PinType.PinCategory="bool",PinType.PinSubCategory="",PinType.PinSubCategoryObject=None,PinType.PinSubCategoryMemberReference=(),PinType.PinValueType=(),PinType.bIsMap=False,PinType.bIsSet=False,PinType.bIsArray=False,PinType.bIsReference=False,PinType.bIsConst=False,PinType.bIsWeakPointer=False,DefaultValue="false",AutogeneratedDefaultValue="false",PersistentGuid=00000000000000000000000000000000,bHidden=False,bNotConnectable=False,bDefaultValueIsReadOnly=False,bDefaultValueIsIgnored=False,bAdvancedView=False,)
+			   CustomProperties Pin (PinId=63A901EA42073DA211529DBD9BDB4FF8,PinName="SweepHitResult",PinToolTip="Sweep Hit Result\nHit Result Structure\n\nHit result from any impact if sweep is true.",Direction="EGPD_Output",PinType.PinCategory="struct",PinType.PinSubCategory="",PinType.PinSubCategoryObject=ScriptStruct'/Script/Engine.HitResult',PinType.PinSubCategoryMemberReference=(),PinType.PinValueType=(),PinType.bIsMap=False,PinType.bIsSet=False,PinType.bIsArray=False,PinType.bIsReference=False,PinType.bIsConst=False,PinType.bIsWeakPointer=False,PersistentGuid=00000000000000000000000000000000,bHidden=False,bNotConnectable=False,bDefaultValueIsReadOnly=False,bDefaultValueIsIgnored=False,bAdvancedView=False,)
+			   CustomProperties Pin (PinId=03544D8C4FC9FC037BE82CB2780E14A1,PinName="bTeleport",PinToolTip="Teleport\nBoolean\n\nWhether we teleport the physics state (if physics collision is enabled for this object). If true, physics velocity for this object is unchanged (so ragdoll parts are not affected by change in location). If false, physics velocity is updated based on the change in position (affecting ragdoll parts). If CCD is on and not teleporting, this will affect objects along the entire sweep volume.",PinType.PinCategory="bool",PinType.PinSubCategory="",PinType.PinSubCategoryObject=None,PinType.PinSubCategoryMemberReference=(),PinType.PinValueType=(),PinType.bIsMap=False,PinType.bIsSet=False,PinType.bIsArray=False,PinType.bIsReference=False,PinType.bIsConst=False,PinType.bIsWeakPointer=False,DefaultValue="false",AutogeneratedDefaultValue="false",PersistentGuid=00000000000000000000000000000000,bHidden=False,bNotConnectable=False,bDefaultValueIsReadOnly=False,bDefaultValueIsIgnored=False,bAdvancedView=False,)
+			End Object
+			Begin Object Class=/Script/BlueprintGraph.K2Node_VariableGet Name="K2Node_VariableGet_20"
+			   VariableReference=(MemberName="VRCameraRoot",bSelfContext=True)
+			   NodePosX=304
+			   NodePosY=144
+			   NodeGuid=B62CB0404A6FE2645BBEDCA2875DCA91
+			   CustomProperties Pin (PinId=14EFDA2C47FF2DF1F58E7A9BA401135E,PinName="VRCameraRoot",Direction="EGPD_Output",PinType.PinCategory="object",PinType.PinSubCategory="",PinType.PinSubCategoryObject=Class'/Script/Engine.SceneComponent',PinType.PinSubCategoryMemberReference=(),PinType.PinValueType=(),PinType.bIsMap=False,PinType.bIsSet=False,PinType.bIsArray=False,PinType.bIsReference=False,PinType.bIsConst=False,PinType.bIsWeakPointer=False,LinkedTo=(K2Node_CallFunction_46 0811E84B4E91F1030D8FE3BAC770A247,),PersistentGuid=00000000000000000000000000000000,bHidden=False,bNotConnectable=False,bDefaultValueIsReadOnly=False,bDefaultValueIsIgnored=False,bAdvancedView=False,)
+			   CustomProperties Pin (PinId=4018F2114B5E254238969EA5B9138178,PinName="self",PinFriendlyName=NSLOCTEXT("K2Node", "Target", "Target"),PinType.PinCategory="object",PinType.PinSubCategory="",PinType.PinSubCategoryObject=BlueprintGeneratedClass'/Game/GearVR_Pawn.GearVR_Pawn_C',PinType.PinSubCategoryMemberReference=(),PinType.PinValueType=(),PinType.bIsMap=False,PinType.bIsSet=False,PinType.bIsArray=False,PinType.bIsReference=False,PinType.bIsConst=False,PinType.bIsWeakPointer=False,PersistentGuid=00000000000000000000000000000000,bHidden=True,bNotConnectable=False,bDefaultValueIsReadOnly=False,bDefaultValueIsIgnored=False,bAdvancedView=False,)
+			End Object
+			Begin Object Class=/Script/BlueprintGraph.K2Node_VariableGet Name="K2Node_VariableGet_40"
+			   VariableReference=(MemberName="CameraHeight",MemberGuid=E6EDE9F147066500D13B0DA9CDB12D84,bSelfContext=True)
+			   NodePosX=320
+			   NodePosY=224
+			   NodeGuid=F6D021224B52F52FB98C4295DDB0F6E6
+			   CustomProperties Pin (PinId=7FE7B76D4888DB151FF54EAB9D25F826,PinName="CameraHeight",Direction="EGPD_Output",PinType.PinCategory="struct",PinType.PinSubCategory="",PinType.PinSubCategoryObject=ScriptStruct'/Script/CoreUObject.Vector',PinType.PinSubCategoryMemberReference=(),PinType.PinValueType=(),PinType.bIsMap=False,PinType.bIsSet=False,PinType.bIsArray=False,PinType.bIsReference=False,PinType.bIsConst=False,PinType.bIsWeakPointer=False,AutogeneratedDefaultValue="0, 0, 0",LinkedTo=(K2Node_CallFunction_46 C2D0A283497566CAFA793CB992D07158,),PersistentGuid=00000000000000000000000000000000,bHidden=False,bNotConnectable=False,bDefaultValueIsReadOnly=False,bDefaultValueIsIgnored=False,bAdvancedView=False,)
+			   CustomProperties Pin (PinId=1FB8519E4FFBB7EC464F1BB843BC179D,PinName="self",PinFriendlyName=NSLOCTEXT("K2Node", "Target", "Target"),PinType.PinCategory="object",PinType.PinSubCategory="",PinType.PinSubCategoryObject=BlueprintGeneratedClass'/Game/GearVR_Pawn.GearVR_Pawn_C',PinType.PinSubCategoryMemberReference=(),PinType.PinValueType=(),PinType.bIsMap=False,PinType.bIsSet=False,PinType.bIsArray=False,PinType.bIsReference=False,PinType.bIsConst=False,PinType.bIsWeakPointer=False,PersistentGuid=00000000000000000000000000000000,bHidden=True,bNotConnectable=False,bDefaultValueIsReadOnly=False,bDefaultValueIsIgnored=False,bAdvancedView=False,)
+			End Object
+
+			END OBJECT
+	[/PARAM]
+[/OBJECT]
+[/EXCERPT:GearVRCamera]
+-->
\ No newline at end of file
diff --git a/Engine/Documentation/Source/Platforms/GearVR/HowTo/GearVRHowTo.INT.udn b/Engine/Documentation/Source/Platforms/GearVR/HowTo/GearVRHowTo.INT.udn
new file mode 100644
--- /dev/null
+++ b/Engine/Documentation/Source/Platforms/GearVR/HowTo/GearVRHowTo.INT.udn
@@ -0,0 +1,34 @@
+Availability: Public
+Crumbs: %ROOT% 
+Title: Gear VR How - To
+Description: How To guides providing step-by-step instructions working with the Gear VR and UE4.
+Type: how-to
+Version: 4.17
+Parent: Platforms/GearVR
+Order:1
+Tags:VR
+topic-image:GearVR_How_To_Header.png
+
+The Gear VR How To's contain step-by-step guides that will walk you through how to use UE4 and Gear VR together.
+
+## Beginner
+
+The Gear VR Beginner How - To's will walk you through setting up your UE4 project camera to work with the Gear VR HMD and also show you how to reduce the size of your Gear VR project.
+
+[DIR(end:"1" output:"topic" type:"how-to" tags:"Basics" parent:"Platforms/GearVR")]
+
+## Intermediate
+
+The Gear VR Intermediate How - To's will walk you through using the Gear VR Controller and setting up the Gear VR Global Menu.
+
+[DIR(end:"1" output:"topic" type:"how-to" tags:"GearVR" parent:"Platforms/GearVR")]
+
+## Advanced
+
+The  Gear VR Advanced How - To's will walk you through debugging your Gear VR projects. 
+
+[DIR(end:"1" output:"topic" type:"how-to" tags:"GearVR" parent:"Platforms/GearVR")]
+
+
+
+
