Copying //UE4/Dev-Documentation to //UE4/Samples-Main (Source: //UE4/Dev-Documentation @ 3623665)

diff --git a/Engine/Documentation/Source/Platforms/GearVR/HowTo/HMDTouchPad/HTGearVRHMDTouchPad.INT.udn b/Engine/Documentation/Source/Platforms/GearVR/HowTo/HMDTouchPad/HTGearVRHMDTouchPad.INT.udn
new file mode 100644
--- /dev/null
+++ b/Engine/Documentation/Source/Platforms/GearVR/HowTo/HMDTouchPad/HTGearVRHMDTouchPad.INT.udn
@@ -0,0 +1,828 @@
+Availability: Public
+Crumbs: %ROOT%
+Title: Use the Gear VR HMD Touchpad 
+Description: Setting up the Gear VR HMD Touchpad 
+Type: how-to
+SkillLevel: Intermidate
+Version: 4.17
+Parent: Platforms/GearVR/HowTo
+Order: 3
+Tags: GearVR
+prereq: Platforms\GearVR\Prerequisites
+prereq: Platforms\GearVR\QuickStart
+prereq: Platforms\GearVR\HowTo\CameraSetup
+topic-image:Platforms\GearVR\HowTo\GearVR_How_To_Header.png
+
+Located on the side of the Gear VR HMD is a touchpad that allows users to interact with the VR project they are currently viewing. Users can interact with the VR projects by tapping or swipe their finger along this touchpad. In the following How - To we will go over setting up the Gear VR HMD touchpad to display different text based on various touchpad inputs.
+
+[region:note]
+This How - To will utilize the pawn the that was created in the [Set Up Cameras for Gear VR How - To](Platforms\GearVR\HowTo\CameraSetup). If you have not completed this How - To it is highly suggested you do before continuing.
+[/region]
+
+
+## Pawn Setup
+
+Below you will find instructions on how your player Pawn needs to be setup so that we can display what type of action is being performed on the Gear VR touchpad.
+
+1. First, open up your player Pawn Blueprint and from the **Components** tab, add a new **TextRender** component that is parented to the **VRCamera** and give it the following settings so that later on we can see what type of action was performed.
+
+    [REGION:lightbox]
+    [![](HTTouchPad_Pawn_Setup.png)(w:826)](HTTouchPad_Pawn_Setup.png)
+    [/REGION]
+
+    [REGION:caption]
+    Click for full image.
+    [/REGION]
+
+    |Property Name | Default Value|
+    |----------|--------------|
+    | **Location** |  X:100 Y:0 Z:0 |
+    | **Rotation** |  X:100 Y:0 Z:180 |
+    | **Text** |  TEMP TEMP TEMP |
+    | **Horizontal Alignment** |  Center |
+    | **Vertical Alignment** |  Text Center |
+
+1. Next, to determine if the input was a tap or a swipe we need to track the distance between the first touch and the last touch. If the distance between these numbers is under a certain amount, then a tap was performed. If the distance between these numbers is over a certain amount, then a swipe took place. To accomplish this, over in the  **Event Graph** create two **Vector 2D** variables with the following names:
+
+    * **TouchStart**
+    * **TouchEnd**
+
+    ![](HTTouchPad_Var_Setup.png)
+
+1. Finally, we need to add **five** **Text** variables that will be used to display if a tap or swipe took place and what direction the swipe originated. Add the following Text variables with the following names and values to your GearVR_Pawn Blueprint now.
+
+    ![](HTTouchPad_Var_TouchSwipe.png)
+
+    |Node Name         | Default Value|
+    |----------         |--------------|
+    | **TouchTap**   |  TouchTap    |
+    | **SwipeLeft**  |  SwipeLeft   |
+    | **SwipeRight** |  SwipeRight  |
+    | **SwipeUp**    |  SwipeUp     |
+    | **SwipeDown**  |  SwipeDown   |
+
+With all of the needed components and variables added to the Pawn, in the next section we will take a look what we need to do to detect if a tap or swipe took place.
+
+## Touch & Swipe Input Setup
+
+In the following section, we will go over how to setup the initial logic in our GearVR_Pawn Blueprint to determine if the user has just tapped or swiped the touchpad.
+
+1. First, search for the following Blueprint nodes add them to the Event Graph:
+
+    [REGION:lightbox]
+    [![](HTTouch_TouchStartEndNodes.png)(w:366)](HTTouch_TouchStartEndNodes.png)
+    [/REGION]
+
+    [REGION:caption]
+    Click for full image.
+    [/REGION]
+
+    |Node Name | Default Value|
+    |----------|--------------|
+    |**Touch**  | N/A |
+    |**Set Touch Start**  | N/A |
+    |**Set Touch End**  | N/A |
+    |**Get Player Controller ** | N/A |
+    |**Get Input Touch State ** | N/A |
+    |**Make Vector 2D** | N/A |
+
+1. Next, we will need to record and then store the location of the first and last touch as we can use the distance between these two numbers to see if a tap or swipe has taken place. To accomplish this, first, connect the Get Player Controller node to the Get Input Touch Start so that we can track when the first touch input the user does. To make things easier to work with, connect the Location X and Y to the Make Vector 2D node and then connect the Make Vector 2D node output to the inputs on both the Touch Start and Touch End variables. Finally, connect the Input Touch Pressed to the Touch Start variable and then connect the Input Touch Released to the Touch End variable. When completed this part of the Blueprint should look like the following image.
+
+    [REGION:lightbox]
+    [![](HTTouch_TouchStartEndSetup.png)(w:499)](HTTouch_TouchStartEndSetup.png)
+    [/REGION]
+
+    [REGION:caption]
+    Click for full image.
+    [/REGION]
+
+1. Now that our Touch Start and Touch End variables can be populated with data, next we are going to set up the logic to help determine if a tap or swipe took place. To accomplish this, first, add the following nodes to the Event Graph.
+
+    [REGION:lightbox]
+    [![](HTTouch_TapSwipeNodes.png)(w:665)](HTTouch_TapSwipeNodes.png)
+    [/REGION]
+
+    [REGION:caption]
+    Click for full image.
+    [/REGION]
+
+    |Node Name | Default Value|
+    |----------|--------------|
+    |**Touch Start  **| N/A |
+    |**Touch End ** | N/A |
+    |**Vector2d - Vector2d ** | N/A |
+    |**Vector 2d  Length ** | N/A |
+    |**Break Vector 2D ** | N/A |
+    |**Absolute (Float)** | N/A |
+    |**float < float **| N/A |
+    |**float > float ** | N/A |
+
+1. Before we can find out if a tap or a swipe took place, we first need to figure out the distance between the Touch Start and Touch End variables by subtracting them from one another. To get this working in our Blueprint, connect both Touch Start and Touch End variables to the inputs on the Vector2d - Vector2d node. 
+
+    [REGION:lightbox]
+    [![](HTTouch_TapOrSwipeLogic_01.png)(w:497)](HTTouch_TapOrSwipeLogic_01.png)
+    [/REGION]
+
+    [REGION:caption]
+    Click for full image.
+    [/REGION]
+
+1. To tell if a tap or a swipe has taken place we can measure the length or magnitude of the Touch Start and Touch End variables. If the length between the two variables is lower than X, a tap has taken place. If the length is higher than X, then a swipe has taken place. To set this up in our Blueprint, take the output of the  Vector2d - Vector2d node and connect it to the input on the Vector 2d  Length node. Then connect the Return Value of the Vector 2d  Length to the first input on the float > float and set the second input on the float > float to a value of **20** as this will define what distance is considered a tap versus a swipe.
+
+    [REGION:lightbox]
+    [![](HTTouch_TapOrSwipeLogic_02.png)(w:704)](HTTouch_TapOrSwipeLogic_02.png)
+    [/REGION]
+
+    [REGION:caption]
+    Click for full image.
+    [/REGION]
+
+    [region:note]
+    You can increase or decrease the distance that needs to be traversed on the touchpad for input to be considered a touch or swipe by increasing the number that is input into the second input on the float > float node  
+    [/region]
+
+1. To check if the user swiped left or right, we will need to see if the X value of the distance between Touch Start and Touch End is less than or greater than 0. If the distance is greater than 0, then it was a swipe left if not then it was a swipe right. To set this up, first, connect the output of the Vector2d - Vector2d node to the input of the Break Vector 2D node. Then connect the X output of the Break Vector 2D node to the first input on the float < float node and leave the second input with a value of 0. 
+
+    [REGION:lightbox]
+    [![](HTTouch_TapOrSwipeLogic_03.png)(w:704)](HTTouch_TapOrSwipeLogic_03.png)
+    [/REGION]
+
+    [REGION:caption]
+    Click for full image.
+    [/REGION]
+
+1. To check if the user swipes up or down, we need to see if the Y value of the distance between Touch Start and Touch End is greater than 1. If the Y value is greater than 1, then swipe up has occurred. If the Y value is less than one, then a swipe down has occurred. To set this up, connect the Y output on the Break Vector 2D node to the first input on the float > float node and then set the second input on the float > float node to a value of 1.
+
+    [REGION:lightbox]
+    [![](HTTouch_TapOrSwipeLogic_04.png)(w:704)](HTTouch_TapOrSwipeLogic_04.png)
+    [/REGION]
+
+    [REGION:caption]
+    Click for full image.
+    [/REGION]
+
+With the Blueprint logic to determine if we are tapping or swiping now completed in the next section, we will take a look at how to use this to display some text that tells the user if the tapped or swiped and which direction the swipe started from.
+
+## Determining Between Taps & Swipes
+
+ In the following section, we will setup the Blueprint to detect the difference between a tap and a swipe and also give it the ability to display some text that shows what type of action took place.
+
+1. First, we need to figure out if the user is doing a tap or a swipe by adding a **Branch** node to the Event Graph and connecting the output of the Set Touch End to the input on the Branch node. Then connect the return value of the float > 20 nodes to the **Condition** input on the Branch node. When completed, this logic will then check the distance between the Touch Start and Touch End and if that distance is less than 20 Unreal Units (UU) a tap took place. If that distance is over 20 UU, then a swipe took place.
+
+    [REGION:lightbox]
+    [![](HTTouch_TapOrSwipeSetup_01.png)(w:865)](HTTouch_TapOrSwipeSetup_01.png)
+    [/REGION]
+
+    [REGION:caption]
+    Click for full image.
+    [/REGION]
+
+1. Now that we know the difference between a tap and swipe, next we need to figure out if the user is swiping up or down versus swiping left or right as we want to do different actions for each. To determine this first duplicate both the Absolute(Float) node and one of the float > float nodes and connect them to match the image below. This part of the Blueprint will then compare the Absolute value of X and Y to see if they are bigger and smaller than one another and then perform different actions on each. For example, if the Absolute value of X and Y are greater than one another then a swipe Left or Right has taken place. If the Absolute value of X and Y are less than one another, then a swipe Up or Down has taken place.
+
+    [REGION:lightbox]
+    [![](HTTouch_TapOrSwipeSetup_02.png)(w:683)](HTTouch_TapOrSwipeSetup_02.png)
+    [/REGION]
+
+    [REGION:caption]
+    Click for full image.
+    [/REGION]
+
+    [region:note]
+    Whenever you are comparing the length of two vectors, you always want to make sure you use the Absolute value of the vectors as it will remove any negative signs allowing you to compare only positive values.
+    [/region]
+
+1. Next, drag off the True output of the existing Branch node and from the Executable Actions list search for and add a Branch Node. Then pull the Return value of the float > float node into the Condition on the newly created Branch node. Now when the user swipes left, or right the True option will fire and when the user swipes up or down the False option will fire. 
+
+    [REGION:lightbox]
+    [![](HTTouch_TapOrSwipeSetup_03.png)(w:530)](HTTouch_TapOrSwipeSetup_03.png)
+    [/REGION]
+
+    [REGION:caption]
+    Click for full image.
+    [/REGION]
+
+1. Add two more Branch nodes to the Event Graph and connect their inputs to the True and False outputs of the Branch node that was added in the previous step. These two new Branch nodes will control what is displayed when the user swipes in a particular direction. The Branch node coming from the True input will control what happens when swiping Left and Right. The Branch node coming from the False input will control what happens when swiping Up or Down.
+
+    [REGION:lightbox]
+    [![](HTTouch_TapOrSwipeSetup_04.png)(w:530)](HTTouch_TapOrSwipeSetup_04.png)
+    [/REGION]
+
+    [REGION:caption]
+    Click for full image.
+    [/REGION]
+
+1. Finally to tell the difference between swiping left and right connect the float < 0 to the Condition input on the first Branch node that was added in the previous. Then to tell the difference between swiping up and down connect the float < 1 to the Condition input on the second Branch node.
+
+    [REGION:lightbox]
+    [![](HTTouch_TapOrSwipeSetup_05.png)(w:614)](HTTouch_TapOrSwipeSetup_05.png)
+    [/REGION]
+
+    [REGION:caption]
+    Click for full image.
+    [/REGION]
+
+With our Blueprint tap and swipe functionality now set up, in the next section we will take a look at displaying differing text based on what type of input is applied to the Gear VR touchpad.
+
+## Performing Different Actions on Touch & Swipe
+
+With the ability to not only the determined difference between a tap and a swipe but also the direction a swipe came from, in this part of the How - To we are going to set up our GearVR_Pawn Blueprint to display text that will show what type of input took place.
+
+1. Drag off the False output from the first Branch node that deals with Tapping or Swiping and from the Executable Actions list, search for and add a **Set Text** (TextRender) node. 
+
+    [REGION:lightbox]
+    [![](HTTouch_TapOrSwipeDisplayText_01.png)(w:703)](HTTouch_TapOrSwipeDisplayText_01.png)
+    [/REGION]
+
+    [REGION:caption]
+    Click for full image.
+    [/REGION]
+
+1. Next, drag the TouchUp Text variable into the Event Graph and select the Get option. Then connect it to the Value input on the Set Text node.
+
+    [REGION:lightbox]
+    [![](HTTouch_TapOrSwipeDisplayText_04.png)(w:703)](HTTouch_TapOrSwipeDisplayText_04.png)
+    [/REGION]
+
+    [REGION:caption]
+    Click for full image.
+    [/REGION]
+
+1. Repeat the previous steps for the Up, Down, Left and Right directions and when completed your Blueprint should look like the following image:
+
+    [REGION:fullwidth]
+    [INCLUDE:#GearVRTouch]
+    [/REGION]
+
+    [REGION:caption]
+    Click for full image.
+    [/REGION]
+
+    [region:note]
+    You can copy the finished version of this Blueprint into your Blueprint by clicking on the Copy Node Graph option and then copying the supplied text and pasting it into your Blueprint.
+    [/region]
+
+1. Once you have Compiled and Saved the GearVR_Pawn Blueprint, deploy the UE4 project on your Samsung Galaxy Smartphone and then put your phone into your Gear VR HMD. Now when you tap or swipe the touchpad with your finger, you should see texts that will display the action you performed like in the video below. 
+
+    [OBJECT:EmbeddedVideo]
+        [PARAMLITERAL:videoid]
+        n-8tTaIK7pg
+        [/PARAMLITERAL]
+        [PARAMLITERAL:width]
+        100
+        [/PARAMLITERAL]
+        [PARAMLITERAL:height]
+        56.25
+        [/PARAMLITERAL]
+        [PARAMLITERAL:units]
+        %
+        [/PARAMLITERAL]
+        [PARAMLITERAL:type]
+        simple
+        [/PARAMLITERAL]
+    [/OBJECT]
+
+    [region:note]
+    If the text fails to show up when you tap or swipe the Gear Vr touchpad, make sure that text has been input into the Default Value of the Text variable. 
+
+    ![](HTTouch_TapOrSwipeDisplayText_03.png)
+    [/region]
+
+## UE4 Project Downloads
+
+Below you will find a link to where you can download the UE4 project that was used to create this example. 
+
+* [**Gear VR HMD Touchpad Example Project**](https://epicgames.box.com/s/19b1w5nk9srv7y3dwkv9c4mdem28rt0z)
+
+<!--
+[EXCERPT:GearVRTouch]	
+[OBJECT:Snippet]
+	[PARAMLITERAL:id]
+		GearVRTouch
+	[/PARAMLITERAL]
+	[PARAMLITERAL:type]
+		blueprint
+	[/PARAMLITERAL]
+	[PARAM:image]
+		[REGION:lightbox]
+		[![](HTTouch_TapOrSwipeDisplayText_02.png)(w:1000)](HTTouch_TapOrSwipeDisplayText_02.png)
+		[/REGION]
+	[/PARAM]
+	[PARAM:snippet]
+			BEGIN OBJECT
+				Begin Object Class=/Script/BlueprintGraph.K2Node_Event Name="K2Node_Event_3"
+				   EventReference=(MemberParent=Class'/Script/Engine.Actor',MemberName="ReceiveBeginPlay")
+				   bOverrideFunction=True
+				   NodePosX=-128
+				   NodePosY=544
+				   bCommentBubblePinned=True
+				   NodeGuid=01A9096D4670FCAE9AED7A9B6BC55D5B
+				   CustomProperties Pin (PinId=FE39330745F60B4B61B9138BA2E683E1,PinName="OutputDelegate",Direction="EGPD_Output",PinType.PinCategory="delegate",PinType.PinSubCategory="",PinType.PinSubCategoryObject=None,PinType.PinSubCategoryMemberReference=(MemberParent=Class'/Script/Engine.Actor',MemberName="ReceiveBeginPlay"),PinType.PinValueType=(),PinType.bIsMap=False,PinType.bIsSet=False,PinType.bIsArray=False,PinType.bIsReference=False,PinType.bIsConst=False,PinType.bIsWeakPointer=False,PersistentGuid=00000000000000000000000000000000,bHidden=False,bNotConnectable=False,bDefaultValueIsReadOnly=False,bDefaultValueIsIgnored=False,bAdvancedView=False,)
+				   CustomProperties Pin (PinId=CFFCA1684C5D4AA0653641B0BB601A74,PinName="then",Direction="EGPD_Output",PinType.PinCategory="exec",PinType.PinSubCategory="",PinType.PinSubCategoryObject=None,PinType.PinSubCategoryMemberReference=(),PinType.PinValueType=(),PinType.bIsMap=False,PinType.bIsSet=False,PinType.bIsArray=False,PinType.bIsReference=False,PinType.bIsConst=False,PinType.bIsWeakPointer=False,LinkedTo=(K2Node_CallFunction_31 7227CDD94FDBCC1F02D6258FDC30999A,),PersistentGuid=00000000000000000000000000000000,bHidden=False,bNotConnectable=False,bDefaultValueIsReadOnly=False,bDefaultValueIsIgnored=False,bAdvancedView=False,)
+				End Object
+				Begin Object Class=/Script/BlueprintGraph.K2Node_CallFunction Name="K2Node_CallFunction_31"
+				   FunctionReference=(MemberParent=Class'/Script/Engine.HeadMountedDisplayFunctionLibrary',MemberName="SetTrackingOrigin")
+				   NodePosX=304
+				   NodePosY=544
+				   NodeComment="PS4"
+				   NodeGuid=500FCC42479A2675EB54D58EC9DFC30F
+				   CustomProperties Pin (PinId=7227CDD94FDBCC1F02D6258FDC30999A,PinName="execute",PinType.PinCategory="exec",PinType.PinSubCategory="",PinType.PinSubCategoryObject=None,PinType.PinSubCategoryMemberReference=(),PinType.PinValueType=(),PinType.bIsMap=False,PinType.bIsSet=False,PinType.bIsArray=False,PinType.bIsReference=False,PinType.bIsConst=False,PinType.bIsWeakPointer=False,LinkedTo=(K2Node_Event_3 CFFCA1684C5D4AA0653641B0BB601A74,),PersistentGuid=00000000000000000000000000000000,bHidden=False,bNotConnectable=False,bDefaultValueIsReadOnly=False,bDefaultValueIsIgnored=False,bAdvancedView=False,)
+				   CustomProperties Pin (PinId=6F56A02B45E8A67EEE79388CF61A44BC,PinName="then",Direction="EGPD_Output",PinType.PinCategory="exec",PinType.PinSubCategory="",PinType.PinSubCategoryObject=None,PinType.PinSubCategoryMemberReference=(),PinType.PinValueType=(),PinType.bIsMap=False,PinType.bIsSet=False,PinType.bIsArray=False,PinType.bIsReference=False,PinType.bIsConst=False,PinType.bIsWeakPointer=False,LinkedTo=(K2Node_CallFunction_32 FE5D38794B000B7752514CAFCC5316E4,),PersistentGuid=00000000000000000000000000000000,bHidden=False,bNotConnectable=False,bDefaultValueIsReadOnly=False,bDefaultValueIsIgnored=False,bAdvancedView=False,)
+				   CustomProperties Pin (PinId=5E2D303F41CB55A5DAC09B8A4032EA14,PinName="self",PinFriendlyName=NSLOCTEXT("K2Node", "Target", "Target"),PinType.PinCategory="object",PinType.PinSubCategory="",PinType.PinSubCategoryObject=Class'/Script/Engine.HeadMountedDisplayFunctionLibrary',PinType.PinSubCategoryMemberReference=(),PinType.PinValueType=(),PinType.bIsMap=False,PinType.bIsSet=False,PinType.bIsArray=False,PinType.bIsReference=False,PinType.bIsConst=False,PinType.bIsWeakPointer=False,DefaultObject="/Script/Engine.Default__HeadMountedDisplayFunctionLibrary",PersistentGuid=00000000000000000000000000000000,bHidden=True,bNotConnectable=False,bDefaultValueIsReadOnly=False,bDefaultValueIsIgnored=False,bAdvancedView=False,)
+				   CustomProperties Pin (PinId=C82F35534E48B3A5FF212CB4CAAEFD61,PinName="Origin",PinType.PinCategory="byte",PinType.PinSubCategory="",PinType.PinSubCategoryObject=Enum'/Script/Engine.EHMDTrackingOrigin',PinType.PinSubCategoryMemberReference=(),PinType.PinValueType=(),PinType.bIsMap=False,PinType.bIsSet=False,PinType.bIsArray=False,PinType.bIsReference=False,PinType.bIsConst=False,PinType.bIsWeakPointer=False,DefaultValue="Eye",PersistentGuid=00000000000000000000000000000000,bHidden=False,bNotConnectable=False,bDefaultValueIsReadOnly=False,bDefaultValueIsIgnored=False,bAdvancedView=False,)
+				End Object
+				Begin Object Class=/Script/BlueprintGraph.K2Node_CallFunction Name="K2Node_CallFunction_32"
+				   FunctionReference=(MemberParent=Class'/Script/Engine.SceneComponent',MemberName="K2_SetRelativeLocation")
+				   NodePosX=640
+				   NodePosY=528
+				   NodeGuid=C501FC114606C0423A722ABE4E8AB47E
+				   CustomProperties Pin (PinId=FE5D38794B000B7752514CAFCC5316E4,PinName="execute",PinType.PinCategory="exec",PinType.PinSubCategory="",PinType.PinSubCategoryObject=None,PinType.PinSubCategoryMemberReference=(),PinType.PinValueType=(),PinType.bIsMap=False,PinType.bIsSet=False,PinType.bIsArray=False,PinType.bIsReference=False,PinType.bIsConst=False,PinType.bIsWeakPointer=False,LinkedTo=(K2Node_CallFunction_31 6F56A02B45E8A67EEE79388CF61A44BC,),PersistentGuid=00000000000000000000000000000000,bHidden=False,bNotConnectable=False,bDefaultValueIsReadOnly=False,bDefaultValueIsIgnored=False,bAdvancedView=False,)
+				   CustomProperties Pin (PinId=1C7A9A54494AE23B7164B7B8725F5D52,PinName="then",Direction="EGPD_Output",PinType.PinCategory="exec",PinType.PinSubCategory="",PinType.PinSubCategoryObject=None,PinType.PinSubCategoryMemberReference=(),PinType.PinValueType=(),PinType.bIsMap=False,PinType.bIsSet=False,PinType.bIsArray=False,PinType.bIsReference=False,PinType.bIsConst=False,PinType.bIsWeakPointer=False,PersistentGuid=00000000000000000000000000000000,bHidden=False,bNotConnectable=False,bDefaultValueIsReadOnly=False,bDefaultValueIsIgnored=False,bAdvancedView=False,)
+				   CustomProperties Pin (PinId=0811E84B4E91F1030D8FE3BAC770A247,PinName="self",PinFriendlyName=NSLOCTEXT("K2Node", "Target", "Target"),PinType.PinCategory="object",PinType.PinSubCategory="",PinType.PinSubCategoryObject=Class'/Script/Engine.SceneComponent',PinType.PinSubCategoryMemberReference=(),PinType.PinValueType=(),PinType.bIsMap=False,PinType.bIsSet=False,PinType.bIsArray=False,PinType.bIsReference=False,PinType.bIsConst=False,PinType.bIsWeakPointer=False,LinkedTo=(K2Node_VariableGet_12 14EFDA2C47FF2DF1F58E7A9BA401135E,),PersistentGuid=00000000000000000000000000000000,bHidden=False,bNotConnectable=False,bDefaultValueIsReadOnly=False,bDefaultValueIsIgnored=False,bAdvancedView=False,)
+				   CustomProperties Pin (PinId=C2D0A283497566CAFA793CB992D07158,PinName="NewLocation",PinType.PinCategory="struct",PinType.PinSubCategory="",PinType.PinSubCategoryObject=ScriptStruct'/Script/CoreUObject.Vector',PinType.PinSubCategoryMemberReference=(),PinType.PinValueType=(),PinType.bIsMap=False,PinType.bIsSet=False,PinType.bIsArray=False,PinType.bIsReference=False,PinType.bIsConst=False,PinType.bIsWeakPointer=False,DefaultValue="0, 0, 0",AutogeneratedDefaultValue="0, 0, 0",LinkedTo=(K2Node_VariableGet_13 7FE7B76D4888DB151FF54EAB9D25F826,),PersistentGuid=00000000000000000000000000000000,bHidden=False,bNotConnectable=False,bDefaultValueIsReadOnly=False,bDefaultValueIsIgnored=False,bAdvancedView=False,)
+				   CustomProperties Pin (PinId=3D1069744CC4C94F298F3980C2A0F289,PinName="bSweep",PinType.PinCategory="bool",PinType.PinSubCategory="",PinType.PinSubCategoryObject=None,PinType.PinSubCategoryMemberReference=(),PinType.PinValueType=(),PinType.bIsMap=False,PinType.bIsSet=False,PinType.bIsArray=False,PinType.bIsReference=False,PinType.bIsConst=False,PinType.bIsWeakPointer=False,DefaultValue="false",AutogeneratedDefaultValue="false",PersistentGuid=00000000000000000000000000000000,bHidden=False,bNotConnectable=False,bDefaultValueIsReadOnly=False,bDefaultValueIsIgnored=False,bAdvancedView=False,)
+				   CustomProperties Pin (PinId=63A901EA42073DA211529DBD9BDB4FF8,PinName="SweepHitResult",Direction="EGPD_Output",PinType.PinCategory="struct",PinType.PinSubCategory="",PinType.PinSubCategoryObject=ScriptStruct'/Script/Engine.HitResult',PinType.PinSubCategoryMemberReference=(),PinType.PinValueType=(),PinType.bIsMap=False,PinType.bIsSet=False,PinType.bIsArray=False,PinType.bIsReference=False,PinType.bIsConst=False,PinType.bIsWeakPointer=False,PersistentGuid=00000000000000000000000000000000,bHidden=False,bNotConnectable=False,bDefaultValueIsReadOnly=False,bDefaultValueIsIgnored=False,bAdvancedView=False,)
+				   CustomProperties Pin (PinId=03544D8C4FC9FC037BE82CB2780E14A1,PinName="bTeleport",PinType.PinCategory="bool",PinType.PinSubCategory="",PinType.PinSubCategoryObject=None,PinType.PinSubCategoryMemberReference=(),PinType.PinValueType=(),PinType.bIsMap=False,PinType.bIsSet=False,PinType.bIsArray=False,PinType.bIsReference=False,PinType.bIsConst=False,PinType.bIsWeakPointer=False,DefaultValue="false",AutogeneratedDefaultValue="false",PersistentGuid=00000000000000000000000000000000,bHidden=False,bNotConnectable=False,bDefaultValueIsReadOnly=False,bDefaultValueIsIgnored=False,bAdvancedView=False,)
+				End Object
+				Begin Object Class=/Script/BlueprintGraph.K2Node_VariableGet Name="K2Node_VariableGet_12"
+				   VariableReference=(MemberName="VRCameraRoot",bSelfContext=True)
+				   NodePosX=384
+				   NodePosY=688
+				   NodeGuid=65DBC621448332008720529827567EC0
+				   CustomProperties Pin (PinId=14EFDA2C47FF2DF1F58E7A9BA401135E,PinName="VRCameraRoot",Direction="EGPD_Output",PinType.PinCategory="object",PinType.PinSubCategory="",PinType.PinSubCategoryObject=Class'/Script/Engine.SceneComponent',PinType.PinSubCategoryMemberReference=(),PinType.PinValueType=(),PinType.bIsMap=False,PinType.bIsSet=False,PinType.bIsArray=False,PinType.bIsReference=False,PinType.bIsConst=False,PinType.bIsWeakPointer=False,LinkedTo=(K2Node_CallFunction_32 0811E84B4E91F1030D8FE3BAC770A247,),PersistentGuid=00000000000000000000000000000000,bHidden=False,bNotConnectable=False,bDefaultValueIsReadOnly=False,bDefaultValueIsIgnored=False,bAdvancedView=False,)
+				   CustomProperties Pin (PinId=4018F2114B5E254238969EA5B9138178,PinName="self",PinFriendlyName=NSLOCTEXT("K2Node", "Target", "Target"),PinType.PinCategory="object",PinType.PinSubCategory="",PinType.PinSubCategoryObject=BlueprintGeneratedClass'/Game/GearVR_Pawn/GearVR_Pawn_Input.GearVR_Pawn_Input_C',PinType.PinSubCategoryMemberReference=(),PinType.PinValueType=(),PinType.bIsMap=False,PinType.bIsSet=False,PinType.bIsArray=False,PinType.bIsReference=False,PinType.bIsConst=False,PinType.bIsWeakPointer=False,PersistentGuid=00000000000000000000000000000000,bHidden=True,bNotConnectable=False,bDefaultValueIsReadOnly=False,bDefaultValueIsIgnored=False,bAdvancedView=False,)
+				End Object
+				Begin Object Class=/Script/BlueprintGraph.K2Node_VariableGet Name="K2Node_VariableGet_13"
+				   VariableReference=(MemberName="CameraHeight",MemberGuid=E6EDE9F147066500D13B0DA9CDB12D84,bSelfContext=True)
+				   NodePosX=400
+				   NodePosY=768
+				   NodeGuid=78F380174F78332B4A1A7D92F14E2BFC
+				   CustomProperties Pin (PinId=7FE7B76D4888DB151FF54EAB9D25F826,PinName="CameraHeight",Direction="EGPD_Output",PinType.PinCategory="struct",PinType.PinSubCategory="",PinType.PinSubCategoryObject=ScriptStruct'/Script/CoreUObject.Vector',PinType.PinSubCategoryMemberReference=(),PinType.PinValueType=(),PinType.bIsMap=False,PinType.bIsSet=False,PinType.bIsArray=False,PinType.bIsReference=False,PinType.bIsConst=False,PinType.bIsWeakPointer=False,AutogeneratedDefaultValue="0, 0, 0",LinkedTo=(K2Node_CallFunction_32 C2D0A283497566CAFA793CB992D07158,),PersistentGuid=00000000000000000000000000000000,bHidden=False,bNotConnectable=False,bDefaultValueIsReadOnly=False,bDefaultValueIsIgnored=False,bAdvancedView=False,)
+				   CustomProperties Pin (PinId=1FB8519E4FFBB7EC464F1BB843BC179D,PinName="self",PinFriendlyName=NSLOCTEXT("K2Node", "Target", "Target"),PinType.PinCategory="object",PinType.PinSubCategory="",PinType.PinSubCategoryObject=BlueprintGeneratedClass'/Game/GearVR_Pawn/GearVR_Pawn_Input.GearVR_Pawn_Input_C',PinType.PinSubCategoryMemberReference=(),PinType.PinValueType=(),PinType.bIsMap=False,PinType.bIsSet=False,PinType.bIsArray=False,PinType.bIsReference=False,PinType.bIsConst=False,PinType.bIsWeakPointer=False,PersistentGuid=00000000000000000000000000000000,bHidden=True,bNotConnectable=False,bDefaultValueIsReadOnly=False,bDefaultValueIsIgnored=False,bAdvancedView=False,)
+				End Object
+				Begin Object Class=/Script/UnrealEd.EdGraphNode_Comment Name="EdGraphNode_Comment_4"
+				   NodePosX=240
+				   NodePosY=464
+				   NodeWidth=752
+				   NodeHeight=368
+				   NodeComment="Gear VR Camera Setup"
+				   NodeGuid=07E0F7D84C383B8AB3634883E3D160A4
+				End Object
+				Begin Object Class=/Script/BlueprintGraph.K2Node_InputTouch Name="K2Node_InputTouch_7"
+				   NodePosX=464
+				   NodePosY=-768
+				   ErrorType=1
+				   ErrorMsg="Error Found more than one function with the same name InpTchEvt_Pressed; second occurance at  InputTouch\nError Found more than one function with the same name InpTchEvt_Released; second occurance at  InputTouch"
+				   NodeGuid=5136DE4146D30E89890CD6B9C1A3C5A8
+				   CustomProperties Pin (PinId=3787414C455842F10BAB63AB58C08F6A,PinName="Pressed",Direction="EGPD_Output",PinType.PinCategory="exec",PinType.PinSubCategory="",PinType.PinSubCategoryObject=None,PinType.PinSubCategoryMemberReference=(),PinType.PinValueType=(),PinType.bIsMap=False,PinType.bIsSet=False,PinType.bIsArray=False,PinType.bIsReference=False,PinType.bIsConst=False,PinType.bIsWeakPointer=False,LinkedTo=(K2Node_VariableSet_44 7AF704324F9FF65DC8729EA838A93153,),PersistentGuid=00000000000000000000000000000000,bHidden=False,bNotConnectable=False,bDefaultValueIsReadOnly=False,bDefaultValueIsIgnored=False,bAdvancedView=False,)
+				   CustomProperties Pin (PinId=C6DEF3D74E3752A1346024AAD8FA60FE,PinName="Released",Direction="EGPD_Output",PinType.PinCategory="exec",PinType.PinSubCategory="",PinType.PinSubCategoryObject=None,PinType.PinSubCategoryMemberReference=(),PinType.PinValueType=(),PinType.bIsMap=False,PinType.bIsSet=False,PinType.bIsArray=False,PinType.bIsReference=False,PinType.bIsConst=False,PinType.bIsWeakPointer=False,LinkedTo=(K2Node_VariableSet_45 E4E9F7E6497589AB9D3F3BA8D6C63943,),PersistentGuid=00000000000000000000000000000000,bHidden=False,bNotConnectable=False,bDefaultValueIsReadOnly=False,bDefaultValueIsIgnored=False,bAdvancedView=False,)
+				   CustomProperties Pin (PinId=5AEF1BE645563EFBE93DDF83F2FC979D,PinName="Moved",Direction="EGPD_Output",PinType.PinCategory="exec",PinType.PinSubCategory="",PinType.PinSubCategoryObject=None,PinType.PinSubCategoryMemberReference=(),PinType.PinValueType=(),PinType.bIsMap=False,PinType.bIsSet=False,PinType.bIsArray=False,PinType.bIsReference=False,PinType.bIsConst=False,PinType.bIsWeakPointer=False,PersistentGuid=00000000000000000000000000000000,bHidden=False,bNotConnectable=False,bDefaultValueIsReadOnly=False,bDefaultValueIsIgnored=False,bAdvancedView=False,)
+				   CustomProperties Pin (PinId=919B26024F770C343C9594A6C534BA6D,PinName="Location",Direction="EGPD_Output",PinType.PinCategory="struct",PinType.PinSubCategory="",PinType.PinSubCategoryObject=ScriptStruct'/Script/CoreUObject.Vector',PinType.PinSubCategoryMemberReference=(),PinType.PinValueType=(),PinType.bIsMap=False,PinType.bIsSet=False,PinType.bIsArray=False,PinType.bIsReference=False,PinType.bIsConst=False,PinType.bIsWeakPointer=False,PersistentGuid=00000000000000000000000000000000,bHidden=False,bNotConnectable=False,bDefaultValueIsReadOnly=False,bDefaultValueIsIgnored=False,bAdvancedView=False,)
+				   CustomProperties Pin (PinId=FE5861B541C701911EC4BA983244C318,PinName="FingerIndex",Direction="EGPD_Output",PinType.PinCategory="byte",PinType.PinSubCategory="",PinType.PinSubCategoryObject=Enum'/Script/InputCore.ETouchIndex',PinType.PinSubCategoryMemberReference=(),PinType.PinValueType=(),PinType.bIsMap=False,PinType.bIsSet=False,PinType.bIsArray=False,PinType.bIsReference=False,PinType.bIsConst=False,PinType.bIsWeakPointer=False,PersistentGuid=00000000000000000000000000000000,bHidden=False,bNotConnectable=False,bDefaultValueIsReadOnly=False,bDefaultValueIsIgnored=False,bAdvancedView=False,)
+				End Object
+				Begin Object Class=/Script/BlueprintGraph.K2Node_CallFunction Name="K2Node_CallFunction_58"
+				   bIsPureFunc=True
+				   FunctionReference=(MemberParent=Class'/Script/Engine.GameplayStatics',MemberName="GetPlayerController")
+				   NodePosX=80
+				   NodePosY=-464
+				   NodeGuid=D698606545AA6C8D01D3F78FC3A4B4C7
+				   CustomProperties Pin (PinId=2CAC5D8C49B95F5A396D5C9EEC343646,PinName="self",PinFriendlyName=NSLOCTEXT("K2Node", "Target", "Target"),PinToolTip="Target\nGameplay Statics Reference",PinType.PinCategory="object",PinType.PinSubCategory="",PinType.PinSubCategoryObject=Class'/Script/Engine.GameplayStatics',PinType.PinSubCategoryMemberReference=(),PinType.PinValueType=(),PinType.bIsMap=False,PinType.bIsSet=False,PinType.bIsArray=False,PinType.bIsReference=False,PinType.bIsConst=False,PinType.bIsWeakPointer=False,DefaultObject="/Script/Engine.Default__GameplayStatics",PersistentGuid=00000000000000000000000000000000,bHidden=True,bNotConnectable=False,bDefaultValueIsReadOnly=False,bDefaultValueIsIgnored=False,bAdvancedView=False,)
+				   CustomProperties Pin (PinId=854EC3F64CB5F143C7C5258338001511,PinName="WorldContextObject",PinToolTip="World Context Object\nObject Reference",PinType.PinCategory="object",PinType.PinSubCategory="",PinType.PinSubCategoryObject=Class'/Script/CoreUObject.Object',PinType.PinSubCategoryMemberReference=(),PinType.PinValueType=(),PinType.bIsMap=False,PinType.bIsSet=False,PinType.bIsArray=False,PinType.bIsReference=False,PinType.bIsConst=True,PinType.bIsWeakPointer=False,PersistentGuid=00000000000000000000000000000000,bHidden=True,bNotConnectable=False,bDefaultValueIsReadOnly=False,bDefaultValueIsIgnored=False,bAdvancedView=False,)
+				   CustomProperties Pin (PinId=F328E375485FF3C5DB8E26ADE2DF9A56,PinName="PlayerIndex",PinToolTip="Player Index\nInteger",PinType.PinCategory="int",PinType.PinSubCategory="",PinType.PinSubCategoryObject=None,PinType.PinSubCategoryMemberReference=(),PinType.PinValueType=(),PinType.bIsMap=False,PinType.bIsSet=False,PinType.bIsArray=False,PinType.bIsReference=False,PinType.bIsConst=False,PinType.bIsWeakPointer=False,DefaultValue="0",AutogeneratedDefaultValue="0",PersistentGuid=00000000000000000000000000000000,bHidden=False,bNotConnectable=False,bDefaultValueIsReadOnly=False,bDefaultValueIsIgnored=False,bAdvancedView=False,)
+				   CustomProperties Pin (PinId=C885E45A42B1529E4AD727BA4459D535,PinName="ReturnValue",PinToolTip="Return Value\nPlayer Controller Reference",Direction="EGPD_Output",PinType.PinCategory="object",PinType.PinSubCategory="",PinType.PinSubCategoryObject=Class'/Script/Engine.PlayerController',PinType.PinSubCategoryMemberReference=(),PinType.PinValueType=(),PinType.bIsMap=False,PinType.bIsSet=False,PinType.bIsArray=False,PinType.bIsReference=False,PinType.bIsConst=False,PinType.bIsWeakPointer=False,LinkedTo=(K2Node_CallFunction_59 4A56954D49F5D9B3195EAA84D1C6AD73,),PersistentGuid=00000000000000000000000000000000,bHidden=False,bNotConnectable=False,bDefaultValueIsReadOnly=False,bDefaultValueIsIgnored=False,bAdvancedView=False,)
+				End Object
+				Begin Object Class=/Script/BlueprintGraph.K2Node_CallFunction Name="K2Node_CallFunction_59"
+				   bIsPureFunc=True
+				   bIsConstFunc=True
+				   FunctionReference=(MemberParent=Class'/Script/Engine.PlayerController',MemberName="GetInputTouchState")
+				   NodePosX=368
+				   NodePosY=-480
+				   NodeGuid=C7E4A049430980DB404A49A7BB48DCCF
+				   CustomProperties Pin (PinId=4A56954D49F5D9B3195EAA84D1C6AD73,PinName="self",PinFriendlyName=NSLOCTEXT("K2Node", "Target", "Target"),PinToolTip="Target\nPlayer Controller Reference",PinType.PinCategory="object",PinType.PinSubCategory="",PinType.PinSubCategoryObject=Class'/Script/Engine.PlayerController',PinType.PinSubCategoryMemberReference=(),PinType.PinValueType=(),PinType.bIsMap=False,PinType.bIsSet=False,PinType.bIsArray=False,PinType.bIsReference=False,PinType.bIsConst=False,PinType.bIsWeakPointer=False,LinkedTo=(K2Node_CallFunction_58 C885E45A42B1529E4AD727BA4459D535,),PersistentGuid=00000000000000000000000000000000,bHidden=False,bNotConnectable=False,bDefaultValueIsReadOnly=False,bDefaultValueIsIgnored=False,bAdvancedView=False,)
+				   CustomProperties Pin (PinId=95E73FBB40BA062C80C70280839A2FE4,PinName="FingerIndex",PinToolTip="Finger Index\nETouchIndex Enum",PinType.PinCategory="byte",PinType.PinSubCategory="",PinType.PinSubCategoryObject=Enum'/Script/InputCore.ETouchIndex',PinType.PinSubCategoryMemberReference=(),PinType.PinValueType=(),PinType.bIsMap=False,PinType.bIsSet=False,PinType.bIsArray=False,PinType.bIsReference=False,PinType.bIsConst=False,PinType.bIsWeakPointer=False,DefaultValue="Touch1",PersistentGuid=00000000000000000000000000000000,bHidden=False,bNotConnectable=False,bDefaultValueIsReadOnly=False,bDefaultValueIsIgnored=False,bAdvancedView=False,)
+				   CustomProperties Pin (PinId=B2007DF34F06C520502E3789A19668D9,PinName="LocationX",PinToolTip="Location X\nFloat",Direction="EGPD_Output",PinType.PinCategory="float",PinType.PinSubCategory="",PinType.PinSubCategoryObject=None,PinType.PinSubCategoryMemberReference=(),PinType.PinValueType=(),PinType.bIsMap=False,PinType.bIsSet=False,PinType.bIsArray=False,PinType.bIsReference=False,PinType.bIsConst=False,PinType.bIsWeakPointer=False,AutogeneratedDefaultValue="0.0",LinkedTo=(K2Node_CallFunction_60 FD97B5EA441E8ED0C1E701A2A9811E29,),PersistentGuid=00000000000000000000000000000000,bHidden=False,bNotConnectable=False,bDefaultValueIsReadOnly=False,bDefaultValueIsIgnored=False,bAdvancedView=False,)
+				   CustomProperties Pin (PinId=BE68D9904F9FB271CCD42CB3A1957A09,PinName="LocationY",PinToolTip="Location Y\nFloat",Direction="EGPD_Output",PinType.PinCategory="float",PinType.PinSubCategory="",PinType.PinSubCategoryObject=None,PinType.PinSubCategoryMemberReference=(),PinType.PinValueType=(),PinType.bIsMap=False,PinType.bIsSet=False,PinType.bIsArray=False,PinType.bIsReference=False,PinType.bIsConst=False,PinType.bIsWeakPointer=False,AutogeneratedDefaultValue="0.0",LinkedTo=(K2Node_CallFunction_60 6CF890D2486F8AD9D08636B1CBC925E2,),PersistentGuid=00000000000000000000000000000000,bHidden=False,bNotConnectable=False,bDefaultValueIsReadOnly=False,bDefaultValueIsIgnored=False,bAdvancedView=False,)
+				   CustomProperties Pin (PinId=E385E07342603A410D9872A4D4846096,PinName="bIsCurrentlyPressed",PinToolTip="Is Currently Pressed\nBoolean",Direction="EGPD_Output",PinType.PinCategory="bool",PinType.PinSubCategory="",PinType.PinSubCategoryObject=None,PinType.PinSubCategoryMemberReference=(),PinType.PinValueType=(),PinType.bIsMap=False,PinType.bIsSet=False,PinType.bIsArray=False,PinType.bIsReference=False,PinType.bIsConst=False,PinType.bIsWeakPointer=False,DefaultValue="false",AutogeneratedDefaultValue="false",PersistentGuid=00000000000000000000000000000000,bHidden=False,bNotConnectable=False,bDefaultValueIsReadOnly=False,bDefaultValueIsIgnored=False,bAdvancedView=False,)
+				End Object
+				Begin Object Class=/Script/BlueprintGraph.K2Node_CallFunction Name="K2Node_CallFunction_60"
+				   bIsPureFunc=True
+				   FunctionReference=(MemberParent=Class'/Script/Engine.KismetMathLibrary',MemberName="MakeVector2D")
+				   NodePosX=768
+				   NodePosY=-464
+				   NodeGuid=2C425A9A46219C96284984AAFD9A056C
+				   CustomProperties Pin (PinId=C7CBA6C543F1CDA1D75CBEA645749EFC,PinName="self",PinFriendlyName=NSLOCTEXT("K2Node", "Target", "Target"),PinToolTip="Target\nKismet Math Library Reference",PinType.PinCategory="object",PinType.PinSubCategory="",PinType.PinSubCategoryObject=Class'/Script/Engine.KismetMathLibrary',PinType.PinSubCategoryMemberReference=(),PinType.PinValueType=(),PinType.bIsMap=False,PinType.bIsSet=False,PinType.bIsArray=False,PinType.bIsReference=False,PinType.bIsConst=False,PinType.bIsWeakPointer=False,DefaultObject="/Script/Engine.Default__KismetMathLibrary",PersistentGuid=00000000000000000000000000000000,bHidden=True,bNotConnectable=False,bDefaultValueIsReadOnly=False,bDefaultValueIsIgnored=False,bAdvancedView=False,)
+				   CustomProperties Pin (PinId=FD97B5EA441E8ED0C1E701A2A9811E29,PinName="X",PinToolTip="X\nFloat",PinType.PinCategory="float",PinType.PinSubCategory="",PinType.PinSubCategoryObject=None,PinType.PinSubCategoryMemberReference=(),PinType.PinValueType=(),PinType.bIsMap=False,PinType.bIsSet=False,PinType.bIsArray=False,PinType.bIsReference=False,PinType.bIsConst=False,PinType.bIsWeakPointer=False,DefaultValue="0.0",AutogeneratedDefaultValue="0.0",LinkedTo=(K2Node_CallFunction_59 B2007DF34F06C520502E3789A19668D9,),PersistentGuid=00000000000000000000000000000000,bHidden=False,bNotConnectable=False,bDefaultValueIsReadOnly=False,bDefaultValueIsIgnored=False,bAdvancedView=False,)
+				   CustomProperties Pin (PinId=6CF890D2486F8AD9D08636B1CBC925E2,PinName="Y",PinToolTip="Y\nFloat",PinType.PinCategory="float",PinType.PinSubCategory="",PinType.PinSubCategoryObject=None,PinType.PinSubCategoryMemberReference=(),PinType.PinValueType=(),PinType.bIsMap=False,PinType.bIsSet=False,PinType.bIsArray=False,PinType.bIsReference=False,PinType.bIsConst=False,PinType.bIsWeakPointer=False,DefaultValue="0.0",AutogeneratedDefaultValue="0.0",LinkedTo=(K2Node_CallFunction_59 BE68D9904F9FB271CCD42CB3A1957A09,),PersistentGuid=00000000000000000000000000000000,bHidden=False,bNotConnectable=False,bDefaultValueIsReadOnly=False,bDefaultValueIsIgnored=False,bAdvancedView=False,)
+				   CustomProperties Pin (PinId=B7F1E6AD41263F98928859BD496F8E6C,PinName="ReturnValue",PinToolTip="Return Value\nVector 2D Structure",Direction="EGPD_Output",PinType.PinCategory="struct",PinType.PinSubCategory="",PinType.PinSubCategoryObject=ScriptStruct'/Script/CoreUObject.Vector2D',PinType.PinSubCategoryMemberReference=(),PinType.PinValueType=(),PinType.bIsMap=False,PinType.bIsSet=False,PinType.bIsArray=False,PinType.bIsReference=False,PinType.bIsConst=False,PinType.bIsWeakPointer=False,LinkedTo=(K2Node_VariableSet_44 92B5DFF4428B093B99BF42843CCC00D0,K2Node_VariableSet_45 00FEB7B44882728B0EC9E9B60EF8696F,),PersistentGuid=00000000000000000000000000000000,bHidden=False,bNotConnectable=False,bDefaultValueIsReadOnly=False,bDefaultValueIsIgnored=False,bAdvancedView=False,)
+				End Object
+				Begin Object Class=/Script/BlueprintGraph.K2Node_VariableSet Name="K2Node_VariableSet_44"
+				   VariableReference=(MemberName="TouchStart",MemberGuid=8DE74C2C4BDD89D971F35FB55CD39E7D,bSelfContext=True)
+				   NodePosX=1184
+				   NodePosY=-736
+				   NodeGuid=C33B224047BB96961AB8E19E3AA50C19
+				   CustomProperties Pin (PinId=7AF704324F9FF65DC8729EA838A93153,PinName="execute",PinType.PinCategory="exec",PinType.PinSubCategory="",PinType.PinSubCategoryObject=None,PinType.PinSubCategoryMemberReference=(),PinType.PinValueType=(),PinType.bIsMap=False,PinType.bIsSet=False,PinType.bIsArray=False,PinType.bIsReference=False,PinType.bIsConst=False,PinType.bIsWeakPointer=False,LinkedTo=(K2Node_InputTouch_7 3787414C455842F10BAB63AB58C08F6A,),PersistentGuid=00000000000000000000000000000000,bHidden=False,bNotConnectable=False,bDefaultValueIsReadOnly=False,bDefaultValueIsIgnored=False,bAdvancedView=False,)
+				   CustomProperties Pin (PinId=635DD28B401A3AFB8DB83AB04CC96792,PinName="then",Direction="EGPD_Output",PinType.PinCategory="exec",PinType.PinSubCategory="",PinType.PinSubCategoryObject=None,PinType.PinSubCategoryMemberReference=(),PinType.PinValueType=(),PinType.bIsMap=False,PinType.bIsSet=False,PinType.bIsArray=False,PinType.bIsReference=False,PinType.bIsConst=False,PinType.bIsWeakPointer=False,PersistentGuid=00000000000000000000000000000000,bHidden=False,bNotConnectable=False,bDefaultValueIsReadOnly=False,bDefaultValueIsIgnored=False,bAdvancedView=False,)
+				   CustomProperties Pin (PinId=92B5DFF4428B093B99BF42843CCC00D0,PinName="TouchStart",PinType.PinCategory="struct",PinType.PinSubCategory="",PinType.PinSubCategoryObject=ScriptStruct'/Script/CoreUObject.Vector2D',PinType.PinSubCategoryMemberReference=(),PinType.PinValueType=(),PinType.bIsMap=False,PinType.bIsSet=False,PinType.bIsArray=False,PinType.bIsReference=False,PinType.bIsConst=False,PinType.bIsWeakPointer=False,LinkedTo=(K2Node_CallFunction_60 B7F1E6AD41263F98928859BD496F8E6C,),PersistentGuid=00000000000000000000000000000000,bHidden=False,bNotConnectable=False,bDefaultValueIsReadOnly=False,bDefaultValueIsIgnored=False,bAdvancedView=False,)
+				   CustomProperties Pin (PinId=062FBFD7426DCEF0D9734F9B5F7AA6FC,PinName="Output_Get",PinToolTip="Retrieves the value of the variable, can use instead of a separate Get node",Direction="EGPD_Output",PinType.PinCategory="struct",PinType.PinSubCategory="",PinType.PinSubCategoryObject=ScriptStruct'/Script/CoreUObject.Vector2D',PinType.PinSubCategoryMemberReference=(),PinType.PinValueType=(),PinType.bIsMap=False,PinType.bIsSet=False,PinType.bIsArray=False,PinType.bIsReference=False,PinType.bIsConst=False,PinType.bIsWeakPointer=False,PersistentGuid=00000000000000000000000000000000,bHidden=False,bNotConnectable=False,bDefaultValueIsReadOnly=False,bDefaultValueIsIgnored=False,bAdvancedView=False,)
+				   CustomProperties Pin (PinId=FABD218942CA35B7D100CA8A966F9442,PinName="self",PinFriendlyName=NSLOCTEXT("K2Node", "Target", "Target"),PinType.PinCategory="object",PinType.PinSubCategory="",PinType.PinSubCategoryObject=BlueprintGeneratedClass'/Game/GearVR_Pawn/GearVR_Pawn_Input.GearVR_Pawn_Input_C',PinType.PinSubCategoryMemberReference=(),PinType.PinValueType=(),PinType.bIsMap=False,PinType.bIsSet=False,PinType.bIsArray=False,PinType.bIsReference=False,PinType.bIsConst=False,PinType.bIsWeakPointer=False,PersistentGuid=00000000000000000000000000000000,bHidden=True,bNotConnectable=False,bDefaultValueIsReadOnly=False,bDefaultValueIsIgnored=False,bAdvancedView=False,)
+				End Object
+				Begin Object Class=/Script/BlueprintGraph.K2Node_VariableSet Name="K2Node_VariableSet_45"
+				   VariableReference=(MemberName="TouchEnd",MemberGuid=964A184E43E09A4A1CC71EA4F90CB54B,bSelfContext=True)
+				   NodePosX=1184
+				   NodePosY=-624
+				   NodeGuid=3077C3D7452F1326FEDEF58D5028F0C5
+				   CustomProperties Pin (PinId=E4E9F7E6497589AB9D3F3BA8D6C63943,PinName="execute",PinType.PinCategory="exec",PinType.PinSubCategory="",PinType.PinSubCategoryObject=None,PinType.PinSubCategoryMemberReference=(),PinType.PinValueType=(),PinType.bIsMap=False,PinType.bIsSet=False,PinType.bIsArray=False,PinType.bIsReference=False,PinType.bIsConst=False,PinType.bIsWeakPointer=False,LinkedTo=(K2Node_InputTouch_7 C6DEF3D74E3752A1346024AAD8FA60FE,),PersistentGuid=00000000000000000000000000000000,bHidden=False,bNotConnectable=False,bDefaultValueIsReadOnly=False,bDefaultValueIsIgnored=False,bAdvancedView=False,)
+				   CustomProperties Pin (PinId=2925450E4AE1D74A2BEF22B0B49305B0,PinName="then",Direction="EGPD_Output",PinType.PinCategory="exec",PinType.PinSubCategory="",PinType.PinSubCategoryObject=None,PinType.PinSubCategoryMemberReference=(),PinType.PinValueType=(),PinType.bIsMap=False,PinType.bIsSet=False,PinType.bIsArray=False,PinType.bIsReference=False,PinType.bIsConst=False,PinType.bIsWeakPointer=False,LinkedTo=(K2Node_IfThenElse_45 9DF08FE24D43CFADCBA446BB6EA29189,),PersistentGuid=00000000000000000000000000000000,bHidden=False,bNotConnectable=False,bDefaultValueIsReadOnly=False,bDefaultValueIsIgnored=False,bAdvancedView=False,)
+				   CustomProperties Pin (PinId=00FEB7B44882728B0EC9E9B60EF8696F,PinName="TouchEnd",PinType.PinCategory="struct",PinType.PinSubCategory="",PinType.PinSubCategoryObject=ScriptStruct'/Script/CoreUObject.Vector2D',PinType.PinSubCategoryMemberReference=(),PinType.PinValueType=(),PinType.bIsMap=False,PinType.bIsSet=False,PinType.bIsArray=False,PinType.bIsReference=False,PinType.bIsConst=False,PinType.bIsWeakPointer=False,LinkedTo=(K2Node_CallFunction_60 B7F1E6AD41263F98928859BD496F8E6C,),PersistentGuid=00000000000000000000000000000000,bHidden=False,bNotConnectable=False,bDefaultValueIsReadOnly=False,bDefaultValueIsIgnored=False,bAdvancedView=False,)
+				   CustomProperties Pin (PinId=E4C2564F40F2A43E28C9E794918B9C19,PinName="Output_Get",PinToolTip="Retrieves the value of the variable, can use instead of a separate Get node",Direction="EGPD_Output",PinType.PinCategory="struct",PinType.PinSubCategory="",PinType.PinSubCategoryObject=ScriptStruct'/Script/CoreUObject.Vector2D',PinType.PinSubCategoryMemberReference=(),PinType.PinValueType=(),PinType.bIsMap=False,PinType.bIsSet=False,PinType.bIsArray=False,PinType.bIsReference=False,PinType.bIsConst=False,PinType.bIsWeakPointer=False,PersistentGuid=00000000000000000000000000000000,bHidden=False,bNotConnectable=False,bDefaultValueIsReadOnly=False,bDefaultValueIsIgnored=False,bAdvancedView=False,)
+				   CustomProperties Pin (PinId=984BB526433CB00B9BAFB49C61B98B34,PinName="self",PinFriendlyName=NSLOCTEXT("K2Node", "Target", "Target"),PinType.PinCategory="object",PinType.PinSubCategory="",PinType.PinSubCategoryObject=BlueprintGeneratedClass'/Game/GearVR_Pawn/GearVR_Pawn_Input.GearVR_Pawn_Input_C',PinType.PinSubCategoryMemberReference=(),PinType.PinValueType=(),PinType.bIsMap=False,PinType.bIsSet=False,PinType.bIsArray=False,PinType.bIsReference=False,PinType.bIsConst=False,PinType.bIsWeakPointer=False,PersistentGuid=00000000000000000000000000000000,bHidden=True,bNotConnectable=False,bDefaultValueIsReadOnly=False,bDefaultValueIsIgnored=False,bAdvancedView=False,)
+				End Object
+				Begin Object Class=/Script/BlueprintGraph.K2Node_VariableGet Name="K2Node_VariableGet_61"
+				   VariableReference=(MemberName="TouchStart",MemberGuid=8DE74C2C4BDD89D971F35FB55CD39E7D,bSelfContext=True)
+				   NodePosX=288
+				   NodePosY=-144
+				   NodeGuid=94C9BA6E4A37A4D2C2808C8C733D7C11
+				   CustomProperties Pin (PinId=556A2EC34F3208A1D433DB84E123DB77,PinName="TouchStart",Direction="EGPD_Output",PinType.PinCategory="struct",PinType.PinSubCategory="",PinType.PinSubCategoryObject=ScriptStruct'/Script/CoreUObject.Vector2D',PinType.PinSubCategoryMemberReference=(),PinType.PinValueType=(),PinType.bIsMap=False,PinType.bIsSet=False,PinType.bIsArray=False,PinType.bIsReference=False,PinType.bIsConst=False,PinType.bIsWeakPointer=False,LinkedTo=(K2Node_CallFunction_62 B0F0B48C4411254049E4F6B9847CFC4D,),PersistentGuid=00000000000000000000000000000000,bHidden=False,bNotConnectable=False,bDefaultValueIsReadOnly=False,bDefaultValueIsIgnored=False,bAdvancedView=False,)
+				   CustomProperties Pin (PinId=DE15687844A3F2B9FEE27D8F9B34A47B,PinName="self",PinFriendlyName=NSLOCTEXT("K2Node", "Target", "Target"),PinType.PinCategory="object",PinType.PinSubCategory="",PinType.PinSubCategoryObject=BlueprintGeneratedClass'/Game/GearVR_Pawn/GearVR_Pawn_Input.GearVR_Pawn_Input_C',PinType.PinSubCategoryMemberReference=(),PinType.PinValueType=(),PinType.bIsMap=False,PinType.bIsSet=False,PinType.bIsArray=False,PinType.bIsReference=False,PinType.bIsConst=False,PinType.bIsWeakPointer=False,PersistentGuid=00000000000000000000000000000000,bHidden=True,bNotConnectable=False,bDefaultValueIsReadOnly=False,bDefaultValueIsIgnored=False,bAdvancedView=False,)
+				End Object
+				Begin Object Class=/Script/BlueprintGraph.K2Node_VariableGet Name="K2Node_VariableGet_62"
+				   VariableReference=(MemberName="TouchEnd",MemberGuid=964A184E43E09A4A1CC71EA4F90CB54B,bSelfContext=True)
+				   NodePosX=288
+				   NodePosY=-96
+				   NodeGuid=AD8470774B023AC7B0F599A4FFFF14AC
+				   CustomProperties Pin (PinId=2737D6564E1EA9A98BB19EA5E746E021,PinName="TouchEnd",Direction="EGPD_Output",PinType.PinCategory="struct",PinType.PinSubCategory="",PinType.PinSubCategoryObject=ScriptStruct'/Script/CoreUObject.Vector2D',PinType.PinSubCategoryMemberReference=(),PinType.PinValueType=(),PinType.bIsMap=False,PinType.bIsSet=False,PinType.bIsArray=False,PinType.bIsReference=False,PinType.bIsConst=False,PinType.bIsWeakPointer=False,LinkedTo=(K2Node_CallFunction_62 9F62C4D348F9F646517B259A080247A4,),PersistentGuid=00000000000000000000000000000000,bHidden=False,bNotConnectable=False,bDefaultValueIsReadOnly=False,bDefaultValueIsIgnored=False,bAdvancedView=False,)
+				   CustomProperties Pin (PinId=3F3DF57F4320A7E3AEF4228DC3A6636A,PinName="self",PinFriendlyName=NSLOCTEXT("K2Node", "Target", "Target"),PinType.PinCategory="object",PinType.PinSubCategory="",PinType.PinSubCategoryObject=BlueprintGeneratedClass'/Game/GearVR_Pawn/GearVR_Pawn_Input.GearVR_Pawn_Input_C',PinType.PinSubCategoryMemberReference=(),PinType.PinValueType=(),PinType.bIsMap=False,PinType.bIsSet=False,PinType.bIsArray=False,PinType.bIsReference=False,PinType.bIsConst=False,PinType.bIsWeakPointer=False,PersistentGuid=00000000000000000000000000000000,bHidden=True,bNotConnectable=False,bDefaultValueIsReadOnly=False,bDefaultValueIsIgnored=False,bAdvancedView=False,)
+				End Object
+				Begin Object Class=/Script/BlueprintGraph.K2Node_CallFunction Name="K2Node_CallFunction_62"
+				   bIsPureFunc=True
+				   FunctionReference=(MemberParent=Class'/Script/Engine.KismetMathLibrary',MemberName="Subtract_Vector2DVector2D")
+				   NodePosX=448
+				   NodePosY=-128
+				   NodeGuid=5DF360D4496C2F0CB47300AAB4804563
+				   CustomProperties Pin (PinId=08BEE81245B70E9AA91628BA03FECA4E,PinName="self",PinFriendlyName=NSLOCTEXT("K2Node", "Target", "Target"),PinType.PinCategory="object",PinType.PinSubCategory="",PinType.PinSubCategoryObject=Class'/Script/Engine.KismetMathLibrary',PinType.PinSubCategoryMemberReference=(),PinType.PinValueType=(),PinType.bIsMap=False,PinType.bIsSet=False,PinType.bIsArray=False,PinType.bIsReference=False,PinType.bIsConst=False,PinType.bIsWeakPointer=False,DefaultObject="/Script/Engine.Default__KismetMathLibrary",PersistentGuid=00000000000000000000000000000000,bHidden=True,bNotConnectable=False,bDefaultValueIsReadOnly=False,bDefaultValueIsIgnored=False,bAdvancedView=False,)
+				   CustomProperties Pin (PinId=B0F0B48C4411254049E4F6B9847CFC4D,PinName="A",PinType.PinCategory="struct",PinType.PinSubCategory="",PinType.PinSubCategoryObject=ScriptStruct'/Script/CoreUObject.Vector2D',PinType.PinSubCategoryMemberReference=(),PinType.PinValueType=(),PinType.bIsMap=False,PinType.bIsSet=False,PinType.bIsArray=False,PinType.bIsReference=False,PinType.bIsConst=False,PinType.bIsWeakPointer=False,LinkedTo=(K2Node_VariableGet_61 556A2EC34F3208A1D433DB84E123DB77,),PersistentGuid=00000000000000000000000000000000,bHidden=False,bNotConnectable=False,bDefaultValueIsReadOnly=False,bDefaultValueIsIgnored=False,bAdvancedView=False,)
+				   CustomProperties Pin (PinId=9F62C4D348F9F646517B259A080247A4,PinName="B",PinType.PinCategory="struct",PinType.PinSubCategory="",PinType.PinSubCategoryObject=ScriptStruct'/Script/CoreUObject.Vector2D',PinType.PinSubCategoryMemberReference=(),PinType.PinValueType=(),PinType.bIsMap=False,PinType.bIsSet=False,PinType.bIsArray=False,PinType.bIsReference=False,PinType.bIsConst=False,PinType.bIsWeakPointer=False,LinkedTo=(K2Node_VariableGet_62 2737D6564E1EA9A98BB19EA5E746E021,),PersistentGuid=00000000000000000000000000000000,bHidden=False,bNotConnectable=False,bDefaultValueIsReadOnly=False,bDefaultValueIsIgnored=False,bAdvancedView=False,)
+				   CustomProperties Pin (PinId=C1DBFE4F48CDFF1B9F4195BE118BD775,PinName="ReturnValue",Direction="EGPD_Output",PinType.PinCategory="struct",PinType.PinSubCategory="",PinType.PinSubCategoryObject=ScriptStruct'/Script/CoreUObject.Vector2D',PinType.PinSubCategoryMemberReference=(),PinType.PinValueType=(),PinType.bIsMap=False,PinType.bIsSet=False,PinType.bIsArray=False,PinType.bIsReference=False,PinType.bIsConst=False,PinType.bIsWeakPointer=False,LinkedTo=(K2Node_CallFunction_64 8F74F65D422A0C8687D5A19616C16E15,K2Node_CallFunction_156 C72184E64E9ABBFFF5E644941010A449,),PersistentGuid=00000000000000000000000000000000,bHidden=False,bNotConnectable=False,bDefaultValueIsReadOnly=False,bDefaultValueIsIgnored=False,bAdvancedView=False,)
+				End Object
+				Begin Object Class=/Script/BlueprintGraph.K2Node_CallFunction Name="K2Node_CallFunction_64"
+				   bIsPureFunc=True
+				   FunctionReference=(MemberParent=Class'/Script/Engine.KismetMathLibrary',MemberName="VSize2D")
+				   NodePosX=624
+				   NodePosY=-208
+				   NodeGuid=5C42A46041234A885BAAFF88D98D0EAC
+				   CustomProperties Pin (PinId=DE827473494FC8F98F4F3883FF85EBF2,PinName="self",PinFriendlyName=NSLOCTEXT("K2Node", "Target", "Target"),PinType.PinCategory="object",PinType.PinSubCategory="",PinType.PinSubCategoryObject=Class'/Script/Engine.KismetMathLibrary',PinType.PinSubCategoryMemberReference=(),PinType.PinValueType=(),PinType.bIsMap=False,PinType.bIsSet=False,PinType.bIsArray=False,PinType.bIsReference=False,PinType.bIsConst=False,PinType.bIsWeakPointer=False,DefaultObject="/Script/Engine.Default__KismetMathLibrary",PersistentGuid=00000000000000000000000000000000,bHidden=True,bNotConnectable=False,bDefaultValueIsReadOnly=False,bDefaultValueIsIgnored=False,bAdvancedView=False,)
+				   CustomProperties Pin (PinId=8F74F65D422A0C8687D5A19616C16E15,PinName="A",PinType.PinCategory="struct",PinType.PinSubCategory="",PinType.PinSubCategoryObject=ScriptStruct'/Script/CoreUObject.Vector2D',PinType.PinSubCategoryMemberReference=(),PinType.PinValueType=(),PinType.bIsMap=False,PinType.bIsSet=False,PinType.bIsArray=False,PinType.bIsReference=False,PinType.bIsConst=False,PinType.bIsWeakPointer=False,LinkedTo=(K2Node_CallFunction_62 C1DBFE4F48CDFF1B9F4195BE118BD775,),PersistentGuid=00000000000000000000000000000000,bHidden=False,bNotConnectable=False,bDefaultValueIsReadOnly=False,bDefaultValueIsIgnored=False,bAdvancedView=False,)
+				   CustomProperties Pin (PinId=81338FEE40BBDA2DB02213AA9B732A00,PinName="ReturnValue",Direction="EGPD_Output",PinType.PinCategory="float",PinType.PinSubCategory="",PinType.PinSubCategoryObject=None,PinType.PinSubCategoryMemberReference=(),PinType.PinValueType=(),PinType.bIsMap=False,PinType.bIsSet=False,PinType.bIsArray=False,PinType.bIsReference=False,PinType.bIsConst=False,PinType.bIsWeakPointer=False,AutogeneratedDefaultValue="0.0",LinkedTo=(K2Node_CallFunction_65 911A2C67468D633F8CD1379602009731,),PersistentGuid=00000000000000000000000000000000,bHidden=False,bNotConnectable=False,bDefaultValueIsReadOnly=False,bDefaultValueIsIgnored=False,bAdvancedView=False,)
+				End Object
+				Begin Object Class=/Script/BlueprintGraph.K2Node_CallFunction Name="K2Node_CallFunction_65"
+				   bIsPureFunc=True
+				   FunctionReference=(MemberParent=Class'/Script/Engine.KismetMathLibrary',MemberName="Greater_FloatFloat")
+				   NodePosX=912
+				   NodePosY=-182
+				   NodeGuid=F584D44C4BDCDBF9B28A1CA6B0438E2B
+				   CustomProperties Pin (PinId=EF6B98C54274B4FEFB841D9372517A51,PinName="self",PinFriendlyName=NSLOCTEXT("K2Node", "Target", "Target"),PinToolTip="Target\nKismet Math Library Reference",PinType.PinCategory="object",PinType.PinSubCategory="",PinType.PinSubCategoryObject=Class'/Script/Engine.KismetMathLibrary',PinType.PinSubCategoryMemberReference=(),PinType.PinValueType=(),PinType.bIsMap=False,PinType.bIsSet=False,PinType.bIsArray=False,PinType.bIsReference=False,PinType.bIsConst=False,PinType.bIsWeakPointer=False,DefaultObject="/Script/Engine.Default__KismetMathLibrary",PersistentGuid=00000000000000000000000000000000,bHidden=True,bNotConnectable=False,bDefaultValueIsReadOnly=False,bDefaultValueIsIgnored=False,bAdvancedView=False,)
+				   CustomProperties Pin (PinId=911A2C67468D633F8CD1379602009731,PinName="A",PinToolTip="A\nFloat",PinType.PinCategory="float",PinType.PinSubCategory="",PinType.PinSubCategoryObject=None,PinType.PinSubCategoryMemberReference=(),PinType.PinValueType=(),PinType.bIsMap=False,PinType.bIsSet=False,PinType.bIsArray=False,PinType.bIsReference=False,PinType.bIsConst=False,PinType.bIsWeakPointer=False,DefaultValue="0.0",AutogeneratedDefaultValue="0.0",LinkedTo=(K2Node_CallFunction_64 81338FEE40BBDA2DB02213AA9B732A00,),PersistentGuid=00000000000000000000000000000000,bHidden=False,bNotConnectable=False,bDefaultValueIsReadOnly=False,bDefaultValueIsIgnored=False,bAdvancedView=False,)
+				   CustomProperties Pin (PinId=A8CA1AFE4516F0B10D2CC581B865EB8D,PinName="B",PinToolTip="B\nFloat",PinType.PinCategory="float",PinType.PinSubCategory="",PinType.PinSubCategoryObject=None,PinType.PinSubCategoryMemberReference=(),PinType.PinValueType=(),PinType.bIsMap=False,PinType.bIsSet=False,PinType.bIsArray=False,PinType.bIsReference=False,PinType.bIsConst=False,PinType.bIsWeakPointer=False,DefaultValue="20",AutogeneratedDefaultValue="0.0",PersistentGuid=00000000000000000000000000000000,bHidden=False,bNotConnectable=False,bDefaultValueIsReadOnly=False,bDefaultValueIsIgnored=False,bAdvancedView=False,)
+				   CustomProperties Pin (PinId=23EB014D4DE7275C055DFF94E10B5584,PinName="ReturnValue",PinToolTip="Return Value\nBoolean",Direction="EGPD_Output",PinType.PinCategory="bool",PinType.PinSubCategory="",PinType.PinSubCategoryObject=None,PinType.PinSubCategoryMemberReference=(),PinType.PinValueType=(),PinType.bIsMap=False,PinType.bIsSet=False,PinType.bIsArray=False,PinType.bIsReference=False,PinType.bIsConst=False,PinType.bIsWeakPointer=False,AutogeneratedDefaultValue="false",LinkedTo=(K2Node_IfThenElse_45 EC64745F41E78DF32486608D583AE2E6,),PersistentGuid=00000000000000000000000000000000,bHidden=False,bNotConnectable=False,bDefaultValueIsReadOnly=False,bDefaultValueIsIgnored=False,bAdvancedView=False,)
+				End Object
+				Begin Object Class=/Script/BlueprintGraph.K2Node_IfThenElse Name="K2Node_IfThenElse_45"
+				   NodePosX=1648
+				   NodePosY=-640
+				   NodeGuid=FC78736641ED350F51ED13BE70826619
+				   CustomProperties Pin (PinId=9DF08FE24D43CFADCBA446BB6EA29189,PinName="execute",PinType.PinCategory="exec",PinType.PinSubCategory="",PinType.PinSubCategoryObject=None,PinType.PinSubCategoryMemberReference=(),PinType.PinValueType=(),PinType.bIsMap=False,PinType.bIsSet=False,PinType.bIsArray=False,PinType.bIsReference=False,PinType.bIsConst=False,PinType.bIsWeakPointer=False,LinkedTo=(K2Node_VariableSet_45 2925450E4AE1D74A2BEF22B0B49305B0,),PersistentGuid=00000000000000000000000000000000,bHidden=False,bNotConnectable=False,bDefaultValueIsReadOnly=False,bDefaultValueIsIgnored=False,bAdvancedView=False,)
+				   CustomProperties Pin (PinId=EC64745F41E78DF32486608D583AE2E6,PinName="Condition",PinType.PinCategory="bool",PinType.PinSubCategory="",PinType.PinSubCategoryObject=None,PinType.PinSubCategoryMemberReference=(),PinType.PinValueType=(),PinType.bIsMap=False,PinType.bIsSet=False,PinType.bIsArray=False,PinType.bIsReference=False,PinType.bIsConst=False,PinType.bIsWeakPointer=False,DefaultValue="false",AutogeneratedDefaultValue="true",LinkedTo=(K2Node_CallFunction_65 23EB014D4DE7275C055DFF94E10B5584,),PersistentGuid=00000000000000000000000000000000,bHidden=False,bNotConnectable=False,bDefaultValueIsReadOnly=False,bDefaultValueIsIgnored=False,bAdvancedView=False,)
+				   CustomProperties Pin (PinId=3306E19940BBEA4DB04353976F98E9D8,PinName="then",PinFriendlyName=NSLOCTEXT("K2Node", "true", "true"),Direction="EGPD_Output",PinType.PinCategory="exec",PinType.PinSubCategory="",PinType.PinSubCategoryObject=None,PinType.PinSubCategoryMemberReference=(),PinType.PinValueType=(),PinType.bIsMap=False,PinType.bIsSet=False,PinType.bIsArray=False,PinType.bIsReference=False,PinType.bIsConst=False,PinType.bIsWeakPointer=False,LinkedTo=(K2Node_IfThenElse_10 9DF08FE24D43CFADCBA446BB6EA29189,),PersistentGuid=00000000000000000000000000000000,bHidden=False,bNotConnectable=False,bDefaultValueIsReadOnly=False,bDefaultValueIsIgnored=False,bAdvancedView=False,)
+				   CustomProperties Pin (PinId=9B8358F143000BB08762918D0E75B5E9,PinName="else",PinFriendlyName=NSLOCTEXT("K2Node", "false", "false"),Direction="EGPD_Output",PinType.PinCategory="exec",PinType.PinSubCategory="",PinType.PinSubCategoryObject=None,PinType.PinSubCategoryMemberReference=(),PinType.PinValueType=(),PinType.bIsMap=False,PinType.bIsSet=False,PinType.bIsArray=False,PinType.bIsReference=False,PinType.bIsConst=False,PinType.bIsWeakPointer=False,LinkedTo=(K2Node_CallFunction_13198 96D7193F444562BF1B9D1A82C15146E1,),PersistentGuid=00000000000000000000000000000000,bHidden=False,bNotConnectable=False,bDefaultValueIsReadOnly=False,bDefaultValueIsIgnored=False,bAdvancedView=False,)
+				End Object
+				Begin Object Class=/Script/BlueprintGraph.K2Node_IfThenElse Name="K2Node_IfThenElse_10"
+				   NodePosX=2432
+				   NodePosY=-272
+				   NodeGuid=9AB8326446E3F0C0E5D6BABC6669EC7A
+				   CustomProperties Pin (PinId=9DF08FE24D43CFADCBA446BB6EA29189,PinName="execute",PinType.PinCategory="exec",PinType.PinSubCategory="",PinType.PinSubCategoryObject=None,PinType.PinSubCategoryMemberReference=(),PinType.PinValueType=(),PinType.bIsMap=False,PinType.bIsSet=False,PinType.bIsArray=False,PinType.bIsReference=False,PinType.bIsConst=False,PinType.bIsWeakPointer=False,LinkedTo=(K2Node_IfThenElse_45 3306E19940BBEA4DB04353976F98E9D8,),PersistentGuid=00000000000000000000000000000000,bHidden=False,bNotConnectable=False,bDefaultValueIsReadOnly=False,bDefaultValueIsIgnored=False,bAdvancedView=False,)
+				   CustomProperties Pin (PinId=EC64745F41E78DF32486608D583AE2E6,PinName="Condition",PinType.PinCategory="bool",PinType.PinSubCategory="",PinType.PinSubCategoryObject=None,PinType.PinSubCategoryMemberReference=(),PinType.PinValueType=(),PinType.bIsMap=False,PinType.bIsSet=False,PinType.bIsArray=False,PinType.bIsReference=False,PinType.bIsConst=False,PinType.bIsWeakPointer=False,DefaultValue="false",AutogeneratedDefaultValue="true",LinkedTo=(K2Node_CallFunction_159 8C0BEFFE4324D8480AA6C191F569AEDB,),PersistentGuid=00000000000000000000000000000000,bHidden=False,bNotConnectable=False,bDefaultValueIsReadOnly=False,bDefaultValueIsIgnored=False,bAdvancedView=False,)
+				   CustomProperties Pin (PinId=3306E19940BBEA4DB04353976F98E9D8,PinName="then",PinFriendlyName=NSLOCTEXT("K2Node", "true", "true"),Direction="EGPD_Output",PinType.PinCategory="exec",PinType.PinSubCategory="",PinType.PinSubCategoryObject=None,PinType.PinSubCategoryMemberReference=(),PinType.PinValueType=(),PinType.bIsMap=False,PinType.bIsSet=False,PinType.bIsArray=False,PinType.bIsReference=False,PinType.bIsConst=False,PinType.bIsWeakPointer=False,LinkedTo=(K2Node_IfThenElse_14 9DF08FE24D43CFADCBA446BB6EA29189,),PersistentGuid=00000000000000000000000000000000,bHidden=False,bNotConnectable=False,bDefaultValueIsReadOnly=False,bDefaultValueIsIgnored=False,bAdvancedView=False,)
+				   CustomProperties Pin (PinId=9B8358F143000BB08762918D0E75B5E9,PinName="else",PinFriendlyName=NSLOCTEXT("K2Node", "false", "false"),Direction="EGPD_Output",PinType.PinCategory="exec",PinType.PinSubCategory="",PinType.PinSubCategoryObject=None,PinType.PinSubCategoryMemberReference=(),PinType.PinValueType=(),PinType.bIsMap=False,PinType.bIsSet=False,PinType.bIsArray=False,PinType.bIsReference=False,PinType.bIsConst=False,PinType.bIsWeakPointer=False,LinkedTo=(K2Node_IfThenElse_0 9DF08FE24D43CFADCBA446BB6EA29189,),PersistentGuid=00000000000000000000000000000000,bHidden=False,bNotConnectable=False,bDefaultValueIsReadOnly=False,bDefaultValueIsIgnored=False,bAdvancedView=False,)
+				End Object
+				Begin Object Class=/Script/UnrealEd.EdGraphNode_Comment Name="EdGraphNode_Comment_13"
+				   NodePosX=32
+				   NodePosY=-528
+				   NodeWidth=944
+				   NodeHeight=224
+				   NodeComment="Get the first touch position X & Y locations and make them into a 2D Vector"
+				   NodeGuid=2E41EE5E47E948D9FF1929890B2E59BF
+				End Object
+				Begin Object Class=/Script/UnrealEd.EdGraphNode_Comment Name="EdGraphNode_Comment_2"
+				   NodePosX=1136
+				   NodePosY=-816
+				   NodeWidth=336
+				   NodeHeight=304
+				   NodeComment="Start & End point of a touch"
+				   NodeGuid=BEEFA3A942F2558A340CC7A0695DD0D9
+				End Object
+				Begin Object Class=/Script/UnrealEd.EdGraphNode_Comment Name="EdGraphNode_Comment_3"
+				   NodePosX=240
+				   NodePosY=-272
+				   NodeWidth=1040
+				   NodeHeight=480
+				   NodeComment="Touch / Swipe Logic"
+				   NodeGuid=FD3322D84CBA38CF215D1FA3F63BE584
+				End Object
+				Begin Object Class=/Script/BlueprintGraph.K2Node_CallFunction Name="K2Node_CallFunction_156"
+				   bIsPureFunc=True
+				   FunctionReference=(MemberParent=Class'/Script/Engine.KismetMathLibrary',MemberName="BreakVector2D")
+				   NodePosX=624
+				   NodePosY=-96
+				   NodeGuid=FD33CB8C48ED847C0EF1E0AA87782F63
+				   CustomProperties Pin (PinId=34D8931645D3182501303D944D34A760,PinName="self",PinFriendlyName=NSLOCTEXT("K2Node", "Target", "Target"),PinToolTip="Target\nKismet Math Library Reference",PinType.PinCategory="object",PinType.PinSubCategory="",PinType.PinSubCategoryObject=Class'/Script/Engine.KismetMathLibrary',PinType.PinSubCategoryMemberReference=(),PinType.PinValueType=(),PinType.bIsMap=False,PinType.bIsSet=False,PinType.bIsArray=False,PinType.bIsReference=False,PinType.bIsConst=False,PinType.bIsWeakPointer=False,DefaultObject="/Script/Engine.Default__KismetMathLibrary",PersistentGuid=00000000000000000000000000000000,bHidden=True,bNotConnectable=False,bDefaultValueIsReadOnly=False,bDefaultValueIsIgnored=False,bAdvancedView=False,)
+				   CustomProperties Pin (PinId=C72184E64E9ABBFFF5E644941010A449,PinName="InVec",PinToolTip="In Vec\nVector 2D Structure",PinType.PinCategory="struct",PinType.PinSubCategory="",PinType.PinSubCategoryObject=ScriptStruct'/Script/CoreUObject.Vector2D',PinType.PinSubCategoryMemberReference=(),PinType.PinValueType=(),PinType.bIsMap=False,PinType.bIsSet=False,PinType.bIsArray=False,PinType.bIsReference=False,PinType.bIsConst=False,PinType.bIsWeakPointer=False,LinkedTo=(K2Node_CallFunction_62 C1DBFE4F48CDFF1B9F4195BE118BD775,),PersistentGuid=00000000000000000000000000000000,bHidden=False,bNotConnectable=False,bDefaultValueIsReadOnly=False,bDefaultValueIsIgnored=False,bAdvancedView=False,)
+				   CustomProperties Pin (PinId=E1512AD84E1F9D6207A10A8B26039124,PinName="X",PinToolTip="X\nFloat",Direction="EGPD_Output",PinType.PinCategory="float",PinType.PinSubCategory="",PinType.PinSubCategoryObject=None,PinType.PinSubCategoryMemberReference=(),PinType.PinValueType=(),PinType.bIsMap=False,PinType.bIsSet=False,PinType.bIsArray=False,PinType.bIsReference=False,PinType.bIsConst=False,PinType.bIsWeakPointer=False,AutogeneratedDefaultValue="0.0",LinkedTo=(K2Node_CallFunction_158 DACE98334FF14507C2FD779CE7B62549,K2Node_CallFunction_138 88CF46414C43AA30CD71F39D9FDDCE8C,),PersistentGuid=00000000000000000000000000000000,bHidden=False,bNotConnectable=False,bDefaultValueIsReadOnly=False,bDefaultValueIsIgnored=False,bAdvancedView=False,)
+				   CustomProperties Pin (PinId=4F93D2654F1A5618FC40B3B36B19A2A5,PinName="Y",PinToolTip="Y\nFloat",Direction="EGPD_Output",PinType.PinCategory="float",PinType.PinSubCategory="",PinType.PinSubCategoryObject=None,PinType.PinSubCategoryMemberReference=(),PinType.PinValueType=(),PinType.bIsMap=False,PinType.bIsSet=False,PinType.bIsArray=False,PinType.bIsReference=False,PinType.bIsConst=False,PinType.bIsWeakPointer=False,AutogeneratedDefaultValue="0.0",LinkedTo=(K2Node_CallFunction_157 DACE98334FF14507C2FD779CE7B62549,K2Node_CallFunction_160 1F99D557483D16A0A0E1AF849A70FA59,),PersistentGuid=00000000000000000000000000000000,bHidden=False,bNotConnectable=False,bDefaultValueIsReadOnly=False,bDefaultValueIsIgnored=False,bAdvancedView=False,)
+				End Object
+				Begin Object Class=/Script/BlueprintGraph.K2Node_CallFunction Name="K2Node_CallFunction_158"
+				   bIsPureFunc=True
+				   FunctionReference=(MemberParent=Class'/Script/Engine.KismetMathLibrary',MemberName="Abs")
+				   NodePosX=912
+				   NodeGuid=19AD0CB84071ED498A2129AFFD9D14FE
+				   CustomProperties Pin (PinId=28C10D244B91255BB48754B6F6618449,PinName="self",PinFriendlyName=NSLOCTEXT("K2Node", "Target", "Target"),PinToolTip="Target\nKismet Math Library Reference",PinType.PinCategory="object",PinType.PinSubCategory="",PinType.PinSubCategoryObject=Class'/Script/Engine.KismetMathLibrary',PinType.PinSubCategoryMemberReference=(),PinType.PinValueType=(),PinType.bIsMap=False,PinType.bIsSet=False,PinType.bIsArray=False,PinType.bIsReference=False,PinType.bIsConst=False,PinType.bIsWeakPointer=False,DefaultObject="/Script/Engine.Default__KismetMathLibrary",PersistentGuid=00000000000000000000000000000000,bHidden=True,bNotConnectable=False,bDefaultValueIsReadOnly=False,bDefaultValueIsIgnored=False,bAdvancedView=False,)
+				   CustomProperties Pin (PinId=DACE98334FF14507C2FD779CE7B62549,PinName="A",PinToolTip="A\nFloat",PinType.PinCategory="float",PinType.PinSubCategory="",PinType.PinSubCategoryObject=None,PinType.PinSubCategoryMemberReference=(),PinType.PinValueType=(),PinType.bIsMap=False,PinType.bIsSet=False,PinType.bIsArray=False,PinType.bIsReference=False,PinType.bIsConst=False,PinType.bIsWeakPointer=False,DefaultValue="0.0",AutogeneratedDefaultValue="0.0",LinkedTo=(K2Node_CallFunction_156 E1512AD84E1F9D6207A10A8B26039124,),PersistentGuid=00000000000000000000000000000000,bHidden=False,bNotConnectable=False,bDefaultValueIsReadOnly=False,bDefaultValueIsIgnored=False,bAdvancedView=False,)
+				   CustomProperties Pin (PinId=CE29B3BB45BD9FC92E5D91955E5A2B2A,PinName="ReturnValue",PinToolTip="Return Value\nFloat",Direction="EGPD_Output",PinType.PinCategory="float",PinType.PinSubCategory="",PinType.PinSubCategoryObject=None,PinType.PinSubCategoryMemberReference=(),PinType.PinValueType=(),PinType.bIsMap=False,PinType.bIsSet=False,PinType.bIsArray=False,PinType.bIsReference=False,PinType.bIsConst=False,PinType.bIsWeakPointer=False,AutogeneratedDefaultValue="0.0",LinkedTo=(K2Node_CallFunction_159 C0BFCDFC40B41F33AC64DFBBC22A7EDF,),PersistentGuid=00000000000000000000000000000000,bHidden=False,bNotConnectable=False,bDefaultValueIsReadOnly=False,bDefaultValueIsIgnored=False,bAdvancedView=False,)
+				End Object
+				Begin Object Class=/Script/BlueprintGraph.K2Node_CallFunction Name="K2Node_CallFunction_157"
+				   bIsPureFunc=True
+				   FunctionReference=(MemberParent=Class'/Script/Engine.KismetMathLibrary',MemberName="Abs")
+				   NodePosX=912
+				   NodePosY=48
+				   NodeGuid=14B27AF644CFA851407491B339C0FC50
+				   CustomProperties Pin (PinId=28C10D244B91255BB48754B6F6618449,PinName="self",PinFriendlyName=NSLOCTEXT("K2Node", "Target", "Target"),PinToolTip="Target\nKismet Math Library Reference",PinType.PinCategory="object",PinType.PinSubCategory="",PinType.PinSubCategoryObject=Class'/Script/Engine.KismetMathLibrary',PinType.PinSubCategoryMemberReference=(),PinType.PinValueType=(),PinType.bIsMap=False,PinType.bIsSet=False,PinType.bIsArray=False,PinType.bIsReference=False,PinType.bIsConst=False,PinType.bIsWeakPointer=False,DefaultObject="/Script/Engine.Default__KismetMathLibrary",PersistentGuid=00000000000000000000000000000000,bHidden=True,bNotConnectable=False,bDefaultValueIsReadOnly=False,bDefaultValueIsIgnored=False,bAdvancedView=False,)
+				   CustomProperties Pin (PinId=DACE98334FF14507C2FD779CE7B62549,PinName="A",PinToolTip="A\nFloat",PinType.PinCategory="float",PinType.PinSubCategory="",PinType.PinSubCategoryObject=None,PinType.PinSubCategoryMemberReference=(),PinType.PinValueType=(),PinType.bIsMap=False,PinType.bIsSet=False,PinType.bIsArray=False,PinType.bIsReference=False,PinType.bIsConst=False,PinType.bIsWeakPointer=False,DefaultValue="0.0",AutogeneratedDefaultValue="0.0",LinkedTo=(K2Node_CallFunction_156 4F93D2654F1A5618FC40B3B36B19A2A5,),PersistentGuid=00000000000000000000000000000000,bHidden=False,bNotConnectable=False,bDefaultValueIsReadOnly=False,bDefaultValueIsIgnored=False,bAdvancedView=False,)
+				   CustomProperties Pin (PinId=CE29B3BB45BD9FC92E5D91955E5A2B2A,PinName="ReturnValue",PinToolTip="Return Value\nFloat",Direction="EGPD_Output",PinType.PinCategory="float",PinType.PinSubCategory="",PinType.PinSubCategoryObject=None,PinType.PinSubCategoryMemberReference=(),PinType.PinValueType=(),PinType.bIsMap=False,PinType.bIsSet=False,PinType.bIsArray=False,PinType.bIsReference=False,PinType.bIsConst=False,PinType.bIsWeakPointer=False,AutogeneratedDefaultValue="0.0",LinkedTo=(K2Node_CallFunction_159 7C7478794A78EE22C8C6C0AE6780DD61,),PersistentGuid=00000000000000000000000000000000,bHidden=False,bNotConnectable=False,bDefaultValueIsReadOnly=False,bDefaultValueIsIgnored=False,bAdvancedView=False,)
+				End Object
+				Begin Object Class=/Script/BlueprintGraph.K2Node_CallFunction Name="K2Node_CallFunction_159"
+				   bIsPureFunc=True
+				   FunctionReference=(MemberParent=Class'/Script/Engine.KismetMathLibrary',MemberName="Greater_FloatFloat")
+				   NodePosX=1088
+				   NodePosY=16
+				   NodeGuid=5ABED29941A9ABC22B12E8930CD73A50
+				   CustomProperties Pin (PinId=BF766DF048159F72CCB51EBF91972004,PinName="self",PinFriendlyName=NSLOCTEXT("K2Node", "Target", "Target"),PinToolTip="Target\nKismet Math Library Reference",PinType.PinCategory="object",PinType.PinSubCategory="",PinType.PinSubCategoryObject=Class'/Script/Engine.KismetMathLibrary',PinType.PinSubCategoryMemberReference=(),PinType.PinValueType=(),PinType.bIsMap=False,PinType.bIsSet=False,PinType.bIsArray=False,PinType.bIsReference=False,PinType.bIsConst=False,PinType.bIsWeakPointer=False,DefaultObject="/Script/Engine.Default__KismetMathLibrary",PersistentGuid=00000000000000000000000000000000,bHidden=True,bNotConnectable=False,bDefaultValueIsReadOnly=False,bDefaultValueIsIgnored=False,bAdvancedView=False,)
+				   CustomProperties Pin (PinId=C0BFCDFC40B41F33AC64DFBBC22A7EDF,PinName="A",PinToolTip="A\nFloat",PinType.PinCategory="float",PinType.PinSubCategory="",PinType.PinSubCategoryObject=None,PinType.PinSubCategoryMemberReference=(),PinType.PinValueType=(),PinType.bIsMap=False,PinType.bIsSet=False,PinType.bIsArray=False,PinType.bIsReference=False,PinType.bIsConst=False,PinType.bIsWeakPointer=False,DefaultValue="0.0",AutogeneratedDefaultValue="0.0",LinkedTo=(K2Node_CallFunction_158 CE29B3BB45BD9FC92E5D91955E5A2B2A,),PersistentGuid=00000000000000000000000000000000,bHidden=False,bNotConnectable=False,bDefaultValueIsReadOnly=False,bDefaultValueIsIgnored=False,bAdvancedView=False,)
+				   CustomProperties Pin (PinId=7C7478794A78EE22C8C6C0AE6780DD61,PinName="B",PinToolTip="B\nFloat",PinType.PinCategory="float",PinType.PinSubCategory="",PinType.PinSubCategoryObject=None,PinType.PinSubCategoryMemberReference=(),PinType.PinValueType=(),PinType.bIsMap=False,PinType.bIsSet=False,PinType.bIsArray=False,PinType.bIsReference=False,PinType.bIsConst=False,PinType.bIsWeakPointer=False,DefaultValue="0.0",AutogeneratedDefaultValue="0.0",LinkedTo=(K2Node_CallFunction_157 CE29B3BB45BD9FC92E5D91955E5A2B2A,),PersistentGuid=00000000000000000000000000000000,bHidden=False,bNotConnectable=False,bDefaultValueIsReadOnly=False,bDefaultValueIsIgnored=False,bAdvancedView=False,)
+				   CustomProperties Pin (PinId=8C0BEFFE4324D8480AA6C191F569AEDB,PinName="ReturnValue",PinToolTip="Return Value\nBoolean",Direction="EGPD_Output",PinType.PinCategory="bool",PinType.PinSubCategory="",PinType.PinSubCategoryObject=None,PinType.PinSubCategoryMemberReference=(),PinType.PinValueType=(),PinType.bIsMap=False,PinType.bIsSet=False,PinType.bIsArray=False,PinType.bIsReference=False,PinType.bIsConst=False,PinType.bIsWeakPointer=False,AutogeneratedDefaultValue="false",LinkedTo=(K2Node_IfThenElse_10 EC64745F41E78DF32486608D583AE2E6,),PersistentGuid=00000000000000000000000000000000,bHidden=False,bNotConnectable=False,bDefaultValueIsReadOnly=False,bDefaultValueIsIgnored=False,bAdvancedView=False,)
+				End Object
+				Begin Object Class=/Script/BlueprintGraph.K2Node_CallFunction Name="K2Node_CallFunction_160"
+				   bIsPureFunc=True
+				   FunctionReference=(MemberParent=Class'/Script/Engine.KismetMathLibrary',MemberName="Greater_FloatFloat")
+				   NodePosX=912
+				   NodePosY=128
+				   NodeGuid=EE02861A418235C5ECBD2482CF2245B2
+				   CustomProperties Pin (PinId=E29DED344FECF8C1C4AE3C8BABACDBB9,PinName="self",PinFriendlyName=NSLOCTEXT("K2Node", "Target", "Target"),PinToolTip="Target\nKismet Math Library Reference",PinType.PinCategory="object",PinType.PinSubCategory="",PinType.PinSubCategoryObject=Class'/Script/Engine.KismetMathLibrary',PinType.PinSubCategoryMemberReference=(),PinType.PinValueType=(),PinType.bIsMap=False,PinType.bIsSet=False,PinType.bIsArray=False,PinType.bIsReference=False,PinType.bIsConst=False,PinType.bIsWeakPointer=False,DefaultObject="/Script/Engine.Default__KismetMathLibrary",PersistentGuid=00000000000000000000000000000000,bHidden=True,bNotConnectable=False,bDefaultValueIsReadOnly=False,bDefaultValueIsIgnored=False,bAdvancedView=False,)
+				   CustomProperties Pin (PinId=1F99D557483D16A0A0E1AF849A70FA59,PinName="A",PinToolTip="A\nFloat",PinType.PinCategory="float",PinType.PinSubCategory="",PinType.PinSubCategoryObject=None,PinType.PinSubCategoryMemberReference=(),PinType.PinValueType=(),PinType.bIsMap=False,PinType.bIsSet=False,PinType.bIsArray=False,PinType.bIsReference=False,PinType.bIsConst=False,PinType.bIsWeakPointer=False,DefaultValue="0.0",AutogeneratedDefaultValue="0.0",LinkedTo=(K2Node_CallFunction_156 4F93D2654F1A5618FC40B3B36B19A2A5,),PersistentGuid=00000000000000000000000000000000,bHidden=False,bNotConnectable=False,bDefaultValueIsReadOnly=False,bDefaultValueIsIgnored=False,bAdvancedView=False,)
+				   CustomProperties Pin (PinId=2977D230440E18DF51DA33808BAD8365,PinName="B",PinToolTip="B\nFloat",PinType.PinCategory="float",PinType.PinSubCategory="",PinType.PinSubCategoryObject=None,PinType.PinSubCategoryMemberReference=(),PinType.PinValueType=(),PinType.bIsMap=False,PinType.bIsSet=False,PinType.bIsArray=False,PinType.bIsReference=False,PinType.bIsConst=False,PinType.bIsWeakPointer=False,DefaultValue="1",AutogeneratedDefaultValue="0.0",PersistentGuid=00000000000000000000000000000000,bHidden=False,bNotConnectable=False,bDefaultValueIsReadOnly=False,bDefaultValueIsIgnored=False,bAdvancedView=False,)
+				   CustomProperties Pin (PinId=04A5FFA94B808583861470BA77C74C63,PinName="ReturnValue",PinToolTip="Return Value\nBoolean",Direction="EGPD_Output",PinType.PinCategory="bool",PinType.PinSubCategory="",PinType.PinSubCategoryObject=None,PinType.PinSubCategoryMemberReference=(),PinType.PinValueType=(),PinType.bIsMap=False,PinType.bIsSet=False,PinType.bIsArray=False,PinType.bIsReference=False,PinType.bIsConst=False,PinType.bIsWeakPointer=False,AutogeneratedDefaultValue="false",LinkedTo=(K2Node_IfThenElse_0 EC64745F41E78DF32486608D583AE2E6,),PersistentGuid=00000000000000000000000000000000,bHidden=False,bNotConnectable=False,bDefaultValueIsReadOnly=False,bDefaultValueIsIgnored=False,bAdvancedView=False,)
+				End Object
+				Begin Object Class=/Script/UnrealEd.EdGraphNode_Comment Name="EdGraphNode_Comment_18"
+				   NodePosX=1632
+				   NodePosY=-720
+				   NodeWidth=208
+				   NodeHeight=192
+				   NodeComment="True = Swipe False = Tap"
+				   NodeGuid=8FF7FBD94487CA5EE7862595A657CFEE
+				End Object
+				Begin Object Class=/Script/UnrealEd.EdGraphNode_Comment Name="EdGraphNode_Comment_19"
+				   NodePosX=2400
+				   NodePosY=-352
+				   NodeWidth=240
+				   NodeHeight=192
+				   NodeComment="True = Left / Right False = Up / Down"
+				   NodeGuid=5ACEDB1048672141A3D5BF89AD9A9142
+				End Object
+				Begin Object Class=/Script/BlueprintGraph.K2Node_IfThenElse Name="K2Node_IfThenElse_14"
+				   NodePosX=2768
+				   NodePosY=-496
+				   NodeGuid=573C18B14C367E5BC8AD72ACBED464C4
+				   CustomProperties Pin (PinId=9DF08FE24D43CFADCBA446BB6EA29189,PinName="execute",PinType.PinCategory="exec",PinType.PinSubCategory="",PinType.PinSubCategoryObject=None,PinType.PinSubCategoryMemberReference=(),PinType.PinValueType=(),PinType.bIsMap=False,PinType.bIsSet=False,PinType.bIsArray=False,PinType.bIsReference=False,PinType.bIsConst=False,PinType.bIsWeakPointer=False,LinkedTo=(K2Node_IfThenElse_10 3306E19940BBEA4DB04353976F98E9D8,),PersistentGuid=00000000000000000000000000000000,bHidden=False,bNotConnectable=False,bDefaultValueIsReadOnly=False,bDefaultValueIsIgnored=False,bAdvancedView=False,)
+				   CustomProperties Pin (PinId=EC64745F41E78DF32486608D583AE2E6,PinName="Condition",PinType.PinCategory="bool",PinType.PinSubCategory="",PinType.PinSubCategoryObject=None,PinType.PinSubCategoryMemberReference=(),PinType.PinValueType=(),PinType.bIsMap=False,PinType.bIsSet=False,PinType.bIsArray=False,PinType.bIsReference=False,PinType.bIsConst=False,PinType.bIsWeakPointer=False,DefaultValue="false",AutogeneratedDefaultValue="true",LinkedTo=(K2Node_CallFunction_138 97BE87AA448C93D13187DE941B2FF014,),PersistentGuid=00000000000000000000000000000000,bHidden=False,bNotConnectable=False,bDefaultValueIsReadOnly=False,bDefaultValueIsIgnored=False,bAdvancedView=False,)
+				   CustomProperties Pin (PinId=3306E19940BBEA4DB04353976F98E9D8,PinName="then",PinFriendlyName=NSLOCTEXT("K2Node", "true", "true"),Direction="EGPD_Output",PinType.PinCategory="exec",PinType.PinSubCategory="",PinType.PinSubCategoryObject=None,PinType.PinSubCategoryMemberReference=(),PinType.PinValueType=(),PinType.bIsMap=False,PinType.bIsSet=False,PinType.bIsArray=False,PinType.bIsReference=False,PinType.bIsConst=False,PinType.bIsWeakPointer=False,LinkedTo=(K2Node_CallFunction_12917 96D7193F444562BF1B9D1A82C15146E1,),PersistentGuid=00000000000000000000000000000000,bHidden=False,bNotConnectable=False,bDefaultValueIsReadOnly=False,bDefaultValueIsIgnored=False,bAdvancedView=False,)
+				   CustomProperties Pin (PinId=9B8358F143000BB08762918D0E75B5E9,PinName="else",PinFriendlyName=NSLOCTEXT("K2Node", "false", "false"),Direction="EGPD_Output",PinType.PinCategory="exec",PinType.PinSubCategory="",PinType.PinSubCategoryObject=None,PinType.PinSubCategoryMemberReference=(),PinType.PinValueType=(),PinType.bIsMap=False,PinType.bIsSet=False,PinType.bIsArray=False,PinType.bIsReference=False,PinType.bIsConst=False,PinType.bIsWeakPointer=False,LinkedTo=(K2Node_CallFunction_12871 96D7193F444562BF1B9D1A82C15146E1,),PersistentGuid=00000000000000000000000000000000,bHidden=False,bNotConnectable=False,bDefaultValueIsReadOnly=False,bDefaultValueIsIgnored=False,bAdvancedView=False,)
+				End Object
+				Begin Object Class=/Script/UnrealEd.EdGraphNode_Comment Name="EdGraphNode_Comment_26"
+				   NodePosX=2704
+				   NodePosY=-576
+				   NodeWidth=256
+				   NodeHeight=192
+				   NodeComment="True  = Swipe Right False = Swipe Left"
+				   NodeGuid=F44179054F2B9439D769F39FC3E71C99
+				End Object
+				Begin Object Class=/Script/BlueprintGraph.K2Node_CallFunction Name="K2Node_CallFunction_138"
+				   bIsPureFunc=True
+				   FunctionReference=(MemberParent=Class'/Script/Engine.KismetMathLibrary',MemberName="Less_FloatFloat")
+				   NodePosX=928
+				   NodePosY=-80
+				   NodeGuid=70B372B94A4EDAF80F8304A334B7C97D
+				   CustomProperties Pin (PinId=0BA5842D4451C002FE3ABBA2203CD18A,PinName="self",PinFriendlyName=NSLOCTEXT("K2Node", "Target", "Target"),PinToolTip="Target\nKismet Math Library Reference",PinType.PinCategory="object",PinType.PinSubCategory="",PinType.PinSubCategoryObject=Class'/Script/Engine.KismetMathLibrary',PinType.PinSubCategoryMemberReference=(),PinType.PinValueType=(),PinType.bIsMap=False,PinType.bIsSet=False,PinType.bIsArray=False,PinType.bIsReference=False,PinType.bIsConst=False,PinType.bIsWeakPointer=False,DefaultObject="/Script/Engine.Default__KismetMathLibrary",PersistentGuid=00000000000000000000000000000000,bHidden=True,bNotConnectable=False,bDefaultValueIsReadOnly=False,bDefaultValueIsIgnored=False,bAdvancedView=False,)
+				   CustomProperties Pin (PinId=88CF46414C43AA30CD71F39D9FDDCE8C,PinName="A",PinToolTip="A\nFloat",PinType.PinCategory="float",PinType.PinSubCategory="",PinType.PinSubCategoryObject=None,PinType.PinSubCategoryMemberReference=(),PinType.PinValueType=(),PinType.bIsMap=False,PinType.bIsSet=False,PinType.bIsArray=False,PinType.bIsReference=False,PinType.bIsConst=False,PinType.bIsWeakPointer=False,DefaultValue="0.0",AutogeneratedDefaultValue="0.0",LinkedTo=(K2Node_CallFunction_156 E1512AD84E1F9D6207A10A8B26039124,),PersistentGuid=00000000000000000000000000000000,bHidden=False,bNotConnectable=False,bDefaultValueIsReadOnly=False,bDefaultValueIsIgnored=False,bAdvancedView=False,)
+				   CustomProperties Pin (PinId=A33C24584B36D3CADE1083ACB60585EA,PinName="B",PinToolTip="B\nFloat",PinType.PinCategory="float",PinType.PinSubCategory="",PinType.PinSubCategoryObject=None,PinType.PinSubCategoryMemberReference=(),PinType.PinValueType=(),PinType.bIsMap=False,PinType.bIsSet=False,PinType.bIsArray=False,PinType.bIsReference=False,PinType.bIsConst=False,PinType.bIsWeakPointer=False,DefaultValue="0",AutogeneratedDefaultValue="0.0",PersistentGuid=00000000000000000000000000000000,bHidden=False,bNotConnectable=False,bDefaultValueIsReadOnly=False,bDefaultValueIsIgnored=False,bAdvancedView=False,)
+				   CustomProperties Pin (PinId=97BE87AA448C93D13187DE941B2FF014,PinName="ReturnValue",PinToolTip="Return Value\nBoolean",Direction="EGPD_Output",PinType.PinCategory="bool",PinType.PinSubCategory="",PinType.PinSubCategoryObject=None,PinType.PinSubCategoryMemberReference=(),PinType.PinValueType=(),PinType.bIsMap=False,PinType.bIsSet=False,PinType.bIsArray=False,PinType.bIsReference=False,PinType.bIsConst=False,PinType.bIsWeakPointer=False,AutogeneratedDefaultValue="false",LinkedTo=(K2Node_IfThenElse_14 EC64745F41E78DF32486608D583AE2E6,),PersistentGuid=00000000000000000000000000000000,bHidden=False,bNotConnectable=False,bDefaultValueIsReadOnly=False,bDefaultValueIsIgnored=False,bAdvancedView=False,)
+				End Object
+				Begin Object Class=/Script/BlueprintGraph.K2Node_IfThenElse Name="K2Node_IfThenElse_0"
+				   NodePosX=2688
+				   NodePosY=96
+				   NodeGuid=834109C748FD858A6CBA77923FEC1F68
+				   CustomProperties Pin (PinId=9DF08FE24D43CFADCBA446BB6EA29189,PinName="execute",PinType.PinCategory="exec",PinType.PinSubCategory="",PinType.PinSubCategoryObject=None,PinType.PinSubCategoryMemberReference=(),PinType.PinValueType=(),PinType.bIsMap=False,PinType.bIsSet=False,PinType.bIsArray=False,PinType.bIsReference=False,PinType.bIsConst=False,PinType.bIsWeakPointer=False,LinkedTo=(K2Node_IfThenElse_10 9B8358F143000BB08762918D0E75B5E9,),PersistentGuid=00000000000000000000000000000000,bHidden=False,bNotConnectable=False,bDefaultValueIsReadOnly=False,bDefaultValueIsIgnored=False,bAdvancedView=False,)
+				   CustomProperties Pin (PinId=EC64745F41E78DF32486608D583AE2E6,PinName="Condition",PinType.PinCategory="bool",PinType.PinSubCategory="",PinType.PinSubCategoryObject=None,PinType.PinSubCategoryMemberReference=(),PinType.PinValueType=(),PinType.bIsMap=False,PinType.bIsSet=False,PinType.bIsArray=False,PinType.bIsReference=False,PinType.bIsConst=False,PinType.bIsWeakPointer=False,DefaultValue="false",AutogeneratedDefaultValue="true",LinkedTo=(K2Node_CallFunction_160 04A5FFA94B808583861470BA77C74C63,),PersistentGuid=00000000000000000000000000000000,bHidden=False,bNotConnectable=False,bDefaultValueIsReadOnly=False,bDefaultValueIsIgnored=False,bAdvancedView=False,)
+				   CustomProperties Pin (PinId=3306E19940BBEA4DB04353976F98E9D8,PinName="then",PinFriendlyName=NSLOCTEXT("K2Node", "true", "true"),Direction="EGPD_Output",PinType.PinCategory="exec",PinType.PinSubCategory="",PinType.PinSubCategoryObject=None,PinType.PinSubCategoryMemberReference=(),PinType.PinValueType=(),PinType.bIsMap=False,PinType.bIsSet=False,PinType.bIsArray=False,PinType.bIsReference=False,PinType.bIsConst=False,PinType.bIsWeakPointer=False,LinkedTo=(K2Node_CallFunction_12918 96D7193F444562BF1B9D1A82C15146E1,),PersistentGuid=00000000000000000000000000000000,bHidden=False,bNotConnectable=False,bDefaultValueIsReadOnly=False,bDefaultValueIsIgnored=False,bAdvancedView=False,)
+				   CustomProperties Pin (PinId=9B8358F143000BB08762918D0E75B5E9,PinName="else",PinFriendlyName=NSLOCTEXT("K2Node", "false", "false"),Direction="EGPD_Output",PinType.PinCategory="exec",PinType.PinSubCategory="",PinType.PinSubCategoryObject=None,PinType.PinSubCategoryMemberReference=(),PinType.PinValueType=(),PinType.bIsMap=False,PinType.bIsSet=False,PinType.bIsArray=False,PinType.bIsReference=False,PinType.bIsConst=False,PinType.bIsWeakPointer=False,LinkedTo=(K2Node_CallFunction_12105 96D7193F444562BF1B9D1A82C15146E1,),PersistentGuid=00000000000000000000000000000000,bHidden=False,bNotConnectable=False,bDefaultValueIsReadOnly=False,bDefaultValueIsIgnored=False,bAdvancedView=False,)
+				End Object
+				Begin Object Class=/Script/UnrealEd.EdGraphNode_Comment Name="EdGraphNode_Comment_34"
+				   NodePosX=2656
+				   NodePosY=16
+				   NodeWidth=240
+				   NodeHeight=192
+				   NodeComment="True Swipe Up False Swipe Down"
+				   NodeGuid=D6B455AE4A846518930BC28A399C7EA4
+				End Object
+				Begin Object Class=/Script/UnrealEd.EdGraphNode_Comment Name="EdGraphNode_Comment_406"
+				   NodePosX=2032
+				   NodePosY=-912
+				   NodeWidth=496
+				   NodeHeight=288
+				   NodeComment="Single Tap"
+				   NodeGuid=DF2EFBF346C7405A08FE709811DC9E10
+				End Object
+				Begin Object Class=/Script/UnrealEd.EdGraphNode_Comment Name="EdGraphNode_Comment_428"
+				   NodePosX=3232
+				   NodePosY=-512
+				   NodeWidth=496
+				   NodeHeight=288
+				   NodeComment="Swipe Left "
+				   NodeGuid=493AECF4417A2E7FB0FCE88812690CD6
+				End Object
+				Begin Object Class=/Script/UnrealEd.EdGraphNode_Comment Name="EdGraphNode_Comment_467"
+				   NodePosX=3232
+				   NodePosY=-880
+				   NodeWidth=496
+				   NodeHeight=272
+				   NodeComment="Swipe Right"
+				   NodeGuid=5B816F28421C2236D6DE9AA47A614F62
+				End Object
+				Begin Object Class=/Script/UnrealEd.EdGraphNode_Comment Name="EdGraphNode_Comment_469"
+				   NodePosX=3200
+				   NodePosY=16
+				   NodeWidth=512
+				   NodeHeight=304
+				   NodeComment="Swipe Up"
+				   NodeGuid=B801FF7E44B7EC696FE6D08BF2FB60C2
+				End Object
+				Begin Object Class=/Script/UnrealEd.EdGraphNode_Comment Name="EdGraphNode_Comment_468"
+				   NodePosX=3200
+				   NodePosY=352
+				   NodeWidth=512
+				   NodeHeight=256
+				   NodeComment="Swipe Down"
+				   NodeGuid=6B73A36A4EBE7FCA096C4F8D425952E0
+				End Object
+				Begin Object Class=/Script/UnrealEd.EdGraphNode_Comment Name="EdGraphNode_Comment_3715"
+				   NodePosX=420
+				   NodePosY=-824
+				   NodeWidth=224
+				   NodeHeight=272
+				   NodeComment="Touch Input"
+				   NodeGuid=F091D09E4E63BC8BBE2EE39588C7AA9E
+				End Object
+				Begin Object Class=/Script/BlueprintGraph.K2Node_VariableGet Name="K2Node_VariableGet_7345"
+				   VariableReference=(MemberName="TextRender",bSelfContext=True)
+				   NodePosX=2064
+				   NodePosY=-781
+				   NodeGuid=4A64FE524F28F099E067E092F02FF82A
+				   CustomProperties Pin (PinId=D370BA9F4E23C50CF00DBF9C3866BEDE,PinName="TextRender",Direction="EGPD_Output",PinType.PinCategory="object",PinType.PinSubCategory="",PinType.PinSubCategoryObject=Class'/Script/Engine.TextRenderComponent',PinType.PinSubCategoryMemberReference=(),PinType.PinValueType=(),PinType.bIsMap=False,PinType.bIsSet=False,PinType.bIsArray=False,PinType.bIsReference=False,PinType.bIsConst=False,PinType.bIsWeakPointer=False,LinkedTo=(K2Node_CallFunction_13198 24D9A43A4CC5BB2E7963438C3FE88719,),PersistentGuid=00000000000000000000000000000000,bHidden=False,bNotConnectable=False,bDefaultValueIsReadOnly=False,bDefaultValueIsIgnored=False,bAdvancedView=False,)
+				   CustomProperties Pin (PinId=83CDB7A849B7EB215EC3139D1E9C6B39,PinName="self",PinFriendlyName=NSLOCTEXT("K2Node", "Target", "Target"),PinType.PinCategory="object",PinType.PinSubCategory="",PinType.PinSubCategoryObject=BlueprintGeneratedClass'/Game/GearVR_Pawn/GearVR_Pawn_Input.GearVR_Pawn_Input_C',PinType.PinSubCategoryMemberReference=(),PinType.PinValueType=(),PinType.bIsMap=False,PinType.bIsSet=False,PinType.bIsArray=False,PinType.bIsReference=False,PinType.bIsConst=False,PinType.bIsWeakPointer=False,PersistentGuid=00000000000000000000000000000000,bHidden=True,bNotConnectable=False,bDefaultValueIsReadOnly=False,bDefaultValueIsIgnored=False,bAdvancedView=False,)
+				End Object
+				Begin Object Class=/Script/BlueprintGraph.K2Node_CallFunction Name="K2Node_CallFunction_13198"
+				   FunctionReference=(MemberParent=Class'/Script/Engine.TextRenderComponent',MemberName="K2_SetText")
+				   NodePosX=2256
+				   NodePosY=-848
+				   NodeGuid=E3D7C9674D95DFA37E4D6799D5FDC7CB
+				   CustomProperties Pin (PinId=96D7193F444562BF1B9D1A82C15146E1,PinName="execute",PinType.PinCategory="exec",PinType.PinSubCategory="",PinType.PinSubCategoryObject=None,PinType.PinSubCategoryMemberReference=(),PinType.PinValueType=(),PinType.bIsMap=False,PinType.bIsSet=False,PinType.bIsArray=False,PinType.bIsReference=False,PinType.bIsConst=False,PinType.bIsWeakPointer=False,LinkedTo=(K2Node_IfThenElse_45 9B8358F143000BB08762918D0E75B5E9,),PersistentGuid=00000000000000000000000000000000,bHidden=False,bNotConnectable=False,bDefaultValueIsReadOnly=False,bDefaultValueIsIgnored=False,bAdvancedView=False,)
+				   CustomProperties Pin (PinId=66DCFA3744967675B795929682EA35F0,PinName="then",Direction="EGPD_Output",PinType.PinCategory="exec",PinType.PinSubCategory="",PinType.PinSubCategoryObject=None,PinType.PinSubCategoryMemberReference=(),PinType.PinValueType=(),PinType.bIsMap=False,PinType.bIsSet=False,PinType.bIsArray=False,PinType.bIsReference=False,PinType.bIsConst=False,PinType.bIsWeakPointer=False,PersistentGuid=00000000000000000000000000000000,bHidden=False,bNotConnectable=False,bDefaultValueIsReadOnly=False,bDefaultValueIsIgnored=False,bAdvancedView=False,)
+				   CustomProperties Pin (PinId=24D9A43A4CC5BB2E7963438C3FE88719,PinName="self",PinFriendlyName=NSLOCTEXT("K2Node", "Target", "Target"),PinType.PinCategory="object",PinType.PinSubCategory="",PinType.PinSubCategoryObject=Class'/Script/Engine.TextRenderComponent',PinType.PinSubCategoryMemberReference=(),PinType.PinValueType=(),PinType.bIsMap=False,PinType.bIsSet=False,PinType.bIsArray=False,PinType.bIsReference=False,PinType.bIsConst=False,PinType.bIsWeakPointer=False,LinkedTo=(K2Node_VariableGet_7345 D370BA9F4E23C50CF00DBF9C3866BEDE,),PersistentGuid=00000000000000000000000000000000,bHidden=False,bNotConnectable=False,bDefaultValueIsReadOnly=False,bDefaultValueIsIgnored=False,bAdvancedView=False,)
+				   CustomProperties Pin (PinId=FD05F54B48F429BC42F882B6033BFA3B,PinName="Value",PinType.PinCategory="text",PinType.PinSubCategory="",PinType.PinSubCategoryObject=None,PinType.PinSubCategoryMemberReference=(),PinType.PinValueType=(),PinType.bIsMap=False,PinType.bIsSet=False,PinType.bIsArray=False,PinType.bIsReference=True,PinType.bIsConst=True,PinType.bIsWeakPointer=False,LinkedTo=(K2Node_VariableGet_1989 6B6233E14C21954E6F46468CF70142A4,),PersistentGuid=00000000000000000000000000000000,bHidden=False,bNotConnectable=False,bDefaultValueIsReadOnly=False,bDefaultValueIsIgnored=True,bAdvancedView=False,)
+				End Object
+				Begin Object Class=/Script/BlueprintGraph.K2Node_VariableGet Name="K2Node_VariableGet_1989"
+				   VariableReference=(MemberName="TouchTap",MemberGuid=2278E57B422E884A9747D48D238B9476,bSelfContext=True)
+				   NodePosX=2080
+				   NodePosY=-720
+				   NodeGuid=FE2C2E4241A5930982ABC7887ADA757B
+				   CustomProperties Pin (PinId=6B6233E14C21954E6F46468CF70142A4,PinName="TouchTap",Direction="EGPD_Output",PinType.PinCategory="text",PinType.PinSubCategory="",PinType.PinSubCategoryObject=None,PinType.PinSubCategoryMemberReference=(),PinType.PinValueType=(),PinType.bIsMap=False,PinType.bIsSet=False,PinType.bIsArray=False,PinType.bIsReference=False,PinType.bIsConst=False,PinType.bIsWeakPointer=False,LinkedTo=(K2Node_CallFunction_13198 FD05F54B48F429BC42F882B6033BFA3B,),PersistentGuid=00000000000000000000000000000000,bHidden=False,bNotConnectable=False,bDefaultValueIsReadOnly=False,bDefaultValueIsIgnored=False,bAdvancedView=False,)
+				   CustomProperties Pin (PinId=07B23AEA4D74958F098B189DD57B3744,PinName="self",PinFriendlyName=NSLOCTEXT("K2Node", "Target", "Target"),PinType.PinCategory="object",PinType.PinSubCategory="",PinType.PinSubCategoryObject=BlueprintGeneratedClass'/Game/GearVR_Pawn/GearVR_Pawn_Input.GearVR_Pawn_Input_C',PinType.PinSubCategoryMemberReference=(),PinType.PinValueType=(),PinType.bIsMap=False,PinType.bIsSet=False,PinType.bIsArray=False,PinType.bIsReference=False,PinType.bIsConst=False,PinType.bIsWeakPointer=False,PersistentGuid=00000000000000000000000000000000,bHidden=True,bNotConnectable=False,bDefaultValueIsReadOnly=False,bDefaultValueIsIgnored=False,bAdvancedView=False,)
+				End Object
+				Begin Object Class=/Script/BlueprintGraph.K2Node_VariableGet Name="K2Node_VariableGet_7462"
+				   VariableReference=(MemberName="TextRender",bSelfContext=True)
+				   NodePosX=3248
+				   NodePosY=-352
+				   NodeGuid=CB3D456D40BBC9650F830498D7DCA887
+				   CustomProperties Pin (PinId=D370BA9F4E23C50CF00DBF9C3866BEDE,PinName="TextRender",Direction="EGPD_Output",PinType.PinCategory="object",PinType.PinSubCategory="",PinType.PinSubCategoryObject=Class'/Script/Engine.TextRenderComponent',PinType.PinSubCategoryMemberReference=(),PinType.PinValueType=(),PinType.bIsMap=False,PinType.bIsSet=False,PinType.bIsArray=False,PinType.bIsReference=False,PinType.bIsConst=False,PinType.bIsWeakPointer=False,LinkedTo=(K2Node_CallFunction_12871 24D9A43A4CC5BB2E7963438C3FE88719,),PersistentGuid=00000000000000000000000000000000,bHidden=False,bNotConnectable=False,bDefaultValueIsReadOnly=False,bDefaultValueIsIgnored=False,bAdvancedView=False,)
+				   CustomProperties Pin (PinId=83CDB7A849B7EB215EC3139D1E9C6B39,PinName="self",PinFriendlyName=NSLOCTEXT("K2Node", "Target", "Target"),PinType.PinCategory="object",PinType.PinSubCategory="",PinType.PinSubCategoryObject=BlueprintGeneratedClass'/Game/GearVR_Pawn/GearVR_Pawn_Input.GearVR_Pawn_Input_C',PinType.PinSubCategoryMemberReference=(),PinType.PinValueType=(),PinType.bIsMap=False,PinType.bIsSet=False,PinType.bIsArray=False,PinType.bIsReference=False,PinType.bIsConst=False,PinType.bIsWeakPointer=False,PersistentGuid=00000000000000000000000000000000,bHidden=True,bNotConnectable=False,bDefaultValueIsReadOnly=False,bDefaultValueIsIgnored=False,bAdvancedView=False,)
+				End Object
+				Begin Object Class=/Script/BlueprintGraph.K2Node_CallFunction Name="K2Node_CallFunction_12871"
+				   FunctionReference=(MemberParent=Class'/Script/Engine.TextRenderComponent',MemberName="K2_SetText")
+				   NodePosX=3440
+				   NodePosY=-432
+				   NodeGuid=348ED6794353CEA9E97CE7A1C8BF73DD
+				   CustomProperties Pin (PinId=96D7193F444562BF1B9D1A82C15146E1,PinName="execute",PinToolTip="\nExec",PinType.PinCategory="exec",PinType.PinSubCategory="",PinType.PinSubCategoryObject=None,PinType.PinSubCategoryMemberReference=(),PinType.PinValueType=(),PinType.bIsMap=False,PinType.bIsSet=False,PinType.bIsArray=False,PinType.bIsReference=False,PinType.bIsConst=False,PinType.bIsWeakPointer=False,LinkedTo=(K2Node_IfThenElse_14 9B8358F143000BB08762918D0E75B5E9,),PersistentGuid=00000000000000000000000000000000,bHidden=False,bNotConnectable=False,bDefaultValueIsReadOnly=False,bDefaultValueIsIgnored=False,bAdvancedView=False,)
+				   CustomProperties Pin (PinId=66DCFA3744967675B795929682EA35F0,PinName="then",PinToolTip="\nExec",Direction="EGPD_Output",PinType.PinCategory="exec",PinType.PinSubCategory="",PinType.PinSubCategoryObject=None,PinType.PinSubCategoryMemberReference=(),PinType.PinValueType=(),PinType.bIsMap=False,PinType.bIsSet=False,PinType.bIsArray=False,PinType.bIsReference=False,PinType.bIsConst=False,PinType.bIsWeakPointer=False,PersistentGuid=00000000000000000000000000000000,bHidden=False,bNotConnectable=False,bDefaultValueIsReadOnly=False,bDefaultValueIsIgnored=False,bAdvancedView=False,)
+				   CustomProperties Pin (PinId=24D9A43A4CC5BB2E7963438C3FE88719,PinName="self",PinFriendlyName=NSLOCTEXT("K2Node", "Target", "Target"),PinToolTip="Target\nText Render Component Reference",PinType.PinCategory="object",PinType.PinSubCategory="",PinType.PinSubCategoryObject=Class'/Script/Engine.TextRenderComponent',PinType.PinSubCategoryMemberReference=(),PinType.PinValueType=(),PinType.bIsMap=False,PinType.bIsSet=False,PinType.bIsArray=False,PinType.bIsReference=False,PinType.bIsConst=False,PinType.bIsWeakPointer=False,LinkedTo=(K2Node_VariableGet_7462 D370BA9F4E23C50CF00DBF9C3866BEDE,),PersistentGuid=00000000000000000000000000000000,bHidden=False,bNotConnectable=False,bDefaultValueIsReadOnly=False,bDefaultValueIsIgnored=False,bAdvancedView=False,)
+				   CustomProperties Pin (PinId=FD05F54B48F429BC42F882B6033BFA3B,PinName="Value",PinToolTip="Value\nText (by ref)",PinType.PinCategory="text",PinType.PinSubCategory="",PinType.PinSubCategoryObject=None,PinType.PinSubCategoryMemberReference=(),PinType.PinValueType=(),PinType.bIsMap=False,PinType.bIsSet=False,PinType.bIsArray=False,PinType.bIsReference=True,PinType.bIsConst=True,PinType.bIsWeakPointer=False,LinkedTo=(K2Node_VariableGet_1990 479E532648FF40BB2454779E2233FB2C,),PersistentGuid=00000000000000000000000000000000,bHidden=False,bNotConnectable=False,bDefaultValueIsReadOnly=False,bDefaultValueIsIgnored=True,bAdvancedView=False,)
+				End Object
+				Begin Object Class=/Script/BlueprintGraph.K2Node_VariableGet Name="K2Node_VariableGet_1990"
+				   VariableReference=(MemberName="SwipeLeft",MemberGuid=C10AE20242A256AA52304B97B3EC333B,bSelfContext=True)
+				   NodePosX=3248
+				   NodePosY=-304
+				   NodeGuid=23E9DD1B48DD4DB5CAFB8AB90BA82162
+				   CustomProperties Pin (PinId=479E532648FF40BB2454779E2233FB2C,PinName="SwipeLeft",Direction="EGPD_Output",PinType.PinCategory="text",PinType.PinSubCategory="",PinType.PinSubCategoryObject=None,PinType.PinSubCategoryMemberReference=(),PinType.PinValueType=(),PinType.bIsMap=False,PinType.bIsSet=False,PinType.bIsArray=False,PinType.bIsReference=False,PinType.bIsConst=False,PinType.bIsWeakPointer=False,LinkedTo=(K2Node_CallFunction_12871 FD05F54B48F429BC42F882B6033BFA3B,),PersistentGuid=00000000000000000000000000000000,bHidden=False,bNotConnectable=False,bDefaultValueIsReadOnly=False,bDefaultValueIsIgnored=False,bAdvancedView=False,)
+				   CustomProperties Pin (PinId=396838D94892A563D7BA6091783CC92B,PinName="self",PinFriendlyName=NSLOCTEXT("K2Node", "Target", "Target"),PinType.PinCategory="object",PinType.PinSubCategory="",PinType.PinSubCategoryObject=BlueprintGeneratedClass'/Game/GearVR_Pawn/GearVR_Pawn_Input.GearVR_Pawn_Input_C',PinType.PinSubCategoryMemberReference=(),PinType.PinValueType=(),PinType.bIsMap=False,PinType.bIsSet=False,PinType.bIsArray=False,PinType.bIsReference=False,PinType.bIsConst=False,PinType.bIsWeakPointer=False,PersistentGuid=00000000000000000000000000000000,bHidden=True,bNotConnectable=False,bDefaultValueIsReadOnly=False,bDefaultValueIsIgnored=False,bAdvancedView=False,)
+				End Object
+				Begin Object Class=/Script/BlueprintGraph.K2Node_CallFunction Name="K2Node_CallFunction_12917"
+				   FunctionReference=(MemberParent=Class'/Script/Engine.TextRenderComponent',MemberName="K2_SetText")
+				   NodePosX=3424
+				   NodePosY=-784
+				   NodeGuid=F274232E4254C0581BC26F835177C8B5
+				   CustomProperties Pin (PinId=96D7193F444562BF1B9D1A82C15146E1,PinName="execute",PinToolTip="\nExec",PinType.PinCategory="exec",PinType.PinSubCategory="",PinType.PinSubCategoryObject=None,PinType.PinSubCategoryMemberReference=(),PinType.PinValueType=(),PinType.bIsMap=False,PinType.bIsSet=False,PinType.bIsArray=False,PinType.bIsReference=False,PinType.bIsConst=False,PinType.bIsWeakPointer=False,LinkedTo=(K2Node_IfThenElse_14 3306E19940BBEA4DB04353976F98E9D8,),PersistentGuid=00000000000000000000000000000000,bHidden=False,bNotConnectable=False,bDefaultValueIsReadOnly=False,bDefaultValueIsIgnored=False,bAdvancedView=False,)
+				   CustomProperties Pin (PinId=66DCFA3744967675B795929682EA35F0,PinName="then",PinToolTip="\nExec",Direction="EGPD_Output",PinType.PinCategory="exec",PinType.PinSubCategory="",PinType.PinSubCategoryObject=None,PinType.PinSubCategoryMemberReference=(),PinType.PinValueType=(),PinType.bIsMap=False,PinType.bIsSet=False,PinType.bIsArray=False,PinType.bIsReference=False,PinType.bIsConst=False,PinType.bIsWeakPointer=False,PersistentGuid=00000000000000000000000000000000,bHidden=False,bNotConnectable=False,bDefaultValueIsReadOnly=False,bDefaultValueIsIgnored=False,bAdvancedView=False,)
+				   CustomProperties Pin (PinId=24D9A43A4CC5BB2E7963438C3FE88719,PinName="self",PinFriendlyName=NSLOCTEXT("K2Node", "Target", "Target"),PinToolTip="Target\nText Render Component Reference",PinType.PinCategory="object",PinType.PinSubCategory="",PinType.PinSubCategoryObject=Class'/Script/Engine.TextRenderComponent',PinType.PinSubCategoryMemberReference=(),PinType.PinValueType=(),PinType.bIsMap=False,PinType.bIsSet=False,PinType.bIsArray=False,PinType.bIsReference=False,PinType.bIsConst=False,PinType.bIsWeakPointer=False,LinkedTo=(K2Node_VariableGet_7480 D370BA9F4E23C50CF00DBF9C3866BEDE,),PersistentGuid=00000000000000000000000000000000,bHidden=False,bNotConnectable=False,bDefaultValueIsReadOnly=False,bDefaultValueIsIgnored=False,bAdvancedView=False,)
+				   CustomProperties Pin (PinId=FD05F54B48F429BC42F882B6033BFA3B,PinName="Value",PinToolTip="Value\nText (by ref)",PinType.PinCategory="text",PinType.PinSubCategory="",PinType.PinSubCategoryObject=None,PinType.PinSubCategoryMemberReference=(),PinType.PinValueType=(),PinType.bIsMap=False,PinType.bIsSet=False,PinType.bIsArray=False,PinType.bIsReference=True,PinType.bIsConst=True,PinType.bIsWeakPointer=False,LinkedTo=(K2Node_VariableGet_1991 893F8E6949C970C06EA6F094A4A734DF,),PersistentGuid=00000000000000000000000000000000,bHidden=False,bNotConnectable=False,bDefaultValueIsReadOnly=False,bDefaultValueIsIgnored=True,bAdvancedView=False,)
+				End Object
+				Begin Object Class=/Script/BlueprintGraph.K2Node_VariableGet Name="K2Node_VariableGet_7480"
+				   VariableReference=(MemberName="TextRender",bSelfContext=True)
+				   NodePosX=3264
+				   NodePosY=-720
+				   NodeGuid=3D64EB0C42F4B765D0FD638F629377F1
+				   CustomProperties Pin (PinId=D370BA9F4E23C50CF00DBF9C3866BEDE,PinName="TextRender",Direction="EGPD_Output",PinType.PinCategory="object",PinType.PinSubCategory="",PinType.PinSubCategoryObject=Class'/Script/Engine.TextRenderComponent',PinType.PinSubCategoryMemberReference=(),PinType.PinValueType=(),PinType.bIsMap=False,PinType.bIsSet=False,PinType.bIsArray=False,PinType.bIsReference=False,PinType.bIsConst=False,PinType.bIsWeakPointer=False,LinkedTo=(K2Node_CallFunction_12917 24D9A43A4CC5BB2E7963438C3FE88719,),PersistentGuid=00000000000000000000000000000000,bHidden=False,bNotConnectable=False,bDefaultValueIsReadOnly=False,bDefaultValueIsIgnored=False,bAdvancedView=False,)
+				   CustomProperties Pin (PinId=83CDB7A849B7EB215EC3139D1E9C6B39,PinName="self",PinFriendlyName=NSLOCTEXT("K2Node", "Target", "Target"),PinType.PinCategory="object",PinType.PinSubCategory="",PinType.PinSubCategoryObject=BlueprintGeneratedClass'/Game/GearVR_Pawn/GearVR_Pawn_Input.GearVR_Pawn_Input_C',PinType.PinSubCategoryMemberReference=(),PinType.PinValueType=(),PinType.bIsMap=False,PinType.bIsSet=False,PinType.bIsArray=False,PinType.bIsReference=False,PinType.bIsConst=False,PinType.bIsWeakPointer=False,PersistentGuid=00000000000000000000000000000000,bHidden=True,bNotConnectable=False,bDefaultValueIsReadOnly=False,bDefaultValueIsIgnored=False,bAdvancedView=False,)
+				End Object
+				Begin Object Class=/Script/BlueprintGraph.K2Node_VariableGet Name="K2Node_VariableGet_1991"
+				   VariableReference=(MemberName="SwipeRight",MemberGuid=F5AB025B4181866BFDCBB2B6346C6E96,bSelfContext=True)
+				   NodePosX=3296
+				   NodePosY=-688
+				   NodeGuid=68F52FEC4114C265858BD2A8F6A51975
+				   CustomProperties Pin (PinId=893F8E6949C970C06EA6F094A4A734DF,PinName="SwipeRight",Direction="EGPD_Output",PinType.PinCategory="text",PinType.PinSubCategory="",PinType.PinSubCategoryObject=None,PinType.PinSubCategoryMemberReference=(),PinType.PinValueType=(),PinType.bIsMap=False,PinType.bIsSet=False,PinType.bIsArray=False,PinType.bIsReference=False,PinType.bIsConst=False,PinType.bIsWeakPointer=False,LinkedTo=(K2Node_CallFunction_12917 FD05F54B48F429BC42F882B6033BFA3B,),PersistentGuid=00000000000000000000000000000000,bHidden=False,bNotConnectable=False,bDefaultValueIsReadOnly=False,bDefaultValueIsIgnored=False,bAdvancedView=False,)
+				   CustomProperties Pin (PinId=CCC2B01345BBBBD19B1B1EA11402D574,PinName="self",PinFriendlyName=NSLOCTEXT("K2Node", "Target", "Target"),PinType.PinCategory="object",PinType.PinSubCategory="",PinType.PinSubCategoryObject=BlueprintGeneratedClass'/Game/GearVR_Pawn/GearVR_Pawn_Input.GearVR_Pawn_Input_C',PinType.PinSubCategoryMemberReference=(),PinType.PinValueType=(),PinType.bIsMap=False,PinType.bIsSet=False,PinType.bIsArray=False,PinType.bIsReference=False,PinType.bIsConst=False,PinType.bIsWeakPointer=False,PersistentGuid=00000000000000000000000000000000,bHidden=True,bNotConnectable=False,bDefaultValueIsReadOnly=False,bDefaultValueIsIgnored=False,bAdvancedView=False,)
+				End Object
+				Begin Object Class=/Script/BlueprintGraph.K2Node_VariableGet Name="K2Node_VariableGet_7481"
+				   VariableReference=(MemberName="TextRender",bSelfContext=True)
+				   NodePosX=3232
+				   NodePosY=208
+				   NodeGuid=E681FA864399BFC7D25DFFBBA7B3704D
+				   CustomProperties Pin (PinId=D370BA9F4E23C50CF00DBF9C3866BEDE,PinName="TextRender",Direction="EGPD_Output",PinType.PinCategory="object",PinType.PinSubCategory="",PinType.PinSubCategoryObject=Class'/Script/Engine.TextRenderComponent',PinType.PinSubCategoryMemberReference=(),PinType.PinValueType=(),PinType.bIsMap=False,PinType.bIsSet=False,PinType.bIsArray=False,PinType.bIsReference=False,PinType.bIsConst=False,PinType.bIsWeakPointer=False,LinkedTo=(K2Node_CallFunction_12918 24D9A43A4CC5BB2E7963438C3FE88719,),PersistentGuid=00000000000000000000000000000000,bHidden=False,bNotConnectable=False,bDefaultValueIsReadOnly=False,bDefaultValueIsIgnored=False,bAdvancedView=False,)
+				   CustomProperties Pin (PinId=83CDB7A849B7EB215EC3139D1E9C6B39,PinName="self",PinFriendlyName=NSLOCTEXT("K2Node", "Target", "Target"),PinType.PinCategory="object",PinType.PinSubCategory="",PinType.PinSubCategoryObject=BlueprintGeneratedClass'/Game/GearVR_Pawn/GearVR_Pawn_Input.GearVR_Pawn_Input_C',PinType.PinSubCategoryMemberReference=(),PinType.PinValueType=(),PinType.bIsMap=False,PinType.bIsSet=False,PinType.bIsArray=False,PinType.bIsReference=False,PinType.bIsConst=False,PinType.bIsWeakPointer=False,PersistentGuid=00000000000000000000000000000000,bHidden=True,bNotConnectable=False,bDefaultValueIsReadOnly=False,bDefaultValueIsIgnored=False,bAdvancedView=False,)
+				End Object
+				Begin Object Class=/Script/BlueprintGraph.K2Node_CallFunction Name="K2Node_CallFunction_12918"
+				   FunctionReference=(MemberParent=Class'/Script/Engine.TextRenderComponent',MemberName="K2_SetText")
+				   NodePosX=3408
+				   NodePosY=96
+				   NodeGuid=EF048E2F435946AE7E2C6394C2E4FBC9
+				   CustomProperties Pin (PinId=96D7193F444562BF1B9D1A82C15146E1,PinName="execute",PinToolTip="\nExec",PinType.PinCategory="exec",PinType.PinSubCategory="",PinType.PinSubCategoryObject=None,PinType.PinSubCategoryMemberReference=(),PinType.PinValueType=(),PinType.bIsMap=False,PinType.bIsSet=False,PinType.bIsArray=False,PinType.bIsReference=False,PinType.bIsConst=False,PinType.bIsWeakPointer=False,LinkedTo=(K2Node_IfThenElse_0 3306E19940BBEA4DB04353976F98E9D8,),PersistentGuid=00000000000000000000000000000000,bHidden=False,bNotConnectable=False,bDefaultValueIsReadOnly=False,bDefaultValueIsIgnored=False,bAdvancedView=False,)
+				   CustomProperties Pin (PinId=66DCFA3744967675B795929682EA35F0,PinName="then",PinToolTip="\nExec",Direction="EGPD_Output",PinType.PinCategory="exec",PinType.PinSubCategory="",PinType.PinSubCategoryObject=None,PinType.PinSubCategoryMemberReference=(),PinType.PinValueType=(),PinType.bIsMap=False,PinType.bIsSet=False,PinType.bIsArray=False,PinType.bIsReference=False,PinType.bIsConst=False,PinType.bIsWeakPointer=False,PersistentGuid=00000000000000000000000000000000,bHidden=False,bNotConnectable=False,bDefaultValueIsReadOnly=False,bDefaultValueIsIgnored=False,bAdvancedView=False,)
+				   CustomProperties Pin (PinId=24D9A43A4CC5BB2E7963438C3FE88719,PinName="self",PinFriendlyName=NSLOCTEXT("K2Node", "Target", "Target"),PinToolTip="Target\nText Render Component Reference",PinType.PinCategory="object",PinType.PinSubCategory="",PinType.PinSubCategoryObject=Class'/Script/Engine.TextRenderComponent',PinType.PinSubCategoryMemberReference=(),PinType.PinValueType=(),PinType.bIsMap=False,PinType.bIsSet=False,PinType.bIsArray=False,PinType.bIsReference=False,PinType.bIsConst=False,PinType.bIsWeakPointer=False,LinkedTo=(K2Node_VariableGet_7481 D370BA9F4E23C50CF00DBF9C3866BEDE,),PersistentGuid=00000000000000000000000000000000,bHidden=False,bNotConnectable=False,bDefaultValueIsReadOnly=False,bDefaultValueIsIgnored=False,bAdvancedView=False,)
+				   CustomProperties Pin (PinId=FD05F54B48F429BC42F882B6033BFA3B,PinName="Value",PinToolTip="Value\nText (by ref)",PinType.PinCategory="text",PinType.PinSubCategory="",PinType.PinSubCategoryObject=None,PinType.PinSubCategoryMemberReference=(),PinType.PinValueType=(),PinType.bIsMap=False,PinType.bIsSet=False,PinType.bIsArray=False,PinType.bIsReference=True,PinType.bIsConst=True,PinType.bIsWeakPointer=False,LinkedTo=(K2Node_VariableGet_1992 BEA51E4248B8DCF5FF42B5928B906A7F,),PersistentGuid=00000000000000000000000000000000,bHidden=False,bNotConnectable=False,bDefaultValueIsReadOnly=False,bDefaultValueIsIgnored=True,bAdvancedView=False,)
+				End Object
+				Begin Object Class=/Script/BlueprintGraph.K2Node_VariableGet Name="K2Node_VariableGet_1992"
+				   VariableReference=(MemberName="SwipeUp",MemberGuid=1B858BD34CEC39BE30CA59AB447162D9,bSelfContext=True)
+				   NodePosX=3248
+				   NodePosY=256
+				   NodeGuid=9A610025422229D7D52BDC9517A1DF28
+				   CustomProperties Pin (PinId=BEA51E4248B8DCF5FF42B5928B906A7F,PinName="SwipeUp",Direction="EGPD_Output",PinType.PinCategory="text",PinType.PinSubCategory="",PinType.PinSubCategoryObject=None,PinType.PinSubCategoryMemberReference=(),PinType.PinValueType=(),PinType.bIsMap=False,PinType.bIsSet=False,PinType.bIsArray=False,PinType.bIsReference=False,PinType.bIsConst=False,PinType.bIsWeakPointer=False,LinkedTo=(K2Node_CallFunction_12918 FD05F54B48F429BC42F882B6033BFA3B,),PersistentGuid=00000000000000000000000000000000,bHidden=False,bNotConnectable=False,bDefaultValueIsReadOnly=False,bDefaultValueIsIgnored=False,bAdvancedView=False,)
+				   CustomProperties Pin (PinId=DDB208D94B151DF392CFBD8CFDD53C2D,PinName="self",PinFriendlyName=NSLOCTEXT("K2Node", "Target", "Target"),PinType.PinCategory="object",PinType.PinSubCategory="",PinType.PinSubCategoryObject=BlueprintGeneratedClass'/Game/GearVR_Pawn/GearVR_Pawn_Input.GearVR_Pawn_Input_C',PinType.PinSubCategoryMemberReference=(),PinType.PinValueType=(),PinType.bIsMap=False,PinType.bIsSet=False,PinType.bIsArray=False,PinType.bIsReference=False,PinType.bIsConst=False,PinType.bIsWeakPointer=False,PersistentGuid=00000000000000000000000000000000,bHidden=True,bNotConnectable=False,bDefaultValueIsReadOnly=False,bDefaultValueIsIgnored=False,bAdvancedView=False,)
+				End Object
+				Begin Object Class=/Script/BlueprintGraph.K2Node_CallFunction Name="K2Node_CallFunction_12105"
+				   FunctionReference=(MemberParent=Class'/Script/Engine.TextRenderComponent',MemberName="K2_SetText")
+				   NodePosX=3392
+				   NodePosY=400
+				   NodeGuid=BA688FE34A5A1B2A134929B9ABA6E5C6
+				   CustomProperties Pin (PinId=96D7193F444562BF1B9D1A82C15146E1,PinName="execute",PinType.PinCategory="exec",PinType.PinSubCategory="",PinType.PinSubCategoryObject=None,PinType.PinSubCategoryMemberReference=(),PinType.PinValueType=(),PinType.bIsMap=False,PinType.bIsSet=False,PinType.bIsArray=False,PinType.bIsReference=False,PinType.bIsConst=False,PinType.bIsWeakPointer=False,LinkedTo=(K2Node_IfThenElse_0 9B8358F143000BB08762918D0E75B5E9,),PersistentGuid=00000000000000000000000000000000,bHidden=False,bNotConnectable=False,bDefaultValueIsReadOnly=False,bDefaultValueIsIgnored=False,bAdvancedView=False,)
+				   CustomProperties Pin (PinId=66DCFA3744967675B795929682EA35F0,PinName="then",Direction="EGPD_Output",PinType.PinCategory="exec",PinType.PinSubCategory="",PinType.PinSubCategoryObject=None,PinType.PinSubCategoryMemberReference=(),PinType.PinValueType=(),PinType.bIsMap=False,PinType.bIsSet=False,PinType.bIsArray=False,PinType.bIsReference=False,PinType.bIsConst=False,PinType.bIsWeakPointer=False,PersistentGuid=00000000000000000000000000000000,bHidden=False,bNotConnectable=False,bDefaultValueIsReadOnly=False,bDefaultValueIsIgnored=False,bAdvancedView=False,)
+				   CustomProperties Pin (PinId=24D9A43A4CC5BB2E7963438C3FE88719,PinName="self",PinFriendlyName=NSLOCTEXT("K2Node", "Target", "Target"),PinType.PinCategory="object",PinType.PinSubCategory="",PinType.PinSubCategoryObject=Class'/Script/Engine.TextRenderComponent',PinType.PinSubCategoryMemberReference=(),PinType.PinValueType=(),PinType.bIsMap=False,PinType.bIsSet=False,PinType.bIsArray=False,PinType.bIsReference=False,PinType.bIsConst=False,PinType.bIsWeakPointer=False,LinkedTo=(K2Node_VariableGet_1970 D370BA9F4E23C50CF00DBF9C3866BEDE,),PersistentGuid=00000000000000000000000000000000,bHidden=False,bNotConnectable=False,bDefaultValueIsReadOnly=False,bDefaultValueIsIgnored=False,bAdvancedView=False,)
+				   CustomProperties Pin (PinId=FD05F54B48F429BC42F882B6033BFA3B,PinName="Value",PinType.PinCategory="text",PinType.PinSubCategory="",PinType.PinSubCategoryObject=None,PinType.PinSubCategoryMemberReference=(),PinType.PinValueType=(),PinType.bIsMap=False,PinType.bIsSet=False,PinType.bIsArray=False,PinType.bIsReference=True,PinType.bIsConst=True,PinType.bIsWeakPointer=False,LinkedTo=(K2Node_VariableGet_1993 89C664C140FF631F8AE9F1897A2308D0,),PersistentGuid=00000000000000000000000000000000,bHidden=False,bNotConnectable=False,bDefaultValueIsReadOnly=False,bDefaultValueIsIgnored=True,bAdvancedView=False,)
+				End Object
+				Begin Object Class=/Script/BlueprintGraph.K2Node_VariableGet Name="K2Node_VariableGet_1970"
+				   VariableReference=(MemberName="TextRender",bSelfContext=True)
+				   NodePosX=3232
+				   NodePosY=480
+				   NodeGuid=4B90E05A4FDAC1EF112CEF8E234EC7BD
+				   CustomProperties Pin (PinId=D370BA9F4E23C50CF00DBF9C3866BEDE,PinName="TextRender",Direction="EGPD_Output",PinType.PinCategory="object",PinType.PinSubCategory="",PinType.PinSubCategoryObject=Class'/Script/Engine.TextRenderComponent',PinType.PinSubCategoryMemberReference=(),PinType.PinValueType=(),PinType.bIsMap=False,PinType.bIsSet=False,PinType.bIsArray=False,PinType.bIsReference=False,PinType.bIsConst=False,PinType.bIsWeakPointer=False,LinkedTo=(K2Node_CallFunction_12105 24D9A43A4CC5BB2E7963438C3FE88719,),PersistentGuid=00000000000000000000000000000000,bHidden=False,bNotConnectable=False,bDefaultValueIsReadOnly=False,bDefaultValueIsIgnored=False,bAdvancedView=False,)
+				   CustomProperties Pin (PinId=83CDB7A849B7EB215EC3139D1E9C6B39,PinName="self",PinFriendlyName=NSLOCTEXT("K2Node", "Target", "Target"),PinType.PinCategory="object",PinType.PinSubCategory="",PinType.PinSubCategoryObject=BlueprintGeneratedClass'/Game/GearVR_Pawn/GearVR_Pawn_Input.GearVR_Pawn_Input_C',PinType.PinSubCategoryMemberReference=(),PinType.PinValueType=(),PinType.bIsMap=False,PinType.bIsSet=False,PinType.bIsArray=False,PinType.bIsReference=False,PinType.bIsConst=False,PinType.bIsWeakPointer=False,PersistentGuid=00000000000000000000000000000000,bHidden=True,bNotConnectable=False,bDefaultValueIsReadOnly=False,bDefaultValueIsIgnored=False,bAdvancedView=False,)
+				End Object
+				Begin Object Class=/Script/BlueprintGraph.K2Node_VariableGet Name="K2Node_VariableGet_1993"
+				   VariableReference=(MemberName="SwipeDown",MemberGuid=263228604828C2A8D38D40A2AC0D9A14,bSelfContext=True)
+				   NodePosX=3232
+				   NodePosY=544
+				   NodeGuid=33063D6D404D5D56A14B4A91DB86F2DF
+				   CustomProperties Pin (PinId=89C664C140FF631F8AE9F1897A2308D0,PinName="SwipeDown",Direction="EGPD_Output",PinType.PinCategory="text",PinType.PinSubCategory="",PinType.PinSubCategoryObject=None,PinType.PinSubCategoryMemberReference=(),PinType.PinValueType=(),PinType.bIsMap=False,PinType.bIsSet=False,PinType.bIsArray=False,PinType.bIsReference=False,PinType.bIsConst=False,PinType.bIsWeakPointer=False,LinkedTo=(K2Node_CallFunction_12105 FD05F54B48F429BC42F882B6033BFA3B,),PersistentGuid=00000000000000000000000000000000,bHidden=False,bNotConnectable=False,bDefaultValueIsReadOnly=False,bDefaultValueIsIgnored=False,bAdvancedView=False,)
+				   CustomProperties Pin (PinId=9B1894964697C33F41448D860FA0F8A2,PinName="self",PinFriendlyName=NSLOCTEXT("K2Node", "Target", "Target"),PinType.PinCategory="object",PinType.PinSubCategory="",PinType.PinSubCategoryObject=BlueprintGeneratedClass'/Game/GearVR_Pawn/GearVR_Pawn_Input.GearVR_Pawn_Input_C',PinType.PinSubCategoryMemberReference=(),PinType.PinValueType=(),PinType.bIsMap=False,PinType.bIsSet=False,PinType.bIsArray=False,PinType.bIsReference=False,PinType.bIsConst=False,PinType.bIsWeakPointer=False,PersistentGuid=00000000000000000000000000000000,bHidden=True,bNotConnectable=False,bDefaultValueIsReadOnly=False,bDefaultValueIsIgnored=False,bAdvancedView=False,)
+				End Object
+			END OBJECT
+	[/PARAM]
+[/OBJECT]
+[/EXCERPT:GearVRTouch]
+-->
+
+
+
+
