Copying //UE4/Dev-Documentation to //UE4/Samples-Main (Source: //UE4/Dev-Documentation @ 3623665)

diff --git a/Engine/Documentation/Source/Platforms/GearVR/HowTo/MCTouchPad/HTGearVRMCTouchpad.INT.udn b/Engine/Documentation/Source/Platforms/GearVR/HowTo/MCTouchPad/HTGearVRMCTouchpad.INT.udn
new file mode 100644
--- /dev/null
+++ b/Engine/Documentation/Source/Platforms/GearVR/HowTo/MCTouchPad/HTGearVRMCTouchpad.INT.udn
@@ -0,0 +1,474 @@
+Availability: Public
+Crumbs: %ROOT%
+Title: Use the Gear VR Touchpad 
+Description: Setting up the Gear VR Motion Controller Touchpad 
+Type: how-to
+SkillLevel: Intermidate
+Version: 4.17
+Parent: Platforms/GearVR/HowTo
+Order: 3
+Tags: GearVR
+topic-image:Platforms\GearVR\HowTo\GearVR_How_To_Header.png
+prereq: Platforms\GearVR\Prerequisites
+prereq: Platforms\GearVR\QuickStart
+prereq: Platforms\GearVR\HowTo\CameraSetup
+prereq: Platforms\GearVR\HowTo\MotionController
+
+Due to the highly precise nature of the Gear VR Motion Controller (MC) Touchpad, it can be used to perform all kinds of different interactions that require high precision. From scrolling through a series of menu options or trying to plot your precise location on a map, there is an almost endless amount of ways to use the Touchpad to interact with the virtual world. In the following How - To we will take a look at how to get and then display the current location that is being touched on the Gear VR Motion Controller Touchpad.
+
+[region:note]
+In this How - To we will be starting with a Pawn that already has a VR Camera and Motion Controller support added to it. If you do not have this setup, please check out the documents listed in the **Prerequisite Topics** list at the top of this page.
+[/region]
+
+## Gear VR Touchpad Setup
+
+1. First, open up your player Pawn Blueprint and from the **Component** tab, add two new **Text Render** component, called **X_Position** and **Y_Position**.  Make sure that both are parented to the **Motion Controller** component like in the image below.
+
+    ![](HTGVRMCTouchpad_00.png)
+
+1. Click on the Text Component and then in the **Details** panel set the following properties with the following values for both the X and Y text. 
+    
+    [REGION:lightbox]
+    [![](HTGVRMCTouchpad_01.png)(w:757)](HTGVRMCTouchpad_01.png)
+    [/REGION]
+
+    [REGION:caption]
+    Click for full image. Note that this is only showing the Y_Position properties. See table below for X_Position properties.
+    [/REGION]
+
+    |Property Name| Value|
+    |------|-----|
+    |**Location X Text**| X = 20, Y = 0, Z = 0|
+    |**Location Y Text**| X = 20, Y = 0, Z = 5|
+    |**Rotation X & Y**| X = 20, Y = 0, Z = -180|
+    |**Horizontal Alignment X & Y**|Center|
+    |**Vertical Alignment X & Y**| Text Center|
+    |**World Size X & Y** |5|
+
+1. Before we can build the logic for the Blueprint, we first need to create two variables to hold and display the X and Y input location on the Touchpad.  Go to the **Variables** section, add two new **Text Variables** called **XPositionText** and **YPositionText**. 
+
+    ![](HTGVRMCTouchpad_02.png)
+
+1. To display the current location that is being pressed on the Touchpad, we will first need to get and then store the X and Y  **Axis Value** of the current touch input. Then, using the **Set Text** node we can display that information via the Text Render Component that we added to our Gear VR Motion Controller. To do this in our Blueprint first, add the following nodes and variables to the Event Graph.
+
+    [REGION:lightbox]
+    [![](HTGVRMCTouchpad_03.png)(w:713)](HTGVRMCTouchpad_03.png)
+    [/REGION]
+
+    [REGION:caption]
+    Click for full image. 
+    [/REGION]
+
+    |Node / Variable Name| Value|
+    |------|-----|
+    |**MotionController(R)Thumbstick X**|N/A|
+    |**ToText(Float)**|N/A|
+    |**SET_ XPosition_Text(Float)**|N/A|
+    |**Set Text**|N/A|
+    |**X_Positiont**|N/A|
+    |**XPositionText**|N/A|
+    |**Append**|A Input: X:|
+    |**ToText(string)**|N/A|
+
+1. Next, we need to build the logic that will display the input location using the text that was attached to the Motion Controller. To do this, we will need to first get the **Axis Value** of the MotionController (R)Thumbstick X  and store that in the XPositionText variable. Once we have that information, we can then display this value using the X_Postion text. When completed, your Blueprint should look like the following image.
+
+    [REGION:fullwidth]
+    [INCLUDE:#MCInput01]
+    [/REGION]
+
+    [REGION:caption]
+    Click image to copy Blueprint code. 
+    [/REGION]
+
+1. Now that we have the X-axis input value, to get the Y axis input value, we can use the same logic, we just need to change the MotionController(R)Thumbstick and variables to use the Y axis. When completed your entire Blueprint should look like the following image.
+
+    [REGION:fullwidth]
+    [INCLUDE:#MCInput02]
+    [/REGION]
+
+    [REGION:caption]
+    Click image to copy Blueprint code. 
+    [/REGION]
+
+1. Once the Blueprint has been compiled and saved, deployed the project to your smartphone, and pick up your GearVR HMD and Motion Controller. Now, when you touch the Touchpad, you should see the numbers update with the location that you are currently touching like in the following video.
+
+    [OBJECT:EmbeddedVideo]
+        [PARAMLITERAL:videoid]
+        O2TZ4zTQQCU
+        [/PARAMLITERAL]
+        [PARAMLITERAL:width]
+        100
+        [/PARAMLITERAL]
+        [PARAMLITERAL:height]
+        56.25
+        [/PARAMLITERAL]
+        [PARAMLITERAL:units]
+        %
+        [/PARAMLITERAL]
+        [PARAMLITERAL:type]
+        simple
+        [/PARAMLITERAL]
+    [/OBJECT]
+
+## UE4 Project Downloads
+
+Below you will find a link to where you can download the UE4 project that was used to create this example. 
+
+* [**Gear VR Motion Controller Touchpad Input Example Project**](https://epicgames.box.com/s/9024yg9aldpvpckncw1c80876ztoxhf9)
+
+<!--
+[EXCERPT:MCInput01]    
+[OBJECT:Snippet]
+    [PARAMLITERAL:id]
+        MCInput01
+    [/PARAMLITERAL]
+    [PARAMLITERAL:type]
+        blueprint
+    [/PARAMLITERAL]
+    [PARAM:image]
+        [REGION:lightbox]
+        [![](HTGVRMCTouchpad_04.png)(w:739)](HTGVRMCTouchpad_04.png)
+        [/REGION]
+    [/PARAM]
+    [PARAM:snippet]
+            BEGIN OBJECT
+                            Begin Object Class=/Script/BlueprintGraph.K2Node_InputAxisKeyEvent Name="K2Node_InputAxisKeyEvent_1"
+               AxisKey=MotionController_Right_Thumbstick_X
+               EventReference=(MemberParent=Package'"/Script/Engine"')
+               CustomFunctionName="InpAxisKeyEvt_MotionController_Right_Thumbstick_X_K2Node_InputAxisKeyEvent_1"
+               NodePosX=704
+               NodePosY=-2480
+               NodeGuid=F607CDF74840D4EB339878B999FAC98E
+               CustomProperties Pin (PinId=60F7E126404F2A5E5FA2C4814371FDA1,PinName="OutputDelegate",Direction="EGPD_Output",PinType.PinCategory="delegate",PinType.PinSubCategory="",PinType.PinSubCategoryObject=None,PinType.PinSubCategoryMemberReference=(MemberParent=BlueprintGeneratedClass'/Game/Blueprints/GearVRMCTouchpadPawn.GearVRMCTouchpadPawn_C',MemberName="InpAxisKeyEvt_MotionController_Right_Thumbstick_X_K2Node_InputAxisKeyEvent_1"),PinType.PinValueType=(),PinType.ContainerType=None,PinType.bIsMap=False,PinType.bIsSet=False,PinType.bIsArray=False,PinType.bIsReference=False,PinType.bIsConst=False,PinType.bIsWeakPointer=False,PersistentGuid=00000000000000000000000000000000,bHidden=False,bNotConnectable=False,bDefaultValueIsReadOnly=False,bDefaultValueIsIgnored=False,bAdvancedView=False,bOrphanedPin=False,)
+               CustomProperties Pin (PinId=079E818244554159C73D02B16AEA1A62,PinName="then",Direction="EGPD_Output",PinType.PinCategory="exec",PinType.PinSubCategory="",PinType.PinSubCategoryObject=None,PinType.PinSubCategoryMemberReference=(),PinType.PinValueType=(),PinType.ContainerType=None,PinType.bIsMap=False,PinType.bIsSet=False,PinType.bIsArray=False,PinType.bIsReference=False,PinType.bIsConst=False,PinType.bIsWeakPointer=False,LinkedTo=(K2Node_VariableSet_245 03D9E61948784EA9BE735BBD6E306A7B,),PersistentGuid=00000000000000000000000000000000,bHidden=False,bNotConnectable=False,bDefaultValueIsReadOnly=False,bDefaultValueIsIgnored=False,bAdvancedView=False,bOrphanedPin=False,)
+               CustomProperties Pin (PinId=18A653DB4F75BEA3C627FDA82DFC64C4,PinName="AxisValue",PinToolTip="Axis Value\nFloat",Direction="EGPD_Output",PinType.PinCategory="float",PinType.PinSubCategory="",PinType.PinSubCategoryObject=None,PinType.PinSubCategoryMemberReference=(),PinType.PinValueType=(),PinType.ContainerType=None,PinType.bIsMap=False,PinType.bIsSet=False,PinType.bIsArray=False,PinType.bIsReference=False,PinType.bIsConst=False,PinType.bIsWeakPointer=False,DefaultValue="0.0",AutogeneratedDefaultValue="0.0",LinkedTo=(K2Node_CallFunction_226 F80FC4A04331E969D71EB5A87073C49B,),PersistentGuid=00000000000000000000000000000000,bHidden=False,bNotConnectable=False,bDefaultValueIsReadOnly=False,bDefaultValueIsIgnored=False,bAdvancedView=False,bOrphanedPin=False,)
+            End Object
+            Begin Object Class=/Script/BlueprintGraph.K2Node_VariableSet Name="K2Node_VariableSet_245"
+               VariableReference=(MemberName="XPositionText",MemberGuid=009C67F84C5646E19F766381539F84C2,bSelfContext=True)
+               NodePosX=1184
+               NodePosY=-2464
+               NodeGuid=7B0F663945474DB84DD00AB915B5AE60
+               CustomProperties Pin (PinId=03D9E61948784EA9BE735BBD6E306A7B,PinName="execute",PinType.PinCategory="exec",PinType.PinSubCategory="",PinType.PinSubCategoryObject=None,PinType.PinSubCategoryMemberReference=(),PinType.PinValueType=(),PinType.ContainerType=None,PinType.bIsMap=False,PinType.bIsSet=False,PinType.bIsArray=False,PinType.bIsReference=False,PinType.bIsConst=False,PinType.bIsWeakPointer=False,LinkedTo=(K2Node_InputAxisKeyEvent_1 079E818244554159C73D02B16AEA1A62,),PersistentGuid=00000000000000000000000000000000,bHidden=False,bNotConnectable=False,bDefaultValueIsReadOnly=False,bDefaultValueIsIgnored=False,bAdvancedView=False,bOrphanedPin=False,)
+               CustomProperties Pin (PinId=AA86D7454ACF6604F9CC0BA5C09ED35E,PinName="then",Direction="EGPD_Output",PinType.PinCategory="exec",PinType.PinSubCategory="",PinType.PinSubCategoryObject=None,PinType.PinSubCategoryMemberReference=(),PinType.PinValueType=(),PinType.ContainerType=None,PinType.bIsMap=False,PinType.bIsSet=False,PinType.bIsArray=False,PinType.bIsReference=False,PinType.bIsConst=False,PinType.bIsWeakPointer=False,LinkedTo=(K2Node_CallFunction_2353 96D7193F444562BF1B9D1A82C15146E1,),PersistentGuid=00000000000000000000000000000000,bHidden=False,bNotConnectable=False,bDefaultValueIsReadOnly=False,bDefaultValueIsIgnored=False,bAdvancedView=False,bOrphanedPin=False,)
+               CustomProperties Pin (PinId=8EBF99184B8C3FF82F0973B50E9A2831,PinName="XPositionText",PinType.PinCategory="text",PinType.PinSubCategory="",PinType.PinSubCategoryObject=None,PinType.PinSubCategoryMemberReference=(),PinType.PinValueType=(),PinType.ContainerType=None,PinType.bIsMap=False,PinType.bIsSet=False,PinType.bIsArray=False,PinType.bIsReference=False,PinType.bIsConst=False,PinType.bIsWeakPointer=False,LinkedTo=(K2Node_CallFunction_226 143CA0A3412175455799A096F47FDA4C,),PersistentGuid=00000000000000000000000000000000,bHidden=False,bNotConnectable=False,bDefaultValueIsReadOnly=False,bDefaultValueIsIgnored=False,bAdvancedView=False,bOrphanedPin=False,)
+               CustomProperties Pin (PinId=1013C0854372B5B2C4818B81798B6C77,PinName="Output_Get",PinToolTip="Retrieves the value of the variable, can use instead of a separate Get node",Direction="EGPD_Output",PinType.PinCategory="text",PinType.PinSubCategory="",PinType.PinSubCategoryObject=None,PinType.PinSubCategoryMemberReference=(),PinType.PinValueType=(),PinType.ContainerType=None,PinType.bIsMap=False,PinType.bIsSet=False,PinType.bIsArray=False,PinType.bIsReference=False,PinType.bIsConst=False,PinType.bIsWeakPointer=False,PersistentGuid=00000000000000000000000000000000,bHidden=False,bNotConnectable=False,bDefaultValueIsReadOnly=False,bDefaultValueIsIgnored=False,bAdvancedView=False,bOrphanedPin=False,)
+               CustomProperties Pin (PinId=1083FF874AB7737633CB3E89CB7D9F9E,PinName="self",PinFriendlyName=NSLOCTEXT("K2Node", "Target", "Target"),PinType.PinCategory="object",PinType.PinSubCategory="",PinType.PinSubCategoryObject=BlueprintGeneratedClass'/Game/Blueprints/GearVRMCTouchpadPawn.GearVRMCTouchpadPawn_C',PinType.PinSubCategoryMemberReference=(),PinType.PinValueType=(),PinType.ContainerType=None,PinType.bIsMap=False,PinType.bIsSet=False,PinType.bIsArray=False,PinType.bIsReference=False,PinType.bIsConst=False,PinType.bIsWeakPointer=False,PersistentGuid=00000000000000000000000000000000,bHidden=True,bNotConnectable=False,bDefaultValueIsReadOnly=False,bDefaultValueIsIgnored=False,bAdvancedView=False,bOrphanedPin=False,)
+            End Object
+            Begin Object Class=/Script/BlueprintGraph.K2Node_VariableGet Name="K2Node_VariableGet_172"
+               VariableReference=(MemberName="X_Position",bSelfContext=True)
+               NodePosX=1296
+               NodePosY=-2368
+               NodeGuid=E4D891F74A0D3667142C8D81CC65E092
+               CustomProperties Pin (PinId=16E9007D45B6E55C549AE1A9E153817A,PinName="X_Position",Direction="EGPD_Output",PinType.PinCategory="object",PinType.PinSubCategory="",PinType.PinSubCategoryObject=Class'/Script/Engine.TextRenderComponent',PinType.PinSubCategoryMemberReference=(),PinType.PinValueType=(),PinType.ContainerType=None,PinType.bIsMap=False,PinType.bIsSet=False,PinType.bIsArray=False,PinType.bIsReference=False,PinType.bIsConst=False,PinType.bIsWeakPointer=False,LinkedTo=(K2Node_CallFunction_2353 24D9A43A4CC5BB2E7963438C3FE88719,),PersistentGuid=00000000000000000000000000000000,bHidden=False,bNotConnectable=False,bDefaultValueIsReadOnly=False,bDefaultValueIsIgnored=False,bAdvancedView=False,bOrphanedPin=False,)
+               CustomProperties Pin (PinId=C611BC784A72C63BFCD9A0B050B0585D,PinName="self",PinFriendlyName=NSLOCTEXT("K2Node", "Target", "Target"),PinType.PinCategory="object",PinType.PinSubCategory="",PinType.PinSubCategoryObject=BlueprintGeneratedClass'/Game/Blueprints/GearVRMCTouchpadPawn.GearVRMCTouchpadPawn_C',PinType.PinSubCategoryMemberReference=(),PinType.PinValueType=(),PinType.ContainerType=None,PinType.bIsMap=False,PinType.bIsSet=False,PinType.bIsArray=False,PinType.bIsReference=False,PinType.bIsConst=False,PinType.bIsWeakPointer=False,PersistentGuid=00000000000000000000000000000000,bHidden=True,bNotConnectable=False,bDefaultValueIsReadOnly=False,bDefaultValueIsIgnored=False,bAdvancedView=False,bOrphanedPin=False,)
+            End Object
+            Begin Object Class=/Script/BlueprintGraph.K2Node_CallFunction Name="K2Node_CallFunction_2353"
+               FunctionReference=(MemberParent=Class'/Script/Engine.TextRenderComponent',MemberName="K2_SetText")
+               NodePosX=1520
+               NodePosY=-2496
+               NodeGuid=9C5E75224A0CE02A79969AA3F0AE70B1
+               CustomProperties Pin (PinId=96D7193F444562BF1B9D1A82C15146E1,PinName="execute",PinToolTip="\nExec",PinType.PinCategory="exec",PinType.PinSubCategory="",PinType.PinSubCategoryObject=None,PinType.PinSubCategoryMemberReference=(),PinType.PinValueType=(),PinType.ContainerType=None,PinType.bIsMap=False,PinType.bIsSet=False,PinType.bIsArray=False,PinType.bIsReference=False,PinType.bIsConst=False,PinType.bIsWeakPointer=False,LinkedTo=(K2Node_VariableSet_245 AA86D7454ACF6604F9CC0BA5C09ED35E,),PersistentGuid=00000000000000000000000000000000,bHidden=False,bNotConnectable=False,bDefaultValueIsReadOnly=False,bDefaultValueIsIgnored=False,bAdvancedView=False,bOrphanedPin=False,)
+               CustomProperties Pin (PinId=66DCFA3744967675B795929682EA35F0,PinName="then",PinToolTip="\nExec",Direction="EGPD_Output",PinType.PinCategory="exec",PinType.PinSubCategory="",PinType.PinSubCategoryObject=None,PinType.PinSubCategoryMemberReference=(),PinType.PinValueType=(),PinType.ContainerType=None,PinType.bIsMap=False,PinType.bIsSet=False,PinType.bIsArray=False,PinType.bIsReference=False,PinType.bIsConst=False,PinType.bIsWeakPointer=False,PersistentGuid=00000000000000000000000000000000,bHidden=False,bNotConnectable=False,bDefaultValueIsReadOnly=False,bDefaultValueIsIgnored=False,bAdvancedView=False,bOrphanedPin=False,)
+               CustomProperties Pin (PinId=24D9A43A4CC5BB2E7963438C3FE88719,PinName="self",PinFriendlyName=NSLOCTEXT("K2Node", "Target", "Target"),PinToolTip="Target\nText Render Component Object Reference",PinType.PinCategory="object",PinType.PinSubCategory="",PinType.PinSubCategoryObject=Class'/Script/Engine.TextRenderComponent',PinType.PinSubCategoryMemberReference=(),PinType.PinValueType=(),PinType.ContainerType=None,PinType.bIsMap=False,PinType.bIsSet=False,PinType.bIsArray=False,PinType.bIsReference=False,PinType.bIsConst=False,PinType.bIsWeakPointer=False,LinkedTo=(K2Node_VariableGet_172 16E9007D45B6E55C549AE1A9E153817A,),PersistentGuid=00000000000000000000000000000000,bHidden=False,bNotConnectable=False,bDefaultValueIsReadOnly=False,bDefaultValueIsIgnored=False,bAdvancedView=False,bOrphanedPin=False,)
+               CustomProperties Pin (PinId=FD05F54B48F429BC42F882B6033BFA3B,PinName="Value",PinToolTip="Value\nText (by ref)",PinType.PinCategory="text",PinType.PinSubCategory="",PinType.PinSubCategoryObject=None,PinType.PinSubCategoryMemberReference=(),PinType.PinValueType=(),PinType.ContainerType=None,PinType.bIsMap=False,PinType.bIsSet=False,PinType.bIsArray=False,PinType.bIsReference=True,PinType.bIsConst=True,PinType.bIsWeakPointer=False,LinkedTo=(K2Node_CallFunction_3064 7C3752734C56A4E232CD5E9FD7051C5F,),PersistentGuid=00000000000000000000000000000000,bHidden=False,bNotConnectable=False,bDefaultValueIsReadOnly=False,bDefaultValueIsIgnored=True,bAdvancedView=False,bOrphanedPin=False,)
+            End Object
+            Begin Object Class=/Script/BlueprintGraph.K2Node_VariableGet Name="K2Node_VariableGet_313"
+               VariableReference=(MemberName="XPositionText",MemberGuid=009C67F84C5646E19F766381539F84C2,bSelfContext=True)
+               NodePosX=832
+               NodePosY=-2192
+               NodeGuid=0C0C2BC34DF31F0B3E7D4B9060871E1B
+               CustomProperties Pin (PinId=C2ABF17C4E4C3F86E7C64FB1135A50A7,PinName="XPositionText",Direction="EGPD_Output",PinType.PinCategory="text",PinType.PinSubCategory="",PinType.PinSubCategoryObject=None,PinType.PinSubCategoryMemberReference=(),PinType.PinValueType=(),PinType.ContainerType=None,PinType.bIsMap=False,PinType.bIsSet=False,PinType.bIsArray=False,PinType.bIsReference=False,PinType.bIsConst=False,PinType.bIsWeakPointer=False,LinkedTo=(K2Node_CallFunction_3252 0539BDC14912061864F1929BACA3E29C,),PersistentGuid=00000000000000000000000000000000,bHidden=False,bNotConnectable=False,bDefaultValueIsReadOnly=False,bDefaultValueIsIgnored=False,bAdvancedView=False,bOrphanedPin=False,)
+               CustomProperties Pin (PinId=404B459841AC051FEE64FDBE286E724F,PinName="self",PinFriendlyName=NSLOCTEXT("K2Node", "Target", "Target"),PinType.PinCategory="object",PinType.PinSubCategory="",PinType.PinSubCategoryObject=BlueprintGeneratedClass'/Game/Blueprints/GearVRMCTouchpadPawn.GearVRMCTouchpadPawn_C',PinType.PinSubCategoryMemberReference=(),PinType.PinValueType=(),PinType.ContainerType=None,PinType.bIsMap=False,PinType.bIsSet=False,PinType.bIsArray=False,PinType.bIsReference=False,PinType.bIsConst=False,PinType.bIsWeakPointer=False,PersistentGuid=00000000000000000000000000000000,bHidden=True,bNotConnectable=False,bDefaultValueIsReadOnly=False,bDefaultValueIsIgnored=False,bAdvancedView=False,bOrphanedPin=False,)
+            End Object
+            Begin Object Class=/Script/BlueprintGraph.K2Node_CommutativeAssociativeBinaryOperator Name="K2Node_CommutativeAssociativeBinaryOperator_718"
+               bIsPureFunc=True
+               FunctionReference=(MemberParent=Class'/Script/Engine.KismetStringLibrary',MemberName="Concat_StrStr")
+               NodePosX=1104
+               NodePosY=-2224
+               NodeGuid=A6D3C8B740CF656CD93AD482F60E3C92
+               CustomProperties Pin (PinId=032F0C204576AD0073097FA4C783B85A,PinName="self",PinFriendlyName=NSLOCTEXT("K2Node", "Target", "Target"),PinToolTip="Target\nKismet String Library Object Reference",PinType.PinCategory="object",PinType.PinSubCategory="",PinType.PinSubCategoryObject=Class'/Script/Engine.KismetStringLibrary',PinType.PinSubCategoryMemberReference=(),PinType.PinValueType=(),PinType.ContainerType=None,PinType.bIsMap=False,PinType.bIsSet=False,PinType.bIsArray=False,PinType.bIsReference=False,PinType.bIsConst=False,PinType.bIsWeakPointer=False,DefaultObject="/Script/Engine.Default__KismetStringLibrary",PersistentGuid=00000000000000000000000000000000,bHidden=True,bNotConnectable=False,bDefaultValueIsReadOnly=False,bDefaultValueIsIgnored=False,bAdvancedView=False,bOrphanedPin=False,)
+               CustomProperties Pin (PinId=D4A956EF48D9C39C100216B65110C2D3,PinName="A",PinToolTip="A\nString\n\nThe original string",PinType.PinCategory="string",PinType.PinSubCategory="",PinType.PinSubCategoryObject=None,PinType.PinSubCategoryMemberReference=(),PinType.PinValueType=(),PinType.ContainerType=None,PinType.bIsMap=False,PinType.bIsSet=False,PinType.bIsArray=False,PinType.bIsReference=False,PinType.bIsConst=False,PinType.bIsWeakPointer=False,DefaultValue="X:",PersistentGuid=00000000000000000000000000000000,bHidden=False,bNotConnectable=False,bDefaultValueIsReadOnly=False,bDefaultValueIsIgnored=False,bAdvancedView=False,bOrphanedPin=False,)
+               CustomProperties Pin (PinId=9A9C2749448F164BFB5DB5BF6E3C5674,PinName="B",PinToolTip="B\nString\n\nThe string to append to A",PinType.PinCategory="string",PinType.PinSubCategory="",PinType.PinSubCategoryObject=None,PinType.PinSubCategoryMemberReference=(),PinType.PinValueType=(),PinType.ContainerType=None,PinType.bIsMap=False,PinType.bIsSet=False,PinType.bIsArray=False,PinType.bIsReference=False,PinType.bIsConst=False,PinType.bIsWeakPointer=False,LinkedTo=(K2Node_CallFunction_3252 C8C0111E4AEA812BAA815AB2A78D2BA7,),PersistentGuid=00000000000000000000000000000000,bHidden=False,bNotConnectable=False,bDefaultValueIsReadOnly=False,bDefaultValueIsIgnored=False,bAdvancedView=False,bOrphanedPin=False,)
+               CustomProperties Pin (PinId=3C2548244EFB0C45430FE294FCFAAD0B,PinName="ReturnValue",PinToolTip="Return Value\nString\n\nA new string which is the concatenation of A+B",Direction="EGPD_Output",PinType.PinCategory="string",PinType.PinSubCategory="",PinType.PinSubCategoryObject=None,PinType.PinSubCategoryMemberReference=(),PinType.PinValueType=(),PinType.ContainerType=None,PinType.bIsMap=False,PinType.bIsSet=False,PinType.bIsArray=False,PinType.bIsReference=False,PinType.bIsConst=False,PinType.bIsWeakPointer=False,LinkedTo=(K2Node_CallFunction_3064 5E00EF9D41203D0B84183E8F49046194,),PersistentGuid=00000000000000000000000000000000,bHidden=False,bNotConnectable=False,bDefaultValueIsReadOnly=False,bDefaultValueIsIgnored=False,bAdvancedView=False,bOrphanedPin=False,)
+            End Object
+            Begin Object Class=/Script/BlueprintGraph.K2Node_CallFunction Name="K2Node_CallFunction_3064"
+               bIsPureFunc=True
+               FunctionReference=(MemberParent=Class'/Script/Engine.KismetTextLibrary',MemberName="Conv_StringToText")
+               NodePosX=1360
+               NodePosY=-2320
+               NodeGuid=53EC61D0468ED08B58F99FA9AC397C10
+               CustomProperties Pin (PinId=671E8245468EDCE29C389F86DBDDBA2C,PinName="self",PinFriendlyName=NSLOCTEXT("K2Node", "Target", "Target"),PinToolTip="Target\nKismet Text Library Object Reference",PinType.PinCategory="object",PinType.PinSubCategory="",PinType.PinSubCategoryObject=Class'/Script/Engine.KismetTextLibrary',PinType.PinSubCategoryMemberReference=(),PinType.PinValueType=(),PinType.ContainerType=None,PinType.bIsMap=False,PinType.bIsSet=False,PinType.bIsArray=False,PinType.bIsReference=False,PinType.bIsConst=False,PinType.bIsWeakPointer=False,DefaultObject="/Script/Engine.Default__KismetTextLibrary",PersistentGuid=00000000000000000000000000000000,bHidden=True,bNotConnectable=False,bDefaultValueIsReadOnly=False,bDefaultValueIsIgnored=False,bAdvancedView=False,bOrphanedPin=False,)
+               CustomProperties Pin (PinId=5E00EF9D41203D0B84183E8F49046194,PinName="InString",PinToolTip="In String\nString",PinType.PinCategory="string",PinType.PinSubCategory="",PinType.PinSubCategoryObject=None,PinType.PinSubCategoryMemberReference=(),PinType.PinValueType=(),PinType.ContainerType=None,PinType.bIsMap=False,PinType.bIsSet=False,PinType.bIsArray=False,PinType.bIsReference=False,PinType.bIsConst=False,PinType.bIsWeakPointer=False,LinkedTo=(K2Node_CommutativeAssociativeBinaryOperator_718 3C2548244EFB0C45430FE294FCFAAD0B,),PersistentGuid=00000000000000000000000000000000,bHidden=False,bNotConnectable=False,bDefaultValueIsReadOnly=False,bDefaultValueIsIgnored=False,bAdvancedView=False,bOrphanedPin=False,)
+               CustomProperties Pin (PinId=7C3752734C56A4E232CD5E9FD7051C5F,PinName="ReturnValue",PinToolTip="Return Value\nText",Direction="EGPD_Output",PinType.PinCategory="text",PinType.PinSubCategory="",PinType.PinSubCategoryObject=None,PinType.PinSubCategoryMemberReference=(),PinType.PinValueType=(),PinType.ContainerType=None,PinType.bIsMap=False,PinType.bIsSet=False,PinType.bIsArray=False,PinType.bIsReference=False,PinType.bIsConst=False,PinType.bIsWeakPointer=False,LinkedTo=(K2Node_CallFunction_2353 FD05F54B48F429BC42F882B6033BFA3B,),PersistentGuid=00000000000000000000000000000000,bHidden=False,bNotConnectable=False,bDefaultValueIsReadOnly=False,bDefaultValueIsIgnored=False,bAdvancedView=False,bOrphanedPin=False,)
+            End Object
+            Begin Object Class=/Script/BlueprintGraph.K2Node_CallFunction Name="K2Node_CallFunction_3252"
+               bIsPureFunc=True
+               FunctionReference=(MemberParent=Class'/Script/Engine.KismetTextLibrary',MemberName="Conv_TextToString")
+               NodePosX=976
+               NodePosY=-2192
+               NodeGuid=E9FA9F574222DE3A0E63AFAE99F4DB1E
+               CustomProperties Pin (PinId=594C56934A30AC39E8049D8710366307,PinName="self",PinFriendlyName=NSLOCTEXT("K2Node", "Target", "Target"),PinType.PinCategory="object",PinType.PinSubCategory="",PinType.PinSubCategoryObject=Class'/Script/Engine.KismetTextLibrary',PinType.PinSubCategoryMemberReference=(),PinType.PinValueType=(),PinType.ContainerType=None,PinType.bIsMap=False,PinType.bIsSet=False,PinType.bIsArray=False,PinType.bIsReference=False,PinType.bIsConst=False,PinType.bIsWeakPointer=False,DefaultObject="/Script/Engine.Default__KismetTextLibrary",PersistentGuid=00000000000000000000000000000000,bHidden=True,bNotConnectable=False,bDefaultValueIsReadOnly=False,bDefaultValueIsIgnored=False,bAdvancedView=False,bOrphanedPin=False,)
+               CustomProperties Pin (PinId=0539BDC14912061864F1929BACA3E29C,PinName="InText",PinType.PinCategory="text",PinType.PinSubCategory="",PinType.PinSubCategoryObject=None,PinType.PinSubCategoryMemberReference=(),PinType.PinValueType=(),PinType.ContainerType=None,PinType.bIsMap=False,PinType.bIsSet=False,PinType.bIsArray=False,PinType.bIsReference=True,PinType.bIsConst=True,PinType.bIsWeakPointer=False,LinkedTo=(K2Node_VariableGet_313 C2ABF17C4E4C3F86E7C64FB1135A50A7,),PersistentGuid=00000000000000000000000000000000,bHidden=False,bNotConnectable=False,bDefaultValueIsReadOnly=False,bDefaultValueIsIgnored=True,bAdvancedView=False,bOrphanedPin=False,)
+               CustomProperties Pin (PinId=C8C0111E4AEA812BAA815AB2A78D2BA7,PinName="ReturnValue",Direction="EGPD_Output",PinType.PinCategory="string",PinType.PinSubCategory="",PinType.PinSubCategoryObject=None,PinType.PinSubCategoryMemberReference=(),PinType.PinValueType=(),PinType.ContainerType=None,PinType.bIsMap=False,PinType.bIsSet=False,PinType.bIsArray=False,PinType.bIsReference=False,PinType.bIsConst=False,PinType.bIsWeakPointer=False,LinkedTo=(K2Node_CommutativeAssociativeBinaryOperator_718 9A9C2749448F164BFB5DB5BF6E3C5674,),PersistentGuid=00000000000000000000000000000000,bHidden=False,bNotConnectable=False,bDefaultValueIsReadOnly=False,bDefaultValueIsIgnored=False,bAdvancedView=False,bOrphanedPin=False,)
+            End Object
+            Begin Object Class=/Script/BlueprintGraph.K2Node_CallFunction Name="K2Node_CallFunction_226"
+               bIsPureFunc=True
+               FunctionReference=(MemberParent=Class'/Script/Engine.KismetTextLibrary',MemberName="Conv_FloatToText")
+               NodePosX=976
+               NodePosY=-2384
+               AdvancedPinDisplay=Hidden
+               NodeGuid=E93C98884AA70223E58999BBA3F7AEA1
+               CustomProperties Pin (PinId=9231A80F48229FB948A5AA9DA215DB97,PinName="self",PinFriendlyName=NSLOCTEXT("K2Node", "Target", "Target"),PinToolTip="Target\nKismet Text Library Object Reference",PinType.PinCategory="object",PinType.PinSubCategory="",PinType.PinSubCategoryObject=Class'/Script/Engine.KismetTextLibrary',PinType.PinSubCategoryMemberReference=(),PinType.PinValueType=(),PinType.ContainerType=None,PinType.bIsMap=False,PinType.bIsSet=False,PinType.bIsArray=False,PinType.bIsReference=False,PinType.bIsConst=False,PinType.bIsWeakPointer=False,DefaultObject="/Script/Engine.Default__KismetTextLibrary",PersistentGuid=00000000000000000000000000000000,bHidden=True,bNotConnectable=False,bDefaultValueIsReadOnly=False,bDefaultValueIsIgnored=False,bAdvancedView=False,bOrphanedPin=False,)
+               CustomProperties Pin (PinId=F80FC4A04331E969D71EB5A87073C49B,PinName="Value",PinToolTip="Value\nFloat",PinType.PinCategory="float",PinType.PinSubCategory="",PinType.PinSubCategoryObject=None,PinType.PinSubCategoryMemberReference=(),PinType.PinValueType=(),PinType.ContainerType=None,PinType.bIsMap=False,PinType.bIsSet=False,PinType.bIsArray=False,PinType.bIsReference=False,PinType.bIsConst=False,PinType.bIsWeakPointer=False,DefaultValue="0.0",AutogeneratedDefaultValue="0.0",LinkedTo=(K2Node_InputAxisKeyEvent_1 18A653DB4F75BEA3C627FDA82DFC64C4,),PersistentGuid=00000000000000000000000000000000,bHidden=False,bNotConnectable=False,bDefaultValueIsReadOnly=False,bDefaultValueIsIgnored=False,bAdvancedView=False,bOrphanedPin=False,)
+               CustomProperties Pin (PinId=059645DC43483A351736259868CB24DE,PinName="RoundingMode",PinToolTip="Rounding Mode\nERoundingMode Enum",PinType.PinCategory="byte",PinType.PinSubCategory="",PinType.PinSubCategoryObject=Enum'/Script/Engine.ERoundingMode',PinType.PinSubCategoryMemberReference=(),PinType.PinValueType=(),PinType.ContainerType=None,PinType.bIsMap=False,PinType.bIsSet=False,PinType.bIsArray=False,PinType.bIsReference=False,PinType.bIsConst=False,PinType.bIsWeakPointer=False,DefaultValue="HalfToEven",PersistentGuid=00000000000000000000000000000000,bHidden=False,bNotConnectable=False,bDefaultValueIsReadOnly=False,bDefaultValueIsIgnored=False,bAdvancedView=True,bOrphanedPin=False,)
+               CustomProperties Pin (PinId=806085F24E097F901F75A398643EC77C,PinName="bUseGrouping",PinToolTip="Use Grouping\nBoolean",PinType.PinCategory="bool",PinType.PinSubCategory="",PinType.PinSubCategoryObject=None,PinType.PinSubCategoryMemberReference=(),PinType.PinValueType=(),PinType.ContainerType=None,PinType.bIsMap=False,PinType.bIsSet=False,PinType.bIsArray=False,PinType.bIsReference=False,PinType.bIsConst=False,PinType.bIsWeakPointer=False,DefaultValue="true",AutogeneratedDefaultValue="true",PersistentGuid=00000000000000000000000000000000,bHidden=False,bNotConnectable=False,bDefaultValueIsReadOnly=False,bDefaultValueIsIgnored=False,bAdvancedView=True,bOrphanedPin=False,)
+               CustomProperties Pin (PinId=AC50E40147020DFEB11D9B9AF578E13F,PinName="MinimumIntegralDigits",PinToolTip="Minimum Integral Digits\nInteger",PinType.PinCategory="int",PinType.PinSubCategory="",PinType.PinSubCategoryObject=None,PinType.PinSubCategoryMemberReference=(),PinType.PinValueType=(),PinType.ContainerType=None,PinType.bIsMap=False,PinType.bIsSet=False,PinType.bIsArray=False,PinType.bIsReference=False,PinType.bIsConst=False,PinType.bIsWeakPointer=False,DefaultValue="1",AutogeneratedDefaultValue="1",PersistentGuid=00000000000000000000000000000000,bHidden=False,bNotConnectable=False,bDefaultValueIsReadOnly=False,bDefaultValueIsIgnored=False,bAdvancedView=True,bOrphanedPin=False,)
+               CustomProperties Pin (PinId=5113623F4E1915E2B21972942DE375D4,PinName="MaximumIntegralDigits",PinToolTip="Maximum Integral Digits\nInteger",PinType.PinCategory="int",PinType.PinSubCategory="",PinType.PinSubCategoryObject=None,PinType.PinSubCategoryMemberReference=(),PinType.PinValueType=(),PinType.ContainerType=None,PinType.bIsMap=False,PinType.bIsSet=False,PinType.bIsArray=False,PinType.bIsReference=False,PinType.bIsConst=False,PinType.bIsWeakPointer=False,DefaultValue="324",AutogeneratedDefaultValue="324",PersistentGuid=00000000000000000000000000000000,bHidden=False,bNotConnectable=False,bDefaultValueIsReadOnly=False,bDefaultValueIsIgnored=False,bAdvancedView=True,bOrphanedPin=False,)
+               CustomProperties Pin (PinId=C03DF2CA437834ED7B2A8C8FFD834DD9,PinName="MinimumFractionalDigits",PinToolTip="Minimum Fractional Digits\nInteger",PinType.PinCategory="int",PinType.PinSubCategory="",PinType.PinSubCategoryObject=None,PinType.PinSubCategoryMemberReference=(),PinType.PinValueType=(),PinType.ContainerType=None,PinType.bIsMap=False,PinType.bIsSet=False,PinType.bIsArray=False,PinType.bIsReference=False,PinType.bIsConst=False,PinType.bIsWeakPointer=False,DefaultValue="0",AutogeneratedDefaultValue="0",PersistentGuid=00000000000000000000000000000000,bHidden=False,bNotConnectable=False,bDefaultValueIsReadOnly=False,bDefaultValueIsIgnored=False,bAdvancedView=True,bOrphanedPin=False,)
+               CustomProperties Pin (PinId=623B2E184048722C576F60A8736CBCC6,PinName="MaximumFractionalDigits",PinToolTip="Maximum Fractional Digits\nInteger",PinType.PinCategory="int",PinType.PinSubCategory="",PinType.PinSubCategoryObject=None,PinType.PinSubCategoryMemberReference=(),PinType.PinValueType=(),PinType.ContainerType=None,PinType.bIsMap=False,PinType.bIsSet=False,PinType.bIsArray=False,PinType.bIsReference=False,PinType.bIsConst=False,PinType.bIsWeakPointer=False,DefaultValue="3",AutogeneratedDefaultValue="3",PersistentGuid=00000000000000000000000000000000,bHidden=False,bNotConnectable=False,bDefaultValueIsReadOnly=False,bDefaultValueIsIgnored=False,bAdvancedView=True,bOrphanedPin=False,)
+               CustomProperties Pin (PinId=143CA0A3412175455799A096F47FDA4C,PinName="ReturnValue",PinToolTip="Return Value\nText",Direction="EGPD_Output",PinType.PinCategory="text",PinType.PinSubCategory="",PinType.PinSubCategoryObject=None,PinType.PinSubCategoryMemberReference=(),PinType.PinValueType=(),PinType.ContainerType=None,PinType.bIsMap=False,PinType.bIsSet=False,PinType.bIsArray=False,PinType.bIsReference=False,PinType.bIsConst=False,PinType.bIsWeakPointer=False,LinkedTo=(K2Node_VariableSet_245 8EBF99184B8C3FF82F0973B50E9A2831,),PersistentGuid=00000000000000000000000000000000,bHidden=False,bNotConnectable=False,bDefaultValueIsReadOnly=False,bDefaultValueIsIgnored=False,bAdvancedView=False,bOrphanedPin=False,)
+            End Object
+            END OBJECT
+    [/PARAM]
+[/OBJECT]
+[/EXCERPT:MCInput01]
+-->
+
+<!--
+[EXCERPT:MCInput02]    
+[OBJECT:Snippet]
+    [PARAMLITERAL:id]
+        MCInput02
+    [/PARAMLITERAL]
+    [PARAMLITERAL:type]
+        blueprint
+    [/PARAMLITERAL]
+    [PARAM:image]
+        [REGION:lightbox]
+        [![](HTGVRMCTouchpad_05.png)(w:793)](HTGVRMCTouchpad_05.png)
+        [/REGION]
+    [/PARAM]
+    [PARAM:snippet]
+            BEGIN OBJECT
+                        Begin Object Class=/Script/BlueprintGraph.K2Node_InputAxisKeyEvent Name="K2Node_InputAxisKeyEvent_22"
+               AxisKey=MotionController_Right_Thumbstick_X
+               EventReference=(MemberParent=Package'"/Script/Engine"')
+               CustomFunctionName="InpAxisKeyEvt_MotionController_Right_Thumbstick_X_K2Node_InputAxisKeyEvent_1"
+               NodePosX=1072
+               NodePosY=-1680
+               ErrorType=1
+               ErrorMsg="Found more than one function with the same name InpAxisKeyEvt_MotionController_Right_Thumbstick_X_K2Node_InputAxisKeyEvent_1; second occurance at  MotionController (R) Thumbstick X\nFound more than one function with the same name InpAxisKeyEvt_MotionController_Right_Thumbstick_X_K2Node_InputAxisKeyEvent_1; second occurance at  MotionController (R) Thumbstick X\nFound more than one function with the same name InpAxisKeyEvt_MotionController_Right_Thumbstick_X_K2Node_InputAxisKeyEvent_1; second occurance at  MotionController (R) Thumbstick X"
+               NodeGuid=75CEF68D4EBE2404DFB0868433E825BD
+               CustomProperties Pin (PinId=60F7E126404F2A5E5FA2C4814371FDA1,PinName="OutputDelegate",Direction="EGPD_Output",PinType.PinCategory="delegate",PinType.PinSubCategory="",PinType.PinSubCategoryObject=None,PinType.PinSubCategoryMemberReference=(MemberParent=BlueprintGeneratedClass'/Game/Blueprints/GearVRMCTouchpadPawn.GearVRMCTouchpadPawn_C',MemberName="InpAxisKeyEvt_MotionController_Right_Thumbstick_X_K2Node_InputAxisKeyEvent_1"),PinType.PinValueType=(),PinType.ContainerType=None,PinType.bIsMap=False,PinType.bIsSet=False,PinType.bIsArray=False,PinType.bIsReference=False,PinType.bIsConst=False,PinType.bIsWeakPointer=False,PersistentGuid=00000000000000000000000000000000,bHidden=False,bNotConnectable=False,bDefaultValueIsReadOnly=False,bDefaultValueIsIgnored=False,bAdvancedView=False,bOrphanedPin=False,)
+               CustomProperties Pin (PinId=079E818244554159C73D02B16AEA1A62,PinName="then",Direction="EGPD_Output",PinType.PinCategory="exec",PinType.PinSubCategory="",PinType.PinSubCategoryObject=None,PinType.PinSubCategoryMemberReference=(),PinType.PinValueType=(),PinType.ContainerType=None,PinType.bIsMap=False,PinType.bIsSet=False,PinType.bIsArray=False,PinType.bIsReference=False,PinType.bIsConst=False,PinType.bIsWeakPointer=False,LinkedTo=(K2Node_VariableSet_18 03D9E61948784EA9BE735BBD6E306A7B,),PersistentGuid=00000000000000000000000000000000,bHidden=False,bNotConnectable=False,bDefaultValueIsReadOnly=False,bDefaultValueIsIgnored=False,bAdvancedView=False,bOrphanedPin=False,)
+               CustomProperties Pin (PinId=18A653DB4F75BEA3C627FDA82DFC64C4,PinName="AxisValue",PinToolTip="Axis Value\nFloat",Direction="EGPD_Output",PinType.PinCategory="float",PinType.PinSubCategory="",PinType.PinSubCategoryObject=None,PinType.PinSubCategoryMemberReference=(),PinType.PinValueType=(),PinType.ContainerType=None,PinType.bIsMap=False,PinType.bIsSet=False,PinType.bIsArray=False,PinType.bIsReference=False,PinType.bIsConst=False,PinType.bIsWeakPointer=False,DefaultValue="0.0",AutogeneratedDefaultValue="0.0",LinkedTo=(K2Node_CallFunction_252 F80FC4A04331E969D71EB5A87073C49B,),PersistentGuid=00000000000000000000000000000000,bHidden=False,bNotConnectable=False,bDefaultValueIsReadOnly=False,bDefaultValueIsIgnored=False,bAdvancedView=False,bOrphanedPin=False,)
+            End Object
+            Begin Object Class=/Script/BlueprintGraph.K2Node_InputAxisKeyEvent Name="K2Node_InputAxisKeyEvent_5"
+               AxisKey=MotionController_Right_Thumbstick_Y
+               EventReference=(MemberParent=Package'"/Script/Engine"')
+               CustomFunctionName="InpAxisKeyEvt_MotionController_Right_Thumbstick_Y_K2Node_InputAxisKeyEvent_5"
+               NodePosX=1056
+               NodePosY=-1232
+               NodeGuid=A96A6DEA4416BDF475691B8F6AC9B9BD
+               CustomProperties Pin (PinId=EEFA527949A99112C5E5CAA51BF20403,PinName="OutputDelegate",Direction="EGPD_Output",PinType.PinCategory="delegate",PinType.PinSubCategory="",PinType.PinSubCategoryObject=None,PinType.PinSubCategoryMemberReference=(MemberParent=BlueprintGeneratedClass'/Game/Blueprints/GearVRMCTouchpadPawn.GearVRMCTouchpadPawn_C',MemberName="InpAxisKeyEvt_MotionController_Right_Thumbstick_Y_K2Node_InputAxisKeyEvent_5"),PinType.PinValueType=(),PinType.ContainerType=None,PinType.bIsMap=False,PinType.bIsSet=False,PinType.bIsArray=False,PinType.bIsReference=False,PinType.bIsConst=False,PinType.bIsWeakPointer=False,PersistentGuid=00000000000000000000000000000000,bHidden=False,bNotConnectable=False,bDefaultValueIsReadOnly=False,bDefaultValueIsIgnored=False,bAdvancedView=False,bOrphanedPin=False,)
+               CustomProperties Pin (PinId=894C3754429D21D1F00E788913AFE23C,PinName="then",Direction="EGPD_Output",PinType.PinCategory="exec",PinType.PinSubCategory="",PinType.PinSubCategoryObject=None,PinType.PinSubCategoryMemberReference=(),PinType.PinValueType=(),PinType.ContainerType=None,PinType.bIsMap=False,PinType.bIsSet=False,PinType.bIsArray=False,PinType.bIsReference=False,PinType.bIsConst=False,PinType.bIsWeakPointer=False,LinkedTo=(K2Node_VariableSet_279 95093EE4473B034373ADC2AB5589E475,),PersistentGuid=00000000000000000000000000000000,bHidden=False,bNotConnectable=False,bDefaultValueIsReadOnly=False,bDefaultValueIsIgnored=False,bAdvancedView=False,bOrphanedPin=False,)
+               CustomProperties Pin (PinId=6383B3D84A6F15F39902AFA9E9362DBC,PinName="AxisValue",PinToolTip="Axis Value\nFloat",Direction="EGPD_Output",PinType.PinCategory="float",PinType.PinSubCategory="",PinType.PinSubCategoryObject=None,PinType.PinSubCategoryMemberReference=(),PinType.PinValueType=(),PinType.ContainerType=None,PinType.bIsMap=False,PinType.bIsSet=False,PinType.bIsArray=False,PinType.bIsReference=False,PinType.bIsConst=False,PinType.bIsWeakPointer=False,DefaultValue="0.0",AutogeneratedDefaultValue="0.0",LinkedTo=(K2Node_CallFunction_4809 F80FC4A04331E969D71EB5A87073C49B,),PersistentGuid=00000000000000000000000000000000,bHidden=False,bNotConnectable=False,bDefaultValueIsReadOnly=False,bDefaultValueIsIgnored=False,bAdvancedView=False,bOrphanedPin=False,)
+            End Object
+            Begin Object Class=/Script/UnrealEd.EdGraphNode_Comment Name="EdGraphNode_Comment_2466"
+               CommentDepth=-2
+               NodePosX=1040
+               NodePosY=-1312
+               NodeWidth=1264
+               NodeHeight=416
+               NodeComment="Displays the MC Y input position"
+               NodeGuid=E64D757A466A56A5059A6681264E8F63
+            End Object
+            Begin Object Class=/Script/BlueprintGraph.K2Node_CallFunction Name="K2Node_CallFunction_227"
+               FunctionReference=(MemberParent=Class'/Script/Engine.TextRenderComponent',MemberName="K2_SetText")
+               NodePosX=2048
+               NodePosY=-1696
+               NodeGuid=DFB1778C419CE3AD03BFBA900A00F408
+               CustomProperties Pin (PinId=96D7193F444562BF1B9D1A82C15146E1,PinName="execute",PinToolTip="\nExec",PinType.PinCategory="exec",PinType.PinSubCategory="",PinType.PinSubCategoryObject=None,PinType.PinSubCategoryMemberReference=(),PinType.PinValueType=(),PinType.ContainerType=None,PinType.bIsMap=False,PinType.bIsSet=False,PinType.bIsArray=False,PinType.bIsReference=False,PinType.bIsConst=False,PinType.bIsWeakPointer=False,LinkedTo=(K2Node_VariableSet_18 AA86D7454ACF6604F9CC0BA5C09ED35E,),PersistentGuid=00000000000000000000000000000000,bHidden=False,bNotConnectable=False,bDefaultValueIsReadOnly=False,bDefaultValueIsIgnored=False,bAdvancedView=False,bOrphanedPin=False,)
+               CustomProperties Pin (PinId=66DCFA3744967675B795929682EA35F0,PinName="then",PinToolTip="\nExec",Direction="EGPD_Output",PinType.PinCategory="exec",PinType.PinSubCategory="",PinType.PinSubCategoryObject=None,PinType.PinSubCategoryMemberReference=(),PinType.PinValueType=(),PinType.ContainerType=None,PinType.bIsMap=False,PinType.bIsSet=False,PinType.bIsArray=False,PinType.bIsReference=False,PinType.bIsConst=False,PinType.bIsWeakPointer=False,PersistentGuid=00000000000000000000000000000000,bHidden=False,bNotConnectable=False,bDefaultValueIsReadOnly=False,bDefaultValueIsIgnored=False,bAdvancedView=False,bOrphanedPin=False,)
+               CustomProperties Pin (PinId=24D9A43A4CC5BB2E7963438C3FE88719,PinName="self",PinFriendlyName=NSLOCTEXT("K2Node", "Target", "Target"),PinToolTip="Target\nText Render Component Object Reference",PinType.PinCategory="object",PinType.PinSubCategory="",PinType.PinSubCategoryObject=Class'/Script/Engine.TextRenderComponent',PinType.PinSubCategoryMemberReference=(),PinType.PinValueType=(),PinType.ContainerType=None,PinType.bIsMap=False,PinType.bIsSet=False,PinType.bIsArray=False,PinType.bIsReference=False,PinType.bIsConst=False,PinType.bIsWeakPointer=False,LinkedTo=(K2Node_VariableGet_57 16E9007D45B6E55C549AE1A9E153817A,),PersistentGuid=00000000000000000000000000000000,bHidden=False,bNotConnectable=False,bDefaultValueIsReadOnly=False,bDefaultValueIsIgnored=False,bAdvancedView=False,bOrphanedPin=False,)
+               CustomProperties Pin (PinId=FD05F54B48F429BC42F882B6033BFA3B,PinName="Value",PinToolTip="Value\nText (by ref)",PinType.PinCategory="text",PinType.PinSubCategory="",PinType.PinSubCategoryObject=None,PinType.PinSubCategoryMemberReference=(),PinType.PinValueType=(),PinType.ContainerType=None,PinType.bIsMap=False,PinType.bIsSet=False,PinType.bIsArray=False,PinType.bIsReference=True,PinType.bIsConst=True,PinType.bIsWeakPointer=False,LinkedTo=(K2Node_CallFunction_251 7C3752734C56A4E232CD5E9FD7051C5F,),PersistentGuid=00000000000000000000000000000000,bHidden=False,bNotConnectable=False,bDefaultValueIsReadOnly=False,bDefaultValueIsIgnored=True,bAdvancedView=False,bOrphanedPin=False,)
+            End Object
+            Begin Object Class=/Script/BlueprintGraph.K2Node_VariableGet Name="K2Node_VariableGet_58"
+               VariableReference=(MemberName="XPositionText",MemberGuid=009C67F84C5646E19F766381539F84C2,bSelfContext=True)
+               NodePosX=1728
+               NodePosY=-1392
+               NodeGuid=F69726B1414F379E5616C388F93FDF20
+               CustomProperties Pin (PinId=C2ABF17C4E4C3F86E7C64FB1135A50A7,PinName="XPositionText",Direction="EGPD_Output",PinType.PinCategory="text",PinType.PinSubCategory="",PinType.PinSubCategoryObject=None,PinType.PinSubCategoryMemberReference=(),PinType.PinValueType=(),PinType.ContainerType=None,PinType.bIsMap=False,PinType.bIsSet=False,PinType.bIsArray=False,PinType.bIsReference=False,PinType.bIsConst=False,PinType.bIsWeakPointer=False,LinkedTo=(K2Node_CallFunction_3063 AFE9F16E4462B3B3867621A9A37B537D,),PersistentGuid=00000000000000000000000000000000,bHidden=False,bNotConnectable=False,bDefaultValueIsReadOnly=False,bDefaultValueIsIgnored=False,bAdvancedView=False,bOrphanedPin=False,)
+               CustomProperties Pin (PinId=404B459841AC051FEE64FDBE286E724F,PinName="self",PinFriendlyName=NSLOCTEXT("K2Node", "Target", "Target"),PinType.PinCategory="object",PinType.PinSubCategory="",PinType.PinSubCategoryObject=BlueprintGeneratedClass'/Game/Blueprints/GearVRMCTouchpadPawn.GearVRMCTouchpadPawn_C',PinType.PinSubCategoryMemberReference=(),PinType.PinValueType=(),PinType.ContainerType=None,PinType.bIsMap=False,PinType.bIsSet=False,PinType.bIsArray=False,PinType.bIsReference=False,PinType.bIsConst=False,PinType.bIsWeakPointer=False,PersistentGuid=00000000000000000000000000000000,bHidden=True,bNotConnectable=False,bDefaultValueIsReadOnly=False,bDefaultValueIsIgnored=False,bAdvancedView=False,bOrphanedPin=False,)
+            End Object
+            Begin Object Class=/Script/BlueprintGraph.K2Node_VariableSet Name="K2Node_VariableSet_18"
+               VariableReference=(MemberName="XPositionText",MemberGuid=009C67F84C5646E19F766381539F84C2,bSelfContext=True)
+               NodePosX=1568
+               NodePosY=-1664
+               NodeGuid=86A8FEE34E41D23A8C67799609E4D43D
+               CustomProperties Pin (PinId=03D9E61948784EA9BE735BBD6E306A7B,PinName="execute",PinType.PinCategory="exec",PinType.PinSubCategory="",PinType.PinSubCategoryObject=None,PinType.PinSubCategoryMemberReference=(),PinType.PinValueType=(),PinType.ContainerType=None,PinType.bIsMap=False,PinType.bIsSet=False,PinType.bIsArray=False,PinType.bIsReference=False,PinType.bIsConst=False,PinType.bIsWeakPointer=False,LinkedTo=(K2Node_InputAxisKeyEvent_22 079E818244554159C73D02B16AEA1A62,),PersistentGuid=00000000000000000000000000000000,bHidden=False,bNotConnectable=False,bDefaultValueIsReadOnly=False,bDefaultValueIsIgnored=False,bAdvancedView=False,bOrphanedPin=False,)
+               CustomProperties Pin (PinId=AA86D7454ACF6604F9CC0BA5C09ED35E,PinName="then",Direction="EGPD_Output",PinType.PinCategory="exec",PinType.PinSubCategory="",PinType.PinSubCategoryObject=None,PinType.PinSubCategoryMemberReference=(),PinType.PinValueType=(),PinType.ContainerType=None,PinType.bIsMap=False,PinType.bIsSet=False,PinType.bIsArray=False,PinType.bIsReference=False,PinType.bIsConst=False,PinType.bIsWeakPointer=False,LinkedTo=(K2Node_CallFunction_227 96D7193F444562BF1B9D1A82C15146E1,),PersistentGuid=00000000000000000000000000000000,bHidden=False,bNotConnectable=False,bDefaultValueIsReadOnly=False,bDefaultValueIsIgnored=False,bAdvancedView=False,bOrphanedPin=False,)
+               CustomProperties Pin (PinId=8EBF99184B8C3FF82F0973B50E9A2831,PinName="XPositionText",PinType.PinCategory="text",PinType.PinSubCategory="",PinType.PinSubCategoryObject=None,PinType.PinSubCategoryMemberReference=(),PinType.PinValueType=(),PinType.ContainerType=None,PinType.bIsMap=False,PinType.bIsSet=False,PinType.bIsArray=False,PinType.bIsReference=False,PinType.bIsConst=False,PinType.bIsWeakPointer=False,LinkedTo=(K2Node_CallFunction_252 143CA0A3412175455799A096F47FDA4C,),PersistentGuid=00000000000000000000000000000000,bHidden=False,bNotConnectable=False,bDefaultValueIsReadOnly=False,bDefaultValueIsIgnored=False,bAdvancedView=False,bOrphanedPin=False,)
+               CustomProperties Pin (PinId=1013C0854372B5B2C4818B81798B6C77,PinName="Output_Get",PinToolTip="Retrieves the value of the variable, can use instead of a separate Get node",Direction="EGPD_Output",PinType.PinCategory="text",PinType.PinSubCategory="",PinType.PinSubCategoryObject=None,PinType.PinSubCategoryMemberReference=(),PinType.PinValueType=(),PinType.ContainerType=None,PinType.bIsMap=False,PinType.bIsSet=False,PinType.bIsArray=False,PinType.bIsReference=False,PinType.bIsConst=False,PinType.bIsWeakPointer=False,PersistentGuid=00000000000000000000000000000000,bHidden=False,bNotConnectable=False,bDefaultValueIsReadOnly=False,bDefaultValueIsIgnored=False,bAdvancedView=False,bOrphanedPin=False,)
+               CustomProperties Pin (PinId=1083FF874AB7737633CB3E89CB7D9F9E,PinName="self",PinFriendlyName=NSLOCTEXT("K2Node", "Target", "Target"),PinType.PinCategory="object",PinType.PinSubCategory="",PinType.PinSubCategoryObject=BlueprintGeneratedClass'/Game/Blueprints/GearVRMCTouchpadPawn.GearVRMCTouchpadPawn_C',PinType.PinSubCategoryMemberReference=(),PinType.PinValueType=(),PinType.ContainerType=None,PinType.bIsMap=False,PinType.bIsSet=False,PinType.bIsArray=False,PinType.bIsReference=False,PinType.bIsConst=False,PinType.bIsWeakPointer=False,PersistentGuid=00000000000000000000000000000000,bHidden=True,bNotConnectable=False,bDefaultValueIsReadOnly=False,bDefaultValueIsIgnored=False,bAdvancedView=False,bOrphanedPin=False,)
+            End Object
+            Begin Object Class=/Script/BlueprintGraph.K2Node_CallFunction Name="K2Node_CallFunction_4809"
+               bIsPureFunc=True
+               FunctionReference=(MemberParent=Class'/Script/Engine.KismetTextLibrary',MemberName="Conv_FloatToText")
+               NodePosX=1408
+               NodePosY=-1120
+               AdvancedPinDisplay=Hidden
+               NodeGuid=EE24529642B1B58323FC6B931A005A1F
+               CustomProperties Pin (PinId=9231A80F48229FB948A5AA9DA215DB97,PinName="self",PinFriendlyName=NSLOCTEXT("K2Node", "Target", "Target"),PinToolTip="Target\nKismet Text Library Object Reference",PinType.PinCategory="object",PinType.PinSubCategory="",PinType.PinSubCategoryObject=Class'/Script/Engine.KismetTextLibrary',PinType.PinSubCategoryMemberReference=(),PinType.PinValueType=(),PinType.ContainerType=None,PinType.bIsMap=False,PinType.bIsSet=False,PinType.bIsArray=False,PinType.bIsReference=False,PinType.bIsConst=False,PinType.bIsWeakPointer=False,DefaultObject="/Script/Engine.Default__KismetTextLibrary",PersistentGuid=00000000000000000000000000000000,bHidden=True,bNotConnectable=False,bDefaultValueIsReadOnly=False,bDefaultValueIsIgnored=False,bAdvancedView=False,bOrphanedPin=False,)
+               CustomProperties Pin (PinId=F80FC4A04331E969D71EB5A87073C49B,PinName="Value",PinToolTip="Value\nFloat",PinType.PinCategory="float",PinType.PinSubCategory="",PinType.PinSubCategoryObject=None,PinType.PinSubCategoryMemberReference=(),PinType.PinValueType=(),PinType.ContainerType=None,PinType.bIsMap=False,PinType.bIsSet=False,PinType.bIsArray=False,PinType.bIsReference=False,PinType.bIsConst=False,PinType.bIsWeakPointer=False,DefaultValue="0.0",AutogeneratedDefaultValue="0.0",LinkedTo=(K2Node_InputAxisKeyEvent_5 6383B3D84A6F15F39902AFA9E9362DBC,),PersistentGuid=00000000000000000000000000000000,bHidden=False,bNotConnectable=False,bDefaultValueIsReadOnly=False,bDefaultValueIsIgnored=False,bAdvancedView=False,bOrphanedPin=False,)
+               CustomProperties Pin (PinId=059645DC43483A351736259868CB24DE,PinName="RoundingMode",PinToolTip="Rounding Mode\nERoundingMode Enum",PinType.PinCategory="byte",PinType.PinSubCategory="",PinType.PinSubCategoryObject=Enum'/Script/Engine.ERoundingMode',PinType.PinSubCategoryMemberReference=(),PinType.PinValueType=(),PinType.ContainerType=None,PinType.bIsMap=False,PinType.bIsSet=False,PinType.bIsArray=False,PinType.bIsReference=False,PinType.bIsConst=False,PinType.bIsWeakPointer=False,DefaultValue="HalfToEven",PersistentGuid=00000000000000000000000000000000,bHidden=False,bNotConnectable=False,bDefaultValueIsReadOnly=False,bDefaultValueIsIgnored=False,bAdvancedView=True,bOrphanedPin=False,)
+               CustomProperties Pin (PinId=806085F24E097F901F75A398643EC77C,PinName="bUseGrouping",PinToolTip="Use Grouping\nBoolean",PinType.PinCategory="bool",PinType.PinSubCategory="",PinType.PinSubCategoryObject=None,PinType.PinSubCategoryMemberReference=(),PinType.PinValueType=(),PinType.ContainerType=None,PinType.bIsMap=False,PinType.bIsSet=False,PinType.bIsArray=False,PinType.bIsReference=False,PinType.bIsConst=False,PinType.bIsWeakPointer=False,DefaultValue="true",AutogeneratedDefaultValue="true",PersistentGuid=00000000000000000000000000000000,bHidden=False,bNotConnectable=False,bDefaultValueIsReadOnly=False,bDefaultValueIsIgnored=False,bAdvancedView=True,bOrphanedPin=False,)
+               CustomProperties Pin (PinId=AC50E40147020DFEB11D9B9AF578E13F,PinName="MinimumIntegralDigits",PinToolTip="Minimum Integral Digits\nInteger",PinType.PinCategory="int",PinType.PinSubCategory="",PinType.PinSubCategoryObject=None,PinType.PinSubCategoryMemberReference=(),PinType.PinValueType=(),PinType.ContainerType=None,PinType.bIsMap=False,PinType.bIsSet=False,PinType.bIsArray=False,PinType.bIsReference=False,PinType.bIsConst=False,PinType.bIsWeakPointer=False,DefaultValue="1",AutogeneratedDefaultValue="1",PersistentGuid=00000000000000000000000000000000,bHidden=False,bNotConnectable=False,bDefaultValueIsReadOnly=False,bDefaultValueIsIgnored=False,bAdvancedView=True,bOrphanedPin=False,)
+               CustomProperties Pin (PinId=5113623F4E1915E2B21972942DE375D4,PinName="MaximumIntegralDigits",PinToolTip="Maximum Integral Digits\nInteger",PinType.PinCategory="int",PinType.PinSubCategory="",PinType.PinSubCategoryObject=None,PinType.PinSubCategoryMemberReference=(),PinType.PinValueType=(),PinType.ContainerType=None,PinType.bIsMap=False,PinType.bIsSet=False,PinType.bIsArray=False,PinType.bIsReference=False,PinType.bIsConst=False,PinType.bIsWeakPointer=False,DefaultValue="324",AutogeneratedDefaultValue="324",PersistentGuid=00000000000000000000000000000000,bHidden=False,bNotConnectable=False,bDefaultValueIsReadOnly=False,bDefaultValueIsIgnored=False,bAdvancedView=True,bOrphanedPin=False,)
+               CustomProperties Pin (PinId=C03DF2CA437834ED7B2A8C8FFD834DD9,PinName="MinimumFractionalDigits",PinToolTip="Minimum Fractional Digits\nInteger",PinType.PinCategory="int",PinType.PinSubCategory="",PinType.PinSubCategoryObject=None,PinType.PinSubCategoryMemberReference=(),PinType.PinValueType=(),PinType.ContainerType=None,PinType.bIsMap=False,PinType.bIsSet=False,PinType.bIsArray=False,PinType.bIsReference=False,PinType.bIsConst=False,PinType.bIsWeakPointer=False,DefaultValue="0",AutogeneratedDefaultValue="0",PersistentGuid=00000000000000000000000000000000,bHidden=False,bNotConnectable=False,bDefaultValueIsReadOnly=False,bDefaultValueIsIgnored=False,bAdvancedView=True,bOrphanedPin=False,)
+               CustomProperties Pin (PinId=623B2E184048722C576F60A8736CBCC6,PinName="MaximumFractionalDigits",PinToolTip="Maximum Fractional Digits\nInteger",PinType.PinCategory="int",PinType.PinSubCategory="",PinType.PinSubCategoryObject=None,PinType.PinSubCategoryMemberReference=(),PinType.PinValueType=(),PinType.ContainerType=None,PinType.bIsMap=False,PinType.bIsSet=False,PinType.bIsArray=False,PinType.bIsReference=False,PinType.bIsConst=False,PinType.bIsWeakPointer=False,DefaultValue="3",AutogeneratedDefaultValue="3",PersistentGuid=00000000000000000000000000000000,bHidden=False,bNotConnectable=False,bDefaultValueIsReadOnly=False,bDefaultValueIsIgnored=False,bAdvancedView=True,bOrphanedPin=False,)
+               CustomProperties Pin (PinId=143CA0A3412175455799A096F47FDA4C,PinName="ReturnValue",PinToolTip="Return Value\nText",Direction="EGPD_Output",PinType.PinCategory="text",PinType.PinSubCategory="",PinType.PinSubCategoryObject=None,PinType.PinSubCategoryMemberReference=(),PinType.PinValueType=(),PinType.ContainerType=None,PinType.bIsMap=False,PinType.bIsSet=False,PinType.bIsArray=False,PinType.bIsReference=False,PinType.bIsConst=False,PinType.bIsWeakPointer=False,LinkedTo=(K2Node_VariableSet_279 7C0D8F8A48E5D1D7FDA9338F9FB261B2,),PersistentGuid=00000000000000000000000000000000,bHidden=False,bNotConnectable=False,bDefaultValueIsReadOnly=False,bDefaultValueIsIgnored=False,bAdvancedView=False,bOrphanedPin=False,)
+            End Object
+            Begin Object Class=/Script/BlueprintGraph.K2Node_CommutativeAssociativeBinaryOperator Name="K2Node_CommutativeAssociativeBinaryOperator_53"
+               bIsPureFunc=True
+               FunctionReference=(MemberParent=Class'/Script/Engine.KismetStringLibrary',MemberName="Concat_StrStr")
+               NodePosX=1728
+               NodePosY=-1488
+               NodeGuid=75D7418E415576AB084472A5AD2B5B5D
+               CustomProperties Pin (PinId=032F0C204576AD0073097FA4C783B85A,PinName="self",PinFriendlyName=NSLOCTEXT("K2Node", "Target", "Target"),PinToolTip="Target\nKismet String Library Object Reference",PinType.PinCategory="object",PinType.PinSubCategory="",PinType.PinSubCategoryObject=Class'/Script/Engine.KismetStringLibrary',PinType.PinSubCategoryMemberReference=(),PinType.PinValueType=(),PinType.ContainerType=None,PinType.bIsMap=False,PinType.bIsSet=False,PinType.bIsArray=False,PinType.bIsReference=False,PinType.bIsConst=False,PinType.bIsWeakPointer=False,DefaultObject="/Script/Engine.Default__KismetStringLibrary",PersistentGuid=00000000000000000000000000000000,bHidden=True,bNotConnectable=False,bDefaultValueIsReadOnly=False,bDefaultValueIsIgnored=False,bAdvancedView=False,bOrphanedPin=False,)
+               CustomProperties Pin (PinId=D4A956EF48D9C39C100216B65110C2D3,PinName="A",PinToolTip="A\nString\n\nThe original string",PinType.PinCategory="string",PinType.PinSubCategory="",PinType.PinSubCategoryObject=None,PinType.PinSubCategoryMemberReference=(),PinType.PinValueType=(),PinType.ContainerType=None,PinType.bIsMap=False,PinType.bIsSet=False,PinType.bIsArray=False,PinType.bIsReference=False,PinType.bIsConst=False,PinType.bIsWeakPointer=False,DefaultValue="X:",PersistentGuid=00000000000000000000000000000000,bHidden=False,bNotConnectable=False,bDefaultValueIsReadOnly=False,bDefaultValueIsIgnored=False,bAdvancedView=False,bOrphanedPin=False,)
+               CustomProperties Pin (PinId=9A9C2749448F164BFB5DB5BF6E3C5674,PinName="B",PinToolTip="B\nString\n\nThe string to append to A",PinType.PinCategory="string",PinType.PinSubCategory="",PinType.PinSubCategoryObject=None,PinType.PinSubCategoryMemberReference=(),PinType.PinValueType=(),PinType.ContainerType=None,PinType.bIsMap=False,PinType.bIsSet=False,PinType.bIsArray=False,PinType.bIsReference=False,PinType.bIsConst=False,PinType.bIsWeakPointer=False,LinkedTo=(K2Node_CallFunction_3063 3B4908FD4CA0E8DABAFDBABB39F6913D,),PersistentGuid=00000000000000000000000000000000,bHidden=False,bNotConnectable=False,bDefaultValueIsReadOnly=False,bDefaultValueIsIgnored=False,bAdvancedView=False,bOrphanedPin=False,)
+               CustomProperties Pin (PinId=3C2548244EFB0C45430FE294FCFAAD0B,PinName="ReturnValue",PinToolTip="Return Value\nString\n\nA new string which is the concatenation of A+B",Direction="EGPD_Output",PinType.PinCategory="string",PinType.PinSubCategory="",PinType.PinSubCategoryObject=None,PinType.PinSubCategoryMemberReference=(),PinType.PinValueType=(),PinType.ContainerType=None,PinType.bIsMap=False,PinType.bIsSet=False,PinType.bIsArray=False,PinType.bIsReference=False,PinType.bIsConst=False,PinType.bIsWeakPointer=False,LinkedTo=(K2Node_CallFunction_251 5E00EF9D41203D0B84183E8F49046194,),PersistentGuid=00000000000000000000000000000000,bHidden=False,bNotConnectable=False,bDefaultValueIsReadOnly=False,bDefaultValueIsIgnored=False,bAdvancedView=False,bOrphanedPin=False,)
+            End Object
+            Begin Object Class=/Script/BlueprintGraph.K2Node_CallFunction Name="K2Node_CallFunction_3063"
+               bIsPureFunc=True
+               FunctionReference=(MemberParent=Class'/Script/Engine.KismetTextLibrary',MemberName="Conv_TextToString")
+               NodePosX=1872
+               NodePosY=-1392
+               NodeGuid=12F411BA49DB1D3764A3EEB841C98DE9
+               CustomProperties Pin (PinId=F816D6154DCAF62B6E5843B1A535A56E,PinName="self",PinFriendlyName=NSLOCTEXT("K2Node", "Target", "Target"),PinToolTip="Target\nKismet Text Library Object Reference",PinType.PinCategory="object",PinType.PinSubCategory="",PinType.PinSubCategoryObject=Class'/Script/Engine.KismetTextLibrary',PinType.PinSubCategoryMemberReference=(),PinType.PinValueType=(),PinType.ContainerType=None,PinType.bIsMap=False,PinType.bIsSet=False,PinType.bIsArray=False,PinType.bIsReference=False,PinType.bIsConst=False,PinType.bIsWeakPointer=False,DefaultObject="/Script/Engine.Default__KismetTextLibrary",PersistentGuid=00000000000000000000000000000000,bHidden=True,bNotConnectable=False,bDefaultValueIsReadOnly=False,bDefaultValueIsIgnored=False,bAdvancedView=False,bOrphanedPin=False,)
+               CustomProperties Pin (PinId=AFE9F16E4462B3B3867621A9A37B537D,PinName="InText",PinToolTip="In Text\nText (by ref)",PinType.PinCategory="text",PinType.PinSubCategory="",PinType.PinSubCategoryObject=None,PinType.PinSubCategoryMemberReference=(),PinType.PinValueType=(),PinType.ContainerType=None,PinType.bIsMap=False,PinType.bIsSet=False,PinType.bIsArray=False,PinType.bIsReference=True,PinType.bIsConst=True,PinType.bIsWeakPointer=False,LinkedTo=(K2Node_VariableGet_58 C2ABF17C4E4C3F86E7C64FB1135A50A7,),PersistentGuid=00000000000000000000000000000000,bHidden=False,bNotConnectable=False,bDefaultValueIsReadOnly=False,bDefaultValueIsIgnored=True,bAdvancedView=False,bOrphanedPin=False,)
+               CustomProperties Pin (PinId=3B4908FD4CA0E8DABAFDBABB39F6913D,PinName="ReturnValue",PinToolTip="Return Value\nString",Direction="EGPD_Output",PinType.PinCategory="string",PinType.PinSubCategory="",PinType.PinSubCategoryObject=None,PinType.PinSubCategoryMemberReference=(),PinType.PinValueType=(),PinType.ContainerType=None,PinType.bIsMap=False,PinType.bIsSet=False,PinType.bIsArray=False,PinType.bIsReference=False,PinType.bIsConst=False,PinType.bIsWeakPointer=False,LinkedTo=(K2Node_CommutativeAssociativeBinaryOperator_53 9A9C2749448F164BFB5DB5BF6E3C5674,),PersistentGuid=00000000000000000000000000000000,bHidden=False,bNotConnectable=False,bDefaultValueIsReadOnly=False,bDefaultValueIsIgnored=False,bAdvancedView=False,bOrphanedPin=False,)
+            End Object
+            Begin Object Class=/Script/BlueprintGraph.K2Node_CallFunction Name="K2Node_CallFunction_251"
+               bIsPureFunc=True
+               FunctionReference=(MemberParent=Class'/Script/Engine.KismetTextLibrary',MemberName="Conv_StringToText")
+               NodePosX=1968
+               NodePosY=-1488
+               NodeGuid=B740B6074979D4BE8D1281B1E51B0023
+               CustomProperties Pin (PinId=671E8245468EDCE29C389F86DBDDBA2C,PinName="self",PinFriendlyName=NSLOCTEXT("K2Node", "Target", "Target"),PinToolTip="Target\nKismet Text Library Object Reference",PinType.PinCategory="object",PinType.PinSubCategory="",PinType.PinSubCategoryObject=Class'/Script/Engine.KismetTextLibrary',PinType.PinSubCategoryMemberReference=(),PinType.PinValueType=(),PinType.ContainerType=None,PinType.bIsMap=False,PinType.bIsSet=False,PinType.bIsArray=False,PinType.bIsReference=False,PinType.bIsConst=False,PinType.bIsWeakPointer=False,DefaultObject="/Script/Engine.Default__KismetTextLibrary",PersistentGuid=00000000000000000000000000000000,bHidden=True,bNotConnectable=False,bDefaultValueIsReadOnly=False,bDefaultValueIsIgnored=False,bAdvancedView=False,bOrphanedPin=False,)
+               CustomProperties Pin (PinId=5E00EF9D41203D0B84183E8F49046194,PinName="InString",PinToolTip="In String\nString",PinType.PinCategory="string",PinType.PinSubCategory="",PinType.PinSubCategoryObject=None,PinType.PinSubCategoryMemberReference=(),PinType.PinValueType=(),PinType.ContainerType=None,PinType.bIsMap=False,PinType.bIsSet=False,PinType.bIsArray=False,PinType.bIsReference=False,PinType.bIsConst=False,PinType.bIsWeakPointer=False,LinkedTo=(K2Node_CommutativeAssociativeBinaryOperator_53 3C2548244EFB0C45430FE294FCFAAD0B,),PersistentGuid=00000000000000000000000000000000,bHidden=False,bNotConnectable=False,bDefaultValueIsReadOnly=False,bDefaultValueIsIgnored=False,bAdvancedView=False,bOrphanedPin=False,)
+               CustomProperties Pin (PinId=7C3752734C56A4E232CD5E9FD7051C5F,PinName="ReturnValue",PinToolTip="Return Value\nText",Direction="EGPD_Output",PinType.PinCategory="text",PinType.PinSubCategory="",PinType.PinSubCategoryObject=None,PinType.PinSubCategoryMemberReference=(),PinType.PinValueType=(),PinType.ContainerType=None,PinType.bIsMap=False,PinType.bIsSet=False,PinType.bIsArray=False,PinType.bIsReference=False,PinType.bIsConst=False,PinType.bIsWeakPointer=False,LinkedTo=(K2Node_CallFunction_227 FD05F54B48F429BC42F882B6033BFA3B,),PersistentGuid=00000000000000000000000000000000,bHidden=False,bNotConnectable=False,bDefaultValueIsReadOnly=False,bDefaultValueIsIgnored=False,bAdvancedView=False,bOrphanedPin=False,)
+            End Object
+            Begin Object Class=/Script/BlueprintGraph.K2Node_CommutativeAssociativeBinaryOperator Name="K2Node_CommutativeAssociativeBinaryOperator_672"
+               bIsPureFunc=True
+               FunctionReference=(MemberParent=Class'/Script/Engine.KismetStringLibrary',MemberName="Concat_StrStr")
+               NodePosX=1792
+               NodePosY=-1072
+               NodeGuid=EF525F5E4D3C12848283AB81165D8470
+               CustomProperties Pin (PinId=032F0C204576AD0073097FA4C783B85A,PinName="self",PinFriendlyName=NSLOCTEXT("K2Node", "Target", "Target"),PinToolTip="Target\nKismet String Library Object Reference",PinType.PinCategory="object",PinType.PinSubCategory="",PinType.PinSubCategoryObject=Class'/Script/Engine.KismetStringLibrary',PinType.PinSubCategoryMemberReference=(),PinType.PinValueType=(),PinType.ContainerType=None,PinType.bIsMap=False,PinType.bIsSet=False,PinType.bIsArray=False,PinType.bIsReference=False,PinType.bIsConst=False,PinType.bIsWeakPointer=False,DefaultObject="/Script/Engine.Default__KismetStringLibrary",PersistentGuid=00000000000000000000000000000000,bHidden=True,bNotConnectable=False,bDefaultValueIsReadOnly=False,bDefaultValueIsIgnored=False,bAdvancedView=False,bOrphanedPin=False,)
+               CustomProperties Pin (PinId=D4A956EF48D9C39C100216B65110C2D3,PinName="A",PinToolTip="A\nString\n\nThe original string",PinType.PinCategory="string",PinType.PinSubCategory="",PinType.PinSubCategoryObject=None,PinType.PinSubCategoryMemberReference=(),PinType.PinValueType=(),PinType.ContainerType=None,PinType.bIsMap=False,PinType.bIsSet=False,PinType.bIsArray=False,PinType.bIsReference=False,PinType.bIsConst=False,PinType.bIsWeakPointer=False,DefaultValue="Y:",PersistentGuid=00000000000000000000000000000000,bHidden=False,bNotConnectable=False,bDefaultValueIsReadOnly=False,bDefaultValueIsIgnored=False,bAdvancedView=False,bOrphanedPin=False,)
+               CustomProperties Pin (PinId=9A9C2749448F164BFB5DB5BF6E3C5674,PinName="B",PinToolTip="B\nString\n\nThe string to append to A",PinType.PinCategory="string",PinType.PinSubCategory="",PinType.PinSubCategoryObject=None,PinType.PinSubCategoryMemberReference=(),PinType.PinValueType=(),PinType.ContainerType=None,PinType.bIsMap=False,PinType.bIsSet=False,PinType.bIsArray=False,PinType.bIsReference=False,PinType.bIsConst=False,PinType.bIsWeakPointer=False,LinkedTo=(K2Node_CallFunction_3068 79C243824DB73E447881D9A9C0ECC9AD,),PersistentGuid=00000000000000000000000000000000,bHidden=False,bNotConnectable=False,bDefaultValueIsReadOnly=False,bDefaultValueIsIgnored=False,bAdvancedView=False,bOrphanedPin=False,)
+               CustomProperties Pin (PinId=3C2548244EFB0C45430FE294FCFAAD0B,PinName="ReturnValue",PinToolTip="Return Value\nString\n\nA new string which is the concatenation of A+B",Direction="EGPD_Output",PinType.PinCategory="string",PinType.PinSubCategory="",PinType.PinSubCategoryObject=None,PinType.PinSubCategoryMemberReference=(),PinType.PinValueType=(),PinType.ContainerType=None,PinType.bIsMap=False,PinType.bIsSet=False,PinType.bIsArray=False,PinType.bIsReference=False,PinType.bIsConst=False,PinType.bIsWeakPointer=False,LinkedTo=(K2Node_CallFunction_3072 617B30A740DC062966C6A4B56C102445,),PersistentGuid=00000000000000000000000000000000,bHidden=False,bNotConnectable=False,bDefaultValueIsReadOnly=False,bDefaultValueIsIgnored=False,bAdvancedView=False,bOrphanedPin=False,)
+            End Object
+            Begin Object Class=/Script/BlueprintGraph.K2Node_VariableGet Name="K2Node_VariableGet_351"
+               VariableReference=(MemberName="YPositionText",MemberGuid=5ACC8CF14FA96A5E6BB8CABB29DF51C1,bSelfContext=True)
+               NodePosX=1696
+               NodePosY=-960
+               NodeGuid=F1D384594538D6AD83AEBFAE2EDBC3D9
+               CustomProperties Pin (PinId=A43A7ACF425D5C49575EFCA2B60BB43A,PinName="YPositionText",Direction="EGPD_Output",PinType.PinCategory="text",PinType.PinSubCategory="",PinType.PinSubCategoryObject=None,PinType.PinSubCategoryMemberReference=(),PinType.PinValueType=(),PinType.ContainerType=None,PinType.bIsMap=False,PinType.bIsSet=False,PinType.bIsArray=False,PinType.bIsReference=False,PinType.bIsConst=False,PinType.bIsWeakPointer=False,LinkedTo=(K2Node_CallFunction_3068 8ED83D284ACF2E1F233ABBB222118E16,),PersistentGuid=00000000000000000000000000000000,bHidden=False,bNotConnectable=False,bDefaultValueIsReadOnly=False,bDefaultValueIsIgnored=False,bAdvancedView=False,bOrphanedPin=False,)
+               CustomProperties Pin (PinId=5E65203C4CDAC93412F45C869886E0EB,PinName="self",PinFriendlyName=NSLOCTEXT("K2Node", "Target", "Target"),PinType.PinCategory="object",PinType.PinSubCategory="",PinType.PinSubCategoryObject=BlueprintGeneratedClass'/Game/Blueprints/GearVRMCTouchpadPawn.GearVRMCTouchpadPawn_C',PinType.PinSubCategoryMemberReference=(),PinType.PinValueType=(),PinType.ContainerType=None,PinType.bIsMap=False,PinType.bIsSet=False,PinType.bIsArray=False,PinType.bIsReference=False,PinType.bIsConst=False,PinType.bIsWeakPointer=False,PersistentGuid=00000000000000000000000000000000,bHidden=True,bNotConnectable=False,bDefaultValueIsReadOnly=False,bDefaultValueIsIgnored=False,bAdvancedView=False,bOrphanedPin=False,)
+            End Object
+            Begin Object Class=/Script/BlueprintGraph.K2Node_CallFunction Name="K2Node_CallFunction_3068"
+               bIsPureFunc=True
+               FunctionReference=(MemberParent=Class'/Script/Engine.KismetTextLibrary',MemberName="Conv_TextToString")
+               NodePosX=1840
+               NodePosY=-960
+               NodeGuid=2B2804C14589E1C0EE3A2188C58E90A6
+               CustomProperties Pin (PinId=E0D158374015974F5189399911C3980F,PinName="self",PinFriendlyName=NSLOCTEXT("K2Node", "Target", "Target"),PinToolTip="Target\nKismet Text Library Object Reference",PinType.PinCategory="object",PinType.PinSubCategory="",PinType.PinSubCategoryObject=Class'/Script/Engine.KismetTextLibrary',PinType.PinSubCategoryMemberReference=(),PinType.PinValueType=(),PinType.ContainerType=None,PinType.bIsMap=False,PinType.bIsSet=False,PinType.bIsArray=False,PinType.bIsReference=False,PinType.bIsConst=False,PinType.bIsWeakPointer=False,DefaultObject="/Script/Engine.Default__KismetTextLibrary",PersistentGuid=00000000000000000000000000000000,bHidden=True,bNotConnectable=False,bDefaultValueIsReadOnly=False,bDefaultValueIsIgnored=False,bAdvancedView=False,bOrphanedPin=False,)
+               CustomProperties Pin (PinId=8ED83D284ACF2E1F233ABBB222118E16,PinName="InText",PinToolTip="In Text\nText (by ref)",PinType.PinCategory="text",PinType.PinSubCategory="",PinType.PinSubCategoryObject=None,PinType.PinSubCategoryMemberReference=(),PinType.PinValueType=(),PinType.ContainerType=None,PinType.bIsMap=False,PinType.bIsSet=False,PinType.bIsArray=False,PinType.bIsReference=True,PinType.bIsConst=True,PinType.bIsWeakPointer=False,LinkedTo=(K2Node_VariableGet_351 A43A7ACF425D5C49575EFCA2B60BB43A,),PersistentGuid=00000000000000000000000000000000,bHidden=False,bNotConnectable=False,bDefaultValueIsReadOnly=False,bDefaultValueIsIgnored=True,bAdvancedView=False,bOrphanedPin=False,)
+               CustomProperties Pin (PinId=79C243824DB73E447881D9A9C0ECC9AD,PinName="ReturnValue",PinToolTip="Return Value\nString",Direction="EGPD_Output",PinType.PinCategory="string",PinType.PinSubCategory="",PinType.PinSubCategoryObject=None,PinType.PinSubCategoryMemberReference=(),PinType.PinValueType=(),PinType.ContainerType=None,PinType.bIsMap=False,PinType.bIsSet=False,PinType.bIsArray=False,PinType.bIsReference=False,PinType.bIsConst=False,PinType.bIsWeakPointer=False,LinkedTo=(K2Node_CommutativeAssociativeBinaryOperator_672 9A9C2749448F164BFB5DB5BF6E3C5674,),PersistentGuid=00000000000000000000000000000000,bHidden=False,bNotConnectable=False,bDefaultValueIsReadOnly=False,bDefaultValueIsIgnored=False,bAdvancedView=False,bOrphanedPin=False,)
+            End Object
+            Begin Object Class=/Script/BlueprintGraph.K2Node_VariableSet Name="K2Node_VariableSet_279"
+               VariableReference=(MemberName="YPositionText",MemberGuid=5ACC8CF14FA96A5E6BB8CABB29DF51C1,bSelfContext=True)
+               NodePosX=1584
+               NodePosY=-1216
+               NodeGuid=9989A9CF4F70FF6AEB741D9A1C4A26F6
+               CustomProperties Pin (PinId=95093EE4473B034373ADC2AB5589E475,PinName="execute",PinType.PinCategory="exec",PinType.PinSubCategory="",PinType.PinSubCategoryObject=None,PinType.PinSubCategoryMemberReference=(),PinType.PinValueType=(),PinType.ContainerType=None,PinType.bIsMap=False,PinType.bIsSet=False,PinType.bIsArray=False,PinType.bIsReference=False,PinType.bIsConst=False,PinType.bIsWeakPointer=False,LinkedTo=(K2Node_InputAxisKeyEvent_5 894C3754429D21D1F00E788913AFE23C,),PersistentGuid=00000000000000000000000000000000,bHidden=False,bNotConnectable=False,bDefaultValueIsReadOnly=False,bDefaultValueIsIgnored=False,bAdvancedView=False,bOrphanedPin=False,)
+               CustomProperties Pin (PinId=9DDD056046E4F8A140B753A455C1EF25,PinName="then",Direction="EGPD_Output",PinType.PinCategory="exec",PinType.PinSubCategory="",PinType.PinSubCategoryObject=None,PinType.PinSubCategoryMemberReference=(),PinType.PinValueType=(),PinType.ContainerType=None,PinType.bIsMap=False,PinType.bIsSet=False,PinType.bIsArray=False,PinType.bIsReference=False,PinType.bIsConst=False,PinType.bIsWeakPointer=False,LinkedTo=(K2Node_CallFunction_12917 96D7193F444562BF1B9D1A82C15146E1,),PersistentGuid=00000000000000000000000000000000,bHidden=False,bNotConnectable=False,bDefaultValueIsReadOnly=False,bDefaultValueIsIgnored=False,bAdvancedView=False,bOrphanedPin=False,)
+               CustomProperties Pin (PinId=7C0D8F8A48E5D1D7FDA9338F9FB261B2,PinName="YPositionText",PinType.PinCategory="text",PinType.PinSubCategory="",PinType.PinSubCategoryObject=None,PinType.PinSubCategoryMemberReference=(),PinType.PinValueType=(),PinType.ContainerType=None,PinType.bIsMap=False,PinType.bIsSet=False,PinType.bIsArray=False,PinType.bIsReference=False,PinType.bIsConst=False,PinType.bIsWeakPointer=False,LinkedTo=(K2Node_CallFunction_4809 143CA0A3412175455799A096F47FDA4C,),PersistentGuid=00000000000000000000000000000000,bHidden=False,bNotConnectable=False,bDefaultValueIsReadOnly=False,bDefaultValueIsIgnored=False,bAdvancedView=False,bOrphanedPin=False,)
+               CustomProperties Pin (PinId=21E43796493470CC2D0323BDE05A12DB,PinName="Output_Get",PinToolTip="Retrieves the value of the variable, can use instead of a separate Get node",Direction="EGPD_Output",PinType.PinCategory="text",PinType.PinSubCategory="",PinType.PinSubCategoryObject=None,PinType.PinSubCategoryMemberReference=(),PinType.PinValueType=(),PinType.ContainerType=None,PinType.bIsMap=False,PinType.bIsSet=False,PinType.bIsArray=False,PinType.bIsReference=False,PinType.bIsConst=False,PinType.bIsWeakPointer=False,PersistentGuid=00000000000000000000000000000000,bHidden=False,bNotConnectable=False,bDefaultValueIsReadOnly=False,bDefaultValueIsIgnored=False,bAdvancedView=False,bOrphanedPin=False,)
+               CustomProperties Pin (PinId=21BBC1C84CB79B2A3BD7D9939AE0CE49,PinName="self",PinFriendlyName=NSLOCTEXT("K2Node", "Target", "Target"),PinType.PinCategory="object",PinType.PinSubCategory="",PinType.PinSubCategoryObject=BlueprintGeneratedClass'/Game/Blueprints/GearVRMCTouchpadPawn.GearVRMCTouchpadPawn_C',PinType.PinSubCategoryMemberReference=(),PinType.PinValueType=(),PinType.ContainerType=None,PinType.bIsMap=False,PinType.bIsSet=False,PinType.bIsArray=False,PinType.bIsReference=False,PinType.bIsConst=False,PinType.bIsWeakPointer=False,PersistentGuid=00000000000000000000000000000000,bHidden=True,bNotConnectable=False,bDefaultValueIsReadOnly=False,bDefaultValueIsIgnored=False,bAdvancedView=False,bOrphanedPin=False,)
+            End Object
+            Begin Object Class=/Script/BlueprintGraph.K2Node_CallFunction Name="K2Node_CallFunction_252"
+               bIsPureFunc=True
+               FunctionReference=(MemberParent=Class'/Script/Engine.KismetTextLibrary',MemberName="Conv_FloatToText")
+               NodePosX=1392
+               NodePosY=-1552
+               AdvancedPinDisplay=Hidden
+               NodeGuid=756EE2EE4105D904CC2393A698FAB7D3
+               CustomProperties Pin (PinId=9231A80F48229FB948A5AA9DA215DB97,PinName="self",PinFriendlyName=NSLOCTEXT("K2Node", "Target", "Target"),PinToolTip="Target\nKismet Text Library Object Reference",PinType.PinCategory="object",PinType.PinSubCategory="",PinType.PinSubCategoryObject=Class'/Script/Engine.KismetTextLibrary',PinType.PinSubCategoryMemberReference=(),PinType.PinValueType=(),PinType.ContainerType=None,PinType.bIsMap=False,PinType.bIsSet=False,PinType.bIsArray=False,PinType.bIsReference=False,PinType.bIsConst=False,PinType.bIsWeakPointer=False,DefaultObject="/Script/Engine.Default__KismetTextLibrary",PersistentGuid=00000000000000000000000000000000,bHidden=True,bNotConnectable=False,bDefaultValueIsReadOnly=False,bDefaultValueIsIgnored=False,bAdvancedView=False,bOrphanedPin=False,)
+               CustomProperties Pin (PinId=F80FC4A04331E969D71EB5A87073C49B,PinName="Value",PinToolTip="Value\nFloat",PinType.PinCategory="float",PinType.PinSubCategory="",PinType.PinSubCategoryObject=None,PinType.PinSubCategoryMemberReference=(),PinType.PinValueType=(),PinType.ContainerType=None,PinType.bIsMap=False,PinType.bIsSet=False,PinType.bIsArray=False,PinType.bIsReference=False,PinType.bIsConst=False,PinType.bIsWeakPointer=False,DefaultValue="0.0",AutogeneratedDefaultValue="0.0",LinkedTo=(K2Node_InputAxisKeyEvent_22 18A653DB4F75BEA3C627FDA82DFC64C4,),PersistentGuid=00000000000000000000000000000000,bHidden=False,bNotConnectable=False,bDefaultValueIsReadOnly=False,bDefaultValueIsIgnored=False,bAdvancedView=False,bOrphanedPin=False,)
+               CustomProperties Pin (PinId=059645DC43483A351736259868CB24DE,PinName="RoundingMode",PinToolTip="Rounding Mode\nERoundingMode Enum",PinType.PinCategory="byte",PinType.PinSubCategory="",PinType.PinSubCategoryObject=Enum'/Script/Engine.ERoundingMode',PinType.PinSubCategoryMemberReference=(),PinType.PinValueType=(),PinType.ContainerType=None,PinType.bIsMap=False,PinType.bIsSet=False,PinType.bIsArray=False,PinType.bIsReference=False,PinType.bIsConst=False,PinType.bIsWeakPointer=False,DefaultValue="HalfToEven",PersistentGuid=00000000000000000000000000000000,bHidden=False,bNotConnectable=False,bDefaultValueIsReadOnly=False,bDefaultValueIsIgnored=False,bAdvancedView=True,bOrphanedPin=False,)
+               CustomProperties Pin (PinId=806085F24E097F901F75A398643EC77C,PinName="bUseGrouping",PinToolTip="Use Grouping\nBoolean",PinType.PinCategory="bool",PinType.PinSubCategory="",PinType.PinSubCategoryObject=None,PinType.PinSubCategoryMemberReference=(),PinType.PinValueType=(),PinType.ContainerType=None,PinType.bIsMap=False,PinType.bIsSet=False,PinType.bIsArray=False,PinType.bIsReference=False,PinType.bIsConst=False,PinType.bIsWeakPointer=False,DefaultValue="true",AutogeneratedDefaultValue="true",PersistentGuid=00000000000000000000000000000000,bHidden=False,bNotConnectable=False,bDefaultValueIsReadOnly=False,bDefaultValueIsIgnored=False,bAdvancedView=True,bOrphanedPin=False,)
+               CustomProperties Pin (PinId=AC50E40147020DFEB11D9B9AF578E13F,PinName="MinimumIntegralDigits",PinToolTip="Minimum Integral Digits\nInteger",PinType.PinCategory="int",PinType.PinSubCategory="",PinType.PinSubCategoryObject=None,PinType.PinSubCategoryMemberReference=(),PinType.PinValueType=(),PinType.ContainerType=None,PinType.bIsMap=False,PinType.bIsSet=False,PinType.bIsArray=False,PinType.bIsReference=False,PinType.bIsConst=False,PinType.bIsWeakPointer=False,DefaultValue="1",AutogeneratedDefaultValue="1",PersistentGuid=00000000000000000000000000000000,bHidden=False,bNotConnectable=False,bDefaultValueIsReadOnly=False,bDefaultValueIsIgnored=False,bAdvancedView=True,bOrphanedPin=False,)
+               CustomProperties Pin (PinId=5113623F4E1915E2B21972942DE375D4,PinName="MaximumIntegralDigits",PinToolTip="Maximum Integral Digits\nInteger",PinType.PinCategory="int",PinType.PinSubCategory="",PinType.PinSubCategoryObject=None,PinType.PinSubCategoryMemberReference=(),PinType.PinValueType=(),PinType.ContainerType=None,PinType.bIsMap=False,PinType.bIsSet=False,PinType.bIsArray=False,PinType.bIsReference=False,PinType.bIsConst=False,PinType.bIsWeakPointer=False,DefaultValue="324",AutogeneratedDefaultValue="324",PersistentGuid=00000000000000000000000000000000,bHidden=False,bNotConnectable=False,bDefaultValueIsReadOnly=False,bDefaultValueIsIgnored=False,bAdvancedView=True,bOrphanedPin=False,)
+               CustomProperties Pin (PinId=C03DF2CA437834ED7B2A8C8FFD834DD9,PinName="MinimumFractionalDigits",PinToolTip="Minimum Fractional Digits\nInteger",PinType.PinCategory="int",PinType.PinSubCategory="",PinType.PinSubCategoryObject=None,PinType.PinSubCategoryMemberReference=(),PinType.PinValueType=(),PinType.ContainerType=None,PinType.bIsMap=False,PinType.bIsSet=False,PinType.bIsArray=False,PinType.bIsReference=False,PinType.bIsConst=False,PinType.bIsWeakPointer=False,DefaultValue="0",AutogeneratedDefaultValue="0",PersistentGuid=00000000000000000000000000000000,bHidden=False,bNotConnectable=False,bDefaultValueIsReadOnly=False,bDefaultValueIsIgnored=False,bAdvancedView=True,bOrphanedPin=False,)
+               CustomProperties Pin (PinId=623B2E184048722C576F60A8736CBCC6,PinName="MaximumFractionalDigits",PinToolTip="Maximum Fractional Digits\nInteger",PinType.PinCategory="int",PinType.PinSubCategory="",PinType.PinSubCategoryObject=None,PinType.PinSubCategoryMemberReference=(),PinType.PinValueType=(),PinType.ContainerType=None,PinType.bIsMap=False,PinType.bIsSet=False,PinType.bIsArray=False,PinType.bIsReference=False,PinType.bIsConst=False,PinType.bIsWeakPointer=False,DefaultValue="3",AutogeneratedDefaultValue="3",PersistentGuid=00000000000000000000000000000000,bHidden=False,bNotConnectable=False,bDefaultValueIsReadOnly=False,bDefaultValueIsIgnored=False,bAdvancedView=True,bOrphanedPin=False,)
+               CustomProperties Pin (PinId=143CA0A3412175455799A096F47FDA4C,PinName="ReturnValue",PinToolTip="Return Value\nText",Direction="EGPD_Output",PinType.PinCategory="text",PinType.PinSubCategory="",PinType.PinSubCategoryObject=None,PinType.PinSubCategoryMemberReference=(),PinType.PinValueType=(),PinType.ContainerType=None,PinType.bIsMap=False,PinType.bIsSet=False,PinType.bIsArray=False,PinType.bIsReference=False,PinType.bIsConst=False,PinType.bIsWeakPointer=False,LinkedTo=(K2Node_VariableSet_18 8EBF99184B8C3FF82F0973B50E9A2831,),PersistentGuid=00000000000000000000000000000000,bHidden=False,bNotConnectable=False,bDefaultValueIsReadOnly=False,bDefaultValueIsIgnored=False,bAdvancedView=False,bOrphanedPin=False,)
+            End Object
+            Begin Object Class=/Script/BlueprintGraph.K2Node_CallFunction Name="K2Node_CallFunction_12917"
+               FunctionReference=(MemberParent=Class'/Script/Engine.TextRenderComponent',MemberName="K2_SetText")
+               NodePosX=1984
+               NodePosY=-1248
+               NodeGuid=4813618A4177816FC10DDABC4E097CED
+               CustomProperties Pin (PinId=96D7193F444562BF1B9D1A82C15146E1,PinName="execute",PinToolTip="\nExec",PinType.PinCategory="exec",PinType.PinSubCategory="",PinType.PinSubCategoryObject=None,PinType.PinSubCategoryMemberReference=(),PinType.PinValueType=(),PinType.ContainerType=None,PinType.bIsMap=False,PinType.bIsSet=False,PinType.bIsArray=False,PinType.bIsReference=False,PinType.bIsConst=False,PinType.bIsWeakPointer=False,LinkedTo=(K2Node_VariableSet_279 9DDD056046E4F8A140B753A455C1EF25,),PersistentGuid=00000000000000000000000000000000,bHidden=False,bNotConnectable=False,bDefaultValueIsReadOnly=False,bDefaultValueIsIgnored=False,bAdvancedView=False,bOrphanedPin=False,)
+               CustomProperties Pin (PinId=66DCFA3744967675B795929682EA35F0,PinName="then",PinToolTip="\nExec",Direction="EGPD_Output",PinType.PinCategory="exec",PinType.PinSubCategory="",PinType.PinSubCategoryObject=None,PinType.PinSubCategoryMemberReference=(),PinType.PinValueType=(),PinType.ContainerType=None,PinType.bIsMap=False,PinType.bIsSet=False,PinType.bIsArray=False,PinType.bIsReference=False,PinType.bIsConst=False,PinType.bIsWeakPointer=False,PersistentGuid=00000000000000000000000000000000,bHidden=False,bNotConnectable=False,bDefaultValueIsReadOnly=False,bDefaultValueIsIgnored=False,bAdvancedView=False,bOrphanedPin=False,)
+               CustomProperties Pin (PinId=24D9A43A4CC5BB2E7963438C3FE88719,PinName="self",PinFriendlyName=NSLOCTEXT("K2Node", "Target", "Target"),PinToolTip="Target\nText Render Component Object Reference",PinType.PinCategory="object",PinType.PinSubCategory="",PinType.PinSubCategoryObject=Class'/Script/Engine.TextRenderComponent',PinType.PinSubCategoryMemberReference=(),PinType.PinValueType=(),PinType.ContainerType=None,PinType.bIsMap=False,PinType.bIsSet=False,PinType.bIsArray=False,PinType.bIsReference=False,PinType.bIsConst=False,PinType.bIsWeakPointer=False,LinkedTo=(K2Node_VariableGet_173 21BEDF554F4C0BD9EE8FB9B35F4090FE,),PersistentGuid=00000000000000000000000000000000,bHidden=False,bNotConnectable=False,bDefaultValueIsReadOnly=False,bDefaultValueIsIgnored=False,bAdvancedView=False,bOrphanedPin=False,)
+               CustomProperties Pin (PinId=FD05F54B48F429BC42F882B6033BFA3B,PinName="Value",PinToolTip="Value\nText (by ref)",PinType.PinCategory="text",PinType.PinSubCategory="",PinType.PinSubCategoryObject=None,PinType.PinSubCategoryMemberReference=(),PinType.PinValueType=(),PinType.ContainerType=None,PinType.bIsMap=False,PinType.bIsSet=False,PinType.bIsArray=False,PinType.bIsReference=True,PinType.bIsConst=True,PinType.bIsWeakPointer=False,LinkedTo=(K2Node_CallFunction_3072 9FA298A74436C0E7966E1AA748279378,),PersistentGuid=00000000000000000000000000000000,bHidden=False,bNotConnectable=False,bDefaultValueIsReadOnly=False,bDefaultValueIsIgnored=True,bAdvancedView=False,bOrphanedPin=False,)
+            End Object
+            Begin Object Class=/Script/BlueprintGraph.K2Node_CallFunction Name="K2Node_CallFunction_3072"
+               bIsPureFunc=True
+               FunctionReference=(MemberParent=Class'/Script/Engine.KismetTextLibrary',MemberName="Conv_StringToText")
+               NodePosX=2000
+               NodePosY=-1088
+               NodeGuid=1C997BD54B6895BCFFD50F915E088428
+               CustomProperties Pin (PinId=C5C5A89844F52D5DF2884DA16208088B,PinName="self",PinFriendlyName=NSLOCTEXT("K2Node", "Target", "Target"),PinToolTip="Target\nKismet Text Library Object Reference",PinType.PinCategory="object",PinType.PinSubCategory="",PinType.PinSubCategoryObject=Class'/Script/Engine.KismetTextLibrary',PinType.PinSubCategoryMemberReference=(),PinType.PinValueType=(),PinType.ContainerType=None,PinType.bIsMap=False,PinType.bIsSet=False,PinType.bIsArray=False,PinType.bIsReference=False,PinType.bIsConst=False,PinType.bIsWeakPointer=False,DefaultObject="/Script/Engine.Default__KismetTextLibrary",PersistentGuid=00000000000000000000000000000000,bHidden=True,bNotConnectable=False,bDefaultValueIsReadOnly=False,bDefaultValueIsIgnored=False,bAdvancedView=False,bOrphanedPin=False,)
+               CustomProperties Pin (PinId=617B30A740DC062966C6A4B56C102445,PinName="InString",PinToolTip="In String\nString",PinType.PinCategory="string",PinType.PinSubCategory="",PinType.PinSubCategoryObject=None,PinType.PinSubCategoryMemberReference=(),PinType.PinValueType=(),PinType.ContainerType=None,PinType.bIsMap=False,PinType.bIsSet=False,PinType.bIsArray=False,PinType.bIsReference=False,PinType.bIsConst=False,PinType.bIsWeakPointer=False,LinkedTo=(K2Node_CommutativeAssociativeBinaryOperator_672 3C2548244EFB0C45430FE294FCFAAD0B,),PersistentGuid=00000000000000000000000000000000,bHidden=False,bNotConnectable=False,bDefaultValueIsReadOnly=False,bDefaultValueIsIgnored=False,bAdvancedView=False,bOrphanedPin=False,)
+               CustomProperties Pin (PinId=9FA298A74436C0E7966E1AA748279378,PinName="ReturnValue",PinToolTip="Return Value\nText",Direction="EGPD_Output",PinType.PinCategory="text",PinType.PinSubCategory="",PinType.PinSubCategoryObject=None,PinType.PinSubCategoryMemberReference=(),PinType.PinValueType=(),PinType.ContainerType=None,PinType.bIsMap=False,PinType.bIsSet=False,PinType.bIsArray=False,PinType.bIsReference=False,PinType.bIsConst=False,PinType.bIsWeakPointer=False,LinkedTo=(K2Node_CallFunction_12917 FD05F54B48F429BC42F882B6033BFA3B,),PersistentGuid=00000000000000000000000000000000,bHidden=False,bNotConnectable=False,bDefaultValueIsReadOnly=False,bDefaultValueIsIgnored=False,bAdvancedView=False,bOrphanedPin=False,)
+            End Object
+            Begin Object Class=/Script/UnrealEd.EdGraphNode_Comment Name="EdGraphNode_Comment_1724"
+               CommentDepth=-2
+               NodePosX=1040
+               NodePosY=-1744
+               NodeWidth=1264
+               NodeHeight=416
+               NodeComment="Displays the MC X input position"
+               NodeGuid=D6B77413401F104BDFA45E900089DC8C
+            End Object
+            Begin Object Class=/Script/BlueprintGraph.K2Node_VariableGet Name="K2Node_VariableGet_57"
+               VariableReference=(MemberName="X_Position",bSelfContext=True)
+               NodePosX=1824
+               NodePosY=-1600
+               NodeGuid=7499B5A04F510E8D45331481C21BF781
+               CustomProperties Pin (PinId=16E9007D45B6E55C549AE1A9E153817A,PinName="X_Position",Direction="EGPD_Output",PinType.PinCategory="object",PinType.PinSubCategory="",PinType.PinSubCategoryObject=Class'/Script/Engine.TextRenderComponent',PinType.PinSubCategoryMemberReference=(),PinType.PinValueType=(),PinType.ContainerType=None,PinType.bIsMap=False,PinType.bIsSet=False,PinType.bIsArray=False,PinType.bIsReference=False,PinType.bIsConst=False,PinType.bIsWeakPointer=False,LinkedTo=(K2Node_CallFunction_227 24D9A43A4CC5BB2E7963438C3FE88719,),PersistentGuid=00000000000000000000000000000000,bHidden=False,bNotConnectable=False,bDefaultValueIsReadOnly=False,bDefaultValueIsIgnored=False,bAdvancedView=False,bOrphanedPin=False,)
+               CustomProperties Pin (PinId=C611BC784A72C63BFCD9A0B050B0585D,PinName="self",PinFriendlyName=NSLOCTEXT("K2Node", "Target", "Target"),PinType.PinCategory="object",PinType.PinSubCategory="",PinType.PinSubCategoryObject=BlueprintGeneratedClass'/Game/Blueprints/GearVRMCTouchpadPawn.GearVRMCTouchpadPawn_C',PinType.PinSubCategoryMemberReference=(),PinType.PinValueType=(),PinType.ContainerType=None,PinType.bIsMap=False,PinType.bIsSet=False,PinType.bIsArray=False,PinType.bIsReference=False,PinType.bIsConst=False,PinType.bIsWeakPointer=False,PersistentGuid=00000000000000000000000000000000,bHidden=True,bNotConnectable=False,bDefaultValueIsReadOnly=False,bDefaultValueIsIgnored=False,bAdvancedView=False,bOrphanedPin=False,)
+            End Object
+            Begin Object Class=/Script/BlueprintGraph.K2Node_VariableGet Name="K2Node_VariableGet_173"
+               VariableReference=(MemberName="Y_Position",bSelfContext=True)
+               NodePosX=1776
+               NodePosY=-1168
+               NodeGuid=66676C7243BEC65168F037AA619FB9C6
+               CustomProperties Pin (PinId=21BEDF554F4C0BD9EE8FB9B35F4090FE,PinName="Y_Position",Direction="EGPD_Output",PinType.PinCategory="object",PinType.PinSubCategory="",PinType.PinSubCategoryObject=Class'/Script/Engine.TextRenderComponent',PinType.PinSubCategoryMemberReference=(),PinType.PinValueType=(),PinType.ContainerType=None,PinType.bIsMap=False,PinType.bIsSet=False,PinType.bIsArray=False,PinType.bIsReference=False,PinType.bIsConst=False,PinType.bIsWeakPointer=False,LinkedTo=(K2Node_CallFunction_12917 24D9A43A4CC5BB2E7963438C3FE88719,),PersistentGuid=00000000000000000000000000000000,bHidden=False,bNotConnectable=False,bDefaultValueIsReadOnly=False,bDefaultValueIsIgnored=False,bAdvancedView=False,bOrphanedPin=False,)
+               CustomProperties Pin (PinId=6CA0302F433B90DEC434B2BCC79A9E27,PinName="self",PinFriendlyName=NSLOCTEXT("K2Node", "Target", "Target"),PinType.PinCategory="object",PinType.PinSubCategory="",PinType.PinSubCategoryObject=BlueprintGeneratedClass'/Game/Blueprints/GearVRMCTouchpadPawn.GearVRMCTouchpadPawn_C',PinType.PinSubCategoryMemberReference=(),PinType.PinValueType=(),PinType.ContainerType=None,PinType.bIsMap=False,PinType.bIsSet=False,PinType.bIsArray=False,PinType.bIsReference=False,PinType.bIsConst=False,PinType.bIsWeakPointer=False,PersistentGuid=00000000000000000000000000000000,bHidden=True,bNotConnectable=False,bDefaultValueIsReadOnly=False,bDefaultValueIsIgnored=False,bAdvancedView=False,bOrphanedPin=False,)
+            End Object
+
+            END OBJECT
+    [/PARAM]
+[/OBJECT]
+[/EXCERPT:MCInput02]
+-->
+
+
+
+
