Copying //UE4/Dev-Documentation to //UE4/Samples-Main (Source: //UE4/Dev-Documentation @ 3623665)

diff --git a/Engine/Documentation/Source/Platforms/GearVR/QuickStart/2/GearVRQuickStart_2.INT.udn b/Engine/Documentation/Source/Platforms/GearVR/QuickStart/2/GearVRQuickStart_2.INT.udn
--- a/Engine/Documentation/Source/Platforms/GearVR/QuickStart/2/GearVRQuickStart_2.INT.udn
+++ b/Engine/Documentation/Source/Platforms/GearVR/QuickStart/2/GearVRQuickStart_2.INT.udn
@@ -1,8 +1,8 @@
 Availability:Public
 Title:2. Gear VR Project Setup
 Crumbs:%ROOT%, Platforms, Platforms/GearVR
-Description:How to setup your C++ based UE4 project to be used with the Gear VR Headset.
-Version: 4.16
+Description:How to setup your UE4 project to be used with the Gear VR Headset.
+Version: 4.17
 Skilllevel: Intermediate
 tags:VR
 tags:GearVR
@@ -13,107 +13,91 @@ Checkpoint:GearVRQS
 
 [nav]
 
+In the following part of the Gear VR Quick Start, we will go over setting up an Unreal Engine 4 (UE4) project so that it can be used with the Gear VR HMD.
+
 ## Steps
 
-1. Once the project has finished being created and compiled, Visual Studio will automatically open (if already installed) along with the editor.
-
-	![](GVR_Editor_Opened.png)
-
-	[region:note]
-	If you are using Unreal Engine 4.8 or later and a Blueprint based project, Visual Studio will not open up as it is not needed. 
-	[/region]
-
-1. Once the UE4 Editor has loaded, from the **Main Toolbar**, open the **Windows** tab and select the **Plugins** menu. 
-	
-	[OBJECT:EmbeddedVideo]
-		[PARAMLITERAL:videoid]
-		eo79kj2i2bc
-		[/PARAMLITERAL]
-		[PARAMLITERAL:width]
-		100
-		[/PARAMLITERAL]
-		[PARAMLITERAL:height]
-		56.25
-		[/PARAMLITERAL]
-		[PARAMLITERAL:units]
-		%
-		[/PARAMLITERAL]
-		[PARAMLITERAL:type]
-		simple
-		[/PARAMLITERAL]
-	[/OBJECT]
-
-1. From the **Plugins** menu, go to the **Head Mounted Displays** section, **disable** the **Oculus** option, and **enable** the **Gear VR** option. Then, click the **Restart Now** button to re-start the editor so the changes can take place. 
-
-	![](GVR_Plugin_Options_Toggle.png)
-
-1. Once the Editor has re-started, select **Edit** from the **Main Toolbar**, select **Project Settings** and click the **Configure Now** button for Android support in the **APKPacking** section.
-
-	[OBJECT:EmbeddedVideo]
-		[PARAMLITERAL:videoid]
-		If_vSzD5h3w
-		[/PARAMLITERAL]
-		[PARAMLITERAL:width]
-		100
-		[/PARAMLITERAL]
-		[PARAMLITERAL:height]
-		56.25
-		[/PARAMLITERAL]
-		[PARAMLITERAL:units]
-		%
-		[/PARAMLITERAL]
-		[PARAMLITERAL:type]
-		simple
-		[/PARAMLITERAL]
-	[/OBJECT]
-
-1. Next, change the **Minimum SDK Version** and the **Target SDK Version** from **9** to **19**. 
-
-	![](GVR_Android_SDK_Version.png)
-
-1. Next, scroll down till you see the **Configure the Android Manifest for deployment to Gear VR** checkbox and enable it. 
-
-	[Region:note]
-	This step is of vital importance as without this on your project will not be able to run on the Gear VR Headset.
-	[/region]
-
-	![](GVR_Deployment_Checkbox.png)
-
-1. Next, under the **Android SDK** section, fill-in the location of the required Android SDK files on your development PC. 
-	You can find the needed folders in the location where you installed the Tegra files. 
-	
-	[Region:note]
-	If you have already installed CodeWorks for Android 1R6, this information should be already filled in for you. The only reason that you should change something here is if the information is missing or pointing to the wrong location.
-	[/region]
-
-	![](GVR_Android_SDK_Setup.png)
-
-1. Finally, under **File** > **Package Project** > **Packing Settings**, un-check **Full Rebuild** to help speed up iteration times. 
-
-	[Region:note]
-	Remember to **Re-Enable** this option when your project is ready for its final publish to ensure that your project will work as intended on the Gear VR Headset. 
-	[/region]
-
-	[OBJECT:EmbeddedVideo]
-		[PARAMLITERAL:videoid]
-		29oB5aowElQ
-		[/PARAMLITERAL]
-		[PARAMLITERAL:width]
-		100
-		[/PARAMLITERAL]
-		[PARAMLITERAL:height]
-		56.25
-		[/PARAMLITERAL]
-		[PARAMLITERAL:units]
-		%
-		[/PARAMLITERAL]
-		[PARAMLITERAL:type]
-		simple
-		[/PARAMLITERAL]
-	[/OBJECT]
+1. Once the UE4 Editor has loaded, go to the **Main Toolbar** and select **Edit** > **Plugins** to open up the Plugin Menu. From the Plugin Menu, choose the **Virtuality Reality** section and make sure that the **OculusVR** Plugin is enabled.
+
+    [REGION:lightbox]
+    [![](GVR_Plugin_Options_Toggle.png)(w:703)](GVR_Plugin_Options_Toggle.png)
+    [/REGION]
+
+    [REGION:caption]
+    Click for full image.
+    [/REGION]
+
+    [region:note]
+    In previous versions of UE4 Gear VR had its own plugin that needed to be enabled. If you are not able to use Gear VR with the version of UE4 you are using, make sure to check the VR plugins section and enable the Gear VR plugin if it is available.
+    [/REGION]
+
+1. For any UE4 project to run when it is deployed to a device, you will need to have at least one level saved. To keep the current level go to the **Toolbar** and press the **Save** icon to save the current level, giving it a name of **GearVRQSMap**.
+
+    [REGION:lightbox]
+    [![](GVR_QS_Saved_Map.png)(w:299)](GVR_QS_Saved_Map.png)
+    [/REGION]
+
+    [REGION:caption]
+    Click for full image.
+    [/REGION]
+
+1. Next, from the Main Menu go to, **Edit** > **Project Settings** and then under **Project** in the **Maps & Modes** section under **Default Maps** input the map that was just saved in the **Editor Start Up** and **Game Default Map**.
+
+    [REGION:lightbox]
+    [![](GVR_QS_Default_Maps.png)(w:546)](GVR_QS_Default_Maps.png)
+    [/REGION]
+
+    [REGION:caption]
+    Click for full image.
+    [/REGION]
+
+1. Now, go to **Engine** > **Input** and in the **Mobile** section, set the **Default Touch Interface** to **None** by clicking on the Default Touch Interface input box and selecting **None** from the drop down.
+
+    [REGION:lightbox]
+    [![](GVR_QS_Remove_Touch.png)(w:546)](GVR_QS_Remove_Touch.png)
+    [/REGION]
+
+    [REGION:caption]
+    Click for full image.
+    [/REGION]
+
+1. Then in **Platforms** under **Android** go to the **ApkPackaging** section and press the **Configure Now** button then set the following settings:
+
+    [REGION:lightbox]
+    [![](GVR_Android_SDK_Version.png)(w:546)](GVR_Android_SDK_Version.png)
+    [/REGION]
+
+    [REGION:caption]
+    Click for full image.
+    [/REGION]
+
+    * **Minimum SDK Version:** 19
+    * **Target SDK Version:** 19
+    * **Enable FullScreen Immersive on KitKat and above devices:** True 
+
+1. Next, go to the **Advanced APKPackaging** section and enable the, **Configure the AndroidManifest for deployment to GearVR** check mark box.
+
+    [REGION:lightbox]
+    [![](GVR_Deployment_Checkbox.png)(w:546)](GVR_Deployment_Checkbox.png)
+    [/REGION]
+
+    [REGION:caption]
+    Click for full image.
+    [/REGION]
+
+1. Now, click on the **Android SDK** section and make sure that Android SDK tools that are needed to build your project point to the correct folder on your PC. 
+    
+    [REGION:lightbox]
+    [![](GVR_Android_SDK_Setup.png)(w:546)](GVR_Android_SDK_Setup.png)
+    [/REGION]
+
+    [REGION:caption]
+    Click for full image.
+    [/REGION]
 
 
 [nav]
 
 ## End Result
-When completed you will now have a UE4 project that is ready to be built, deployed and viewed on your Samsung Galaxy smartphone.
\ No newline at end of file
+
+When completed you will now have a UE4 project that is ready to be deployed and viewed with your Gear VR HMD.
\ No newline at end of file
diff --git a/Engine/Documentation/Source/Platforms/GearVR/QuickStart/3/GearVRQuickStart_3.INT.udn b/Engine/Documentation/Source/Platforms/GearVR/QuickStart/3/GearVRQuickStart_3.INT.udn
--- a/Engine/Documentation/Source/Platforms/GearVR/QuickStart/3/GearVRQuickStart_3.INT.udn
+++ b/Engine/Documentation/Source/Platforms/GearVR/QuickStart/3/GearVRQuickStart_3.INT.udn
@@ -1,7 +1,7 @@
 Availability:Public
 Title:3. Gear VR Project Packaging & Deployment
 Crumbs:%ROOT%, Platforms/GearVR
-Description:Instructions on how to deploy projects form UE4 to use with the Samsung Gear VR Headset.
+Description:Instructions on how to deploy projects from UE4 to use with the Samsung Gear VR Headset.
 Version: 4.16
 Skilllevel: Intermediate
 tags:VR
@@ -16,43 +16,71 @@ Related: Platforms\GearVR\Debugging
 
 [nav]
 
+In the following part of the Gear VR Quick Start, we will go over how to launch your UE4 project so that it can be viewed with the Gear VR HMD.
+
 ## Steps
 
-1. First make sure that your Samsung Galaxy Note 4 Smartphone is connected to a USB cord and plugged into a USB port on the PC you are developing on.
+1. First, make sure that your Samsung Galaxy Smartphone is connected to your development PC via a USB cable and that your PC can see the device using the **ADB** command **ADB Devices**.
 
-1. Next start the packaging process by going to the **Main Toolbar** and selecting **File** > **Package Project** > **Android** > **Android ETC 2**.
+1. From the Main Toolbar, click on the small white triangle next to the **Launch** icon, to open the **Options for Launching on a Device**.
 
-	![](GVR_Deploy_To_Android.png)
+    [REGION:lightbox]
+    [![](GVR_Deploy_To_Android.png)(w:550)](GVR_Deploy_To_Android.png)
+    [/REGION]
 
-1. Once you click on the **Android ETC 2** option, a Windows dialogue box will be displayed. From this dialogue box, select or create a folder to save your packaged project to and then press the **OK** button to begin the packaging process.
+    [REGION:caption]
+    Click for full image.
+    [/REGION]
 
-	![](GVR_Save_Project.png)
+1. From the Main Toolbar, click on the small white triangle next to the **Launch** icon, to open the **Options for Launching on a Device**.
 
-1. Depending on the size of your project, the packaging process could take a few minutes or a few hours to complete. 
-	When it is completed the folder that you selected/created to save the build to should now have three files in it and should look similar to the image below. 
+    [REGION:lightbox]
+    [![](GVR_Deployed_Files.png)(w:550)](GVR_Deployed_Files.png)
+    [/REGION]
 
-	![](GVR_Deployed_Files.png)
+    [REGION:caption]
+    Click for full image.
+    [/REGION]
 
-1. Next, run the .bat file which is located in the folder that the build was packaged to by **Double-clicking** on it with your **Left Mouse Button**.
-	Once you run the .bat file, a command line window will be displayed that shows the current progress of the deployment to your Smartphone.
-	When your project has been successfully deployed to your Smartphone, the command line window will shut itself down.
-	
-	![](GVR_Run_Bat_File.png)
+    [region:note]
+    If prompted to save, make sure that you do even if you just saved before to ensure that the latest content makes it onto your device.
+    [/region]
 
-1. Once you run the .bat file, a command line window will be displayed that will show you the deployment progress as well as any errors that may have occurred. 
+1. Check in the lower corner of the UE4 Editor; a progress window will be displayed that will tell you what stage the packaging is currently on. If this is the first time you are packing your project, this could take a few minutes depending on the size of your project. 
 
-	![](GVR_Deployment_Window.png)
+    [REGION:lightbox]
+    [![](GVR_Deployment_Window.png)(w:550)](GVR_Deployment_Window.png)
+    [/REGION]
 
-	[region:note]
-	The command line window that is created when running the .bat file will close down on it's own when it is completed. Closing this window manually could cause an issue with your deployment and should be avoided. 
-	[/region]
+    [REGION:caption]
+    Click for full image.
+    [/REGION]
 
-1. Finally, after the command line window has closed, check your Smartphone for your newly created application which should use the same name used for your UE4 C++ project that it was created from.
-	Launch the application and then place your Smartphone into the Gear VR Headset.
+    [region:note]
+    To see detailed information about what is currently happing during the packing process, click on the **Show Output Log** link in the progress window.
+    [/region]
 
-	![](GVR_On_Phone.png)
+1. When the packing and installing process has completed, a notification will be displayed, and a sound will be played. You can then put your device in your Gear VR HMD and when you put it on you should something similar to the following video.
 
-[nav]
+    [OBJECT:EmbeddedVideo]
+        [PARAMLITERAL:videoid]
+        Bq3Ov32Bwy8
+        [/PARAMLITERAL]
+        [PARAMLITERAL:width]
+        100
+        [/PARAMLITERAL]
+        [PARAMLITERAL:height]
+        56.25
+        [/PARAMLITERAL]
+        [PARAMLITERAL:units]
+        %
+        [/PARAMLITERAL]
+        [PARAMLITERAL:type]
+        simple
+        [/PARAMLITERAL]
+    [/OBJECT]
 
 ## End Result
 When completed you will be able to launch the project and view it using your  Samsung Galaxy smartphone and the Gear VR HMD.
+
+[nav]
\ No newline at end of file
diff --git a/Engine/Documentation/Source/Platforms/GearVR/QuickStart/4/GearVRQuickStart_4.INT.udn b/Engine/Documentation/Source/Platforms/GearVR/QuickStart/4/GearVRQuickStart_4.INT.udn
--- a/Engine/Documentation/Source/Platforms/GearVR/QuickStart/4/GearVRQuickStart_4.INT.udn
+++ b/Engine/Documentation/Source/Platforms/GearVR/QuickStart/4/GearVRQuickStart_4.INT.udn
@@ -3,7 +3,7 @@ Title:4. On Your Own!
 Crumbs:%ROOT%, Platforms/GearVR
 Description:Using the methods provided during this guide, polish up your level with additional content and read up on related topics. 
 Navigation:platform
-Version: 4.16
+Version: 4.17
 Skilllevel: Intermediate
 tags:VR
 tags:GearVR
@@ -14,13 +14,13 @@ Checkpoint:GearVRQS
 
 [nav]
 
-Now that you can deploy a UE4 project to your Android based smartphone and view it with your Gear VR HMD try adding the following items.
+Now that you can deploy a UE4 project and view it with your Gear VR HMD, try and do the following with your project.
 
 * Use the [Migrate Tool](Engine/Content/Browser/UserGuide/Migrate/) to move content from one of the mobile content examples into your project for use.
 
 * Adjust the settings of the [Default Device Profiles](Platforms/Mobile/Performance/TipsAndTricks) to ensure that your project will run as fast as possible when deployed to your Android based smartphone.
 
-* Use the [GPU Profiler](Engine/Performance/GPU/) to help track your projects performance while you are building it.
+* Use the [GPU Profiler](Engine/Performance/GPU/) to help track your project's performance while you are building it.
 
 
 
diff --git a/Engine/Documentation/Source/Platforms/GearVR/QuickStart/GearVRQuickStart.INT.udn b/Engine/Documentation/Source/Platforms/GearVR/QuickStart/GearVRQuickStart.INT.udn
--- a/Engine/Documentation/Source/Platforms/GearVR/QuickStart/GearVRQuickStart.INT.udn
+++ b/Engine/Documentation/Source/Platforms/GearVR/QuickStart/GearVRQuickStart.INT.udn
@@ -1,25 +1,26 @@
 Availability:Public
 Title:Samsung Gear VR UE4 Quick Start
-Crumbs:%ROOT%, Platforms/GearVR
+Crumbs:%ROOT%
 Description:Instructions on how to quickly setup and deploy projects for use with the Samsung Gear VR Headset.
 Navigation:platform
 Platform:Gear VR
 Skilllevel: Intermediate
-Version:4.9
+Version:4.17
 type:quick start
 type:landing
 parent:Platforms/GearVR
 tags:VR
 tags:GearVR
-order:1
+order:2
 Checkpoint:GearVRQS
+prereq:Platforms/Android/GettingStarted/
 prereq:Platforms\GearVR\Prerequisites
+Topic-image:Platforms\GearVR\SamsungGearVRDevelopment_topic.png
 
 ![](GVR_Headset_00.png)
 
-
 ## Goals
-The Gear VR Quick Start will walk you through how to setup your development PC and Samsung Galaxy Note 4 Smartphone to work with the Samsung Gear VR headset.
+The Gear VR Quick Start will walk you through how to setup your development PC and Samsung Smartphone to work with the Gear VR headset.
 
 ## Objectives
 * Create a new UE4 project that is targeted towards Gear VR development.
diff --git a/Engine/Documentation/Source/Platforms/GoogleVR/BestPractices/GoogleVRBestPractices.INT.udn b/Engine/Documentation/Source/Platforms/GoogleVR/BestPractices/GoogleVRBestPractices.INT.udn
new file mode 100644
--- /dev/null
+++ b/Engine/Documentation/Source/Platforms/GoogleVR/BestPractices/GoogleVRBestPractices.INT.udn
@@ -0,0 +1,150 @@
+Availability:Public
+Crumbs: %ROOT%
+Title: GoogleVR Best Practices
+Description: Information on how to get the most out of your GoogleVR UE4 project.
+Type: Reference
+Version: 4.17
+parent:Platforms/GoogleVR
+Order: 2
+tags:VR
+tags:Google VR
+Related: Platforms/Mobile/Performance
+Related: Engine/Performance
+Related: Engine/Performance/StatCommands
+Related: Platforms/Mobile/Textures
+Related: Platforms/Mobile/Materials
+Related: Engine/Rendering/Materials/MaterialProperties/BlendModes
+Related: Engine/Rendering/Materials/CustomizedUVs
+
+[TOC(start:2 end:3)]
+
+When developing content for Virtual Reality (VR) projects, certain steps must be taken to ensure that the user will have the best VR experience possible. In the following Best Practices guide we will cover all you need to know to help guarantee that the user has the best possible experience when viewing your project using GoogleVR.
+
+
+## GoogleVR Rendering Options
+While combining UE4 and GoogleVR can produce some very immersive VR experiences, there are some rendering features that UE4 can not perform while using GoogleVR due to hardware limitations. Many of the restrictions that you will encounter are similar to limitations for [mobile devices](Platforms/Mobile/Performance).
+However, because of the unique rendering requirements for GoogleVR, there are few extra things that must be considered.In the following sections, we will cover what these requirements are and how they can affect your GoogleVR project. 
+
+### Mobile HDR
+Mobile HDR will not work with Google VR, and it should be left disabled. 
+
+### Post Processing
+Due to the demanding rendering requirements for GoogleVR, many Post Processing features, unfortunately, will not work. Below you will find a list of all the Post Process features that will work with GoogelVR.
+
+|Post Processing Section |Feature Name| Description |
+| -------- | ---------- |------|
+| Auto Exposure|Exposure Bias |Can be used to increase or decrease the scenes brightness. |
+
+### Materials
+Creating and interacting with Materials for GoogleVR projects is very similar to creating and interacting with Materials for other projects. You still use the Material Editor to create Materials, and you can still use many of the Material nodes you find in the Material Editor. However, there are a few restrictions/considerations that must be adhered to when making Materials for GoogleVR to ensure that what you make will work with.
+
+| Material Input Name | Description |
+| -------- | ---------- |
+| **Base Color** |Base Color defines the overall color of the Material. |
+| **Roughness** |Roughness defines how rough or smooth the Material is. |
+| **Metallic** |Metallic defines how metal-like the Material is. |
+| **Specular** |Specular defines the intensity of the reflections. |
+| **Normal**  |Normal helps to add detail that would otherwise be too costly to model in. |
+| **Emissive**  |Emissive helps to make parts of Materials glow like they are emitting light. |
+| **Opacity**  |Opacity helps to define how to see through something is. |
+| **Opacity Mask** |Opacity Mask helps to define if a certain pixel is see through or not. |
+
+[REGION:note]
+Please keep in mind that the both Translucent and Masked Materials can be very expensive to render, especially on hardware like Smartphones. Whenever possible, you should use [Opaque Materials](Engine/Rendering/Materials/MaterialProperties/BlendModes) as they are much cheaper to render.
+[/region]
+
+### Material Shading Models
+Due to hardware limitations only the following **Material Shading Models** are available for use in projects that use the GoogleVR.
+| Shading Model Name | Description |
+| -------- | ---------- |
+| **Default Lit** |Default Lit is, as the name suggests, the default shading model and will likely be the one that you will use most often. This shading model makes use of direct and indirect lighting, as well as specularity for reflections.. |
+| **Unlit** |The Unlit Shading Model only outputs Emissive for color, making this perfect for special effects such as fire or illuminating objects.|
+
+[region:note]
+Try to use the Unlit Shading Model as much as possible with your GoogleVR projects as it is the cheapest to use. 
+[/region]
+
+### Material Blend Modes
+
+Material Blend Modes describe how the output of the current Material will blend over what is already being drawn in the background. 
+When creating content for the Googel VR, the following Material Blend Modes are available for use.
+
+| Blend Mode | Description |
+| ---- | ----- |
+| **Opaque** | For a Solid object, Opaque is the cheapest material. But, as its name implies, it has no support for translucency. That means it will never create overdraw, and rendering one pixel of the material will always take the same time depending on the complexity of the **Node Graph**. |
+| **Masked** |  If you need some pixels not drawn on your material, this is the way to go. The material only creates overdraw for the area that is masked out. But the material is 100\% opaque, or 100\% see through. The bigger the area that is masked out, the more overdraw is created. That means more expensive. |
+| **Translucent** | This is the most expensive **BlendMode**. It can look like a Masked and Opaque but will have an additional effect on performance regardless of actually being transparent or not. |
+| **Additive** | Is a cheaper alternative to Translucent. It does not need to sort the pixels, they just get added to the scene. But that means it can only brighten. Black pixels on the material do not affect the scene.|
+| **Modulated** | Is the second cheapest alternative to Translucent, but it multiplies its Pixels with the scene. That means it can only darken the scene.|
+[region:note]
+When using the Modulate blend mode make sure to disable **Separate Translucency** or your Material will not compile correctly.
+  ![](T_Separate_Trans.png)
+[/region]
+
+### Mobile Specific Material Settings
+On the **Main Material Node** under the **Details Panel** you will find a section called **Mobile**.
+The settings contained here are for helping to lessen Material rendering costs on mobile devices.  
+| Property | Description |
+| -------- | ---------- |
+| **Fully Rough** | When enabled, saves some shader ALU instructions and one sampler.  This effectively disables Specular while maintaining overall reflective energy. |
+| **Use Lightmap Directionality** | When disabled, lighting from light maps will be flat but cheaper. |
+
+### Customized UVs
+[Customized UVs](Engine/Rendering/Materials/CustomizedUVs) allow you to perform certain shader operations in vertex shaders instead of pixel shaders to help avoid dependent texture fetches. 
+This offers an excellent way to claw back much-needed performance when doing simple things like UV manipulation to make a texture title more than it currently does. 
+Whenever possible make sure to do all UV manipulation using Customized UVs especially when developing content for GoogleVR.
+
+[region:note]
+Customized UVs are currently only implemented for certain component types such as Static Meshes, Skeletal Meshes, BSP, Landscape, and Mesh Particles.
+[/region]
+
+### Reflections
+GoogleVR does not support Screen Space Reflections(SSR) which means you will only be able to generate reflections that come from [Reflection Capture Actors](Resources/ContentExamples/Reflections/1_4/) that are placed in levels. 
+Also, keep in mind that the reflections that are generated from the reflections actors will not be of very high fidelity due to hardware limitations. 
+
+## GoogleVR & Content Setup
+When creating art content for your GoogleVR projects, there are few things that you need to keep in mind to ensure that your project can run as smooth as possible. In the following section, we will cover what you need to be concerned with when making art content for GoogleVR.
+
+### Content Hard limits
+Due to certain hardware restrictions on mobile devices, the following restrictions must be adhered to when creating content. Failing to follow these restrictions will result in your project not working.
+
+* All mesh types can have up to but not over 65,000 vertices due to the lack of 32-bit index support on mobile hardware.
+* Skeletal Meshes can have up to but not over 75 bones.
+
+### Mesh Considerations
+When creating Meshes to use with your GoogleVR Project the following information needs to be kept in mind to ensure that what you are creating will work with GoogleVR. 
+
+* Use as few [Material IDs](Engine/Content/FBX/Materials) on 3D meshes as possible to save on draw calls.
+* Make sure to make use of [Static Mesh LODs](Engine/Content/FBX/StaticMeshes) whenever possible.
+* Keep an eye on your 3D meshes polygon count and try and keep it as [Low Poly](https://en.wikipedia.org/wiki/Low_poly) as possible.
+* Make sure that all static models have a 2nd UV set so they can make use of [Light Mapping](Engine/Content/Types/StaticMeshes/LightmapUnwrapping).
+* Whenever possible try and make use of [Imposter Sprites](Engine/Content/Tools/RenderToTextureTools/3) to replace far away Static Meshes as they are much cheaper to render.
+
+### Texture Considerations
+When creating Textures to use with your GoogleVR project the following information needs to be kept in mind to ensure that what you are creating will not be detrimental to your projects overall performance. 
+
+* Textures can not be bigger than 2048 pixels in either the X or Y dimension.
+* Textures sizes must be in a power of 2 (For example, 2, 4, 8, 16, 32, 64, 128, 256, 512, 1024, 2048)
+* Loading up to 512 MB of Textures into memory is possible with GoogleVR however it is recommended that you do not load more than 128 MB of Textures into memory at any given time.
+* Use square Textures whenever possible as they make a more efficient use of memory. 
+* Wherever possible try and use [Texture Atlasing](http://en.wikipedia.org/wiki/Texture_atlas) to reduce the number of Textures needed in memory.
+* [RGB Texture Masking / Packing](Engine/Rendering/Materials/HowTo/Masking) is an excellent way to cut down on unique Textures in memory.
+* Only **TC_Default** and **TC_NormalMap** compression settings should be used as they are the only two that will work on mobile platforms. 
+
+### Level Considerations
+Getting the maximum performance out of your levels on GoogleVR is a delicate balancing act that will require time to get just right. 
+The following suggestions/things to lookout for can help you quickly identify areas of your levels that might cause performance issues.
+
+* Always make sure that the level lighting has been [built](Engine/Rendering/LightingAndShadows/Lightmass) to ensure you get the best performance possible. If you see a red **Lighting needs to be rebuilt** error message on the screen, this means that your lighting needs to be rebuilt.
+* Make sure that all the lights in you level are set to [Static](Engine/Rendering/LightingAndShadows/LightMobility/StaticLights) as the other lighting options can be costly to render.
+* Check to see if your level is making use of [Precomputed Visibility](Engine/Actors/Volumes) and if not make sure to add one to the level and re-build lighting to enable it in the level's World Settings. 
+* Draw calls for the entire scene should not be greater than 100 for any view. This can be seen with [stat openglrhi](Engine/Performance/StatCommands) on the device, [stat d3d11rhi](Engine/Performance/StatCommands) in ES2 previewer on PC.
+* The total triangle count of the level should not be greater than 100,000 from any view. This can be seen with [stat openglrhi](Engine/Performance/StatCommands) on the device, [stat d3d11rhi](Engine/Performance/StatCommands) in ES2 previewer on PC.
+* Be very careful in how Translucent and Opacity Masked Materials are used in the level as they are costly to render.
+* Make sure to [Performance Profile](Engine/Performance) your level often to make sure you catch any performance issues sooner rather than later.
+
+
+
+
+
+
diff --git a/Engine/Documentation/Source/Platforms/GoogleVR/Daydream/Daydream.INT.udn b/Engine/Documentation/Source/Platforms/GoogleVR/Daydream/Daydream.INT.udn
--- a/Engine/Documentation/Source/Platforms/GoogleVR/Daydream/Daydream.INT.udn
+++ b/Engine/Documentation/Source/Platforms/GoogleVR/Daydream/Daydream.INT.udn
@@ -1,13 +1,12 @@
-Availability:Public
-Title:Google VR and Daydream Devices
-Crumbs:%ROOT%, Platforms/GoogleVR
-Description:Which Google VR features are specific to Daydream devices.
-Navigation:platform
-Skilllevel: Intermediate
-Version: 4.16
-parent:Platforms/GoogleVR
-type:overview
-tags:VR
+Availability: Public
+Crumbs: %ROOT%
+Title: Google VR and Daydream Devices
+Description: Which Google VR features are specific to Daydream devices.
+Type: Overview
+SkillLevel: Intermediate
+Version: 4.17
+Parent: Platforms/GoogleVR
+Order: 2
 tags:Google VR
 tags:Daydream
 topic-image:Platforms\GoogleVR\GoogleVR_Topic_Image.png
@@ -21,19 +20,19 @@ Google has a provided a very detailed document that goes over everything you nee
 
 While Google VR can work for almost any Android based device, there are certain features that will only work with Daydream devices. In the following section, we will take a look at these features and what you have to do to use them in your projects.
 
-## Downloading Google VR UE4 Branch
-To get all of the pre-built functionality Google Daydream VR offers you will need to download and compile a special version of UE4 that you will find linked below. This version of UE4 contains the needed code and assets that are required to display the Daydream Motion Controller in the manner Google wants it. 
+## Downloading Google Daydream for UE4
 
-* [Google VR Github Pull]([REDACTED_PROFILE_URL]/UnrealEngine/releases)
-
-Once you have downloaded this version of UE4, follow the directions for the [Downloading Unreal as a Zip from Github](GettingStarted/DownloadingUnrealEngine/DownloadZip) page to complete the rest of the setup process.
+Google Daydream now comes with the UE4 binary that you can download from the Epic Games launcher and is also available in the UE4 source that you download and compile from Github. 
+* [Epic Games Launcher](https://www.epicgames.com/unrealtournament/download)
+* [Epic Games Github](https://github.com/EpicGames)
 
 ## Android NDK 12b
-In order to access the full functionality of Daydream, you will need to download and install the Android NDK version 12b. 
 
-* [NDK 12b Downlaod](https://developer.android.com/ndk/downloads/older_releases.html)
+To access the full functionality of Daydream, you will need to make sure that your Android Works NDK version has been updated to 12b. For more information about installing Android Works, check out the link below. 
+
+* [Setting up Android Works](Platforms/Android/GettingStarted/1)
 
-Once downloaded and installed you will need to point UE4 to the new NDK files by going to the **Android SDK** > **SDKConfig** section and changing the NDK path to the new NDK 12b install location and set the SDK API level to **matchndk** and the NDK API Level to **android-24**.
+Once downloaded and installed you will need to make sure that your UE4 projects uses the new NDK files by going to the **Android SDK** > **SDKConfig** section and set the SDK API level to **matchndk** and the NDK API Level to **android-24**.
 
 [REGION:lightbox]
 [![](GDaydream_New_NDK_00.png)(w:502)](GDaydream_New_NDK_00.png)
@@ -43,93 +42,60 @@ Once downloaded and installed you will need to point UE4 to the new NDK files by
 Click for full image.
 [/REGION]
 
-## Scanline Racing
+## Using Scanline Racing
+
 **Scanline Racing** is a feature specific to Daydream devices that when enabled helps to reduce latency in your VR project. To enable Scanline Racing in your Google VR project for Daydream devices you will need to do the following.
 
 1. Open up your **Project Settings** and go to **Platforms** > **Android**. Under the **APKPackaging** section, set the following: 
 
-	* Minimum SDK Version: **24**
-	* Target SDK Version: **24**
+    * Minimum SDK Version: **24**
+    * Target SDK Version: **24**
 
-	[REGION:lightbox]
-	[![](GDaydream_Arm64_Support_00.png)(w:502)](GDaydream_Arm64_Support_00.png)
-	[/REGION]
+    [REGION:lightbox]
+    [![](GDaydream_Arm64_Support_00.png)(w:502)](GDaydream_Arm64_Support_00.png)
+    [/REGION]
 
-	[REGION:caption]
-	Click for full image.
-	[/REGION]
+    [REGION:caption]
+    Click for full image.
+    [/REGION]
 
 1. Under the **Build** section change the following:
-	* Support armv7(aka armeabi-v7a): **Disabled** 
-	* Support arm64 (aka arm64-v8a): **Enabled**
+    * Support armv7(aka armeabi-v7a): **Disabled** 
+    * Support arm64 (aka arm64-v8a): **Enabled**
 
-	[REGION:lightbox]
-	[![](GDaydream_Build_Support.png)(w:502)](GDaydream_Build_Support.png)
-	[/REGION]
+    [REGION:lightbox]
+    [![](GDaydream_Build_Support.png)(w:502)](GDaydream_Build_Support.png)
+    [/REGION]
 
-	[REGION:caption]
-	Click for full image.
-	[/REGION]
+    [REGION:caption]
+    Click for full image.
+    [/REGION]
 
     [region:note]
     Enabling this will add significant time to your first cook as the arm64 libraries need to be compiled.
     [/region]
 
 1. Then under the **Advanced APKPackaging** section and enable the following options.
-	* Configure GoogleVR Deployment Mode: **Daydream**
-	* Configure GoogleVR for sustained-performance mode (scanline racing): **Enabled**
+    * Configure GoogleVR Deployment Mode: **Daydream**
+    * Configure GoogleVR for sustained-performance mode (scanline racing): **Enabled**
 
-	[REGION:lightbox]
-	[![](GDaydream_Scanline_Racing_Support_00.png)(w:502)](GDaydream_Scanline_Racing_Support_00.png)
-	[/REGION]
+    [REGION:lightbox]
+    [![](GDaydream_Scanline_Racing_Support_00.png)(w:502)](GDaydream_Scanline_Racing_Support_00.png)
+    [/REGION]
 
-	[REGION:caption]
-	Click for full image.
-	[/REGION]
+    [REGION:caption]
+    Click for full image.
+    [/REGION]
 
 
 Then next time you cook your UE4 project and deploy it to your Daydream device low latency rendering will be enabled.
 
-## Motion Controllers
-Working with Motion Controllers for Daydream devices is much like working with Motion Controllers for the HTC Vive or Oculus Rift. To quickly add Motion Controller functionality to your UE4 Daydream project you need to do the following.
-
-[region:note]
-For a more in-depth look at how to use the Google Daydream Controller Motion Controllers in UE4, make sure to check out the [Offical Google Arm Model Documents](https://developers.google.com/vr/unreal/arm-model) page for more information on how to use the Arm Model.
-[/region]
-
-1. First, open up your project's Plugins menu and under the Input section, enable the **Google VR Motion Controller** plugin.
-
-	[REGION:lightbox]
-	[![](GDaydream_Plugins.png)(w:692)](GDaydream_Plugins.png)
-	[/REGION]
-
-	[REGION:caption]
-	Click for full image.
-	[/REGION]
-
-    [region:note]
-    Please note that to use the Plugin you will need to re-start your UE4 Editor. Failing to do so will result in you not being able to use the Google VR Motion Controller. 
-    [/region]
-
-1. Next, open up your projects Pawn and add the **GoogleVRMotionController** Component to the Components list. When completed your Pawn's setup should match the following image. 
+## Daydream Motion Controllers
 
-    ![](GDaydream_MC_Addation.png)
+One of the unique features of Daydream devices is their support of Motion Controllers. To see how you setup Motion Controllers for Daydream devices, check out the following. How - To.
 
-    [region:tip]
-    The Motion Controller component should use the same position and height of the VR Camera.
-    [/region]
-
-1. When you run your project your motion controllers should look similar to the following image.
-
-	[REGION:lightbox]
-	[![](unreal-controller-model.png)(w:724)](unreal-controller-model.png)
-	[/REGION]
-
-	[region:note]
-	If you are using the 4.14 binary build of UE4 the above content will be missing and you should instead use the standard Motion Controller component. All of this functionality is included with the 4.15 and later releases of UE4. If you need to this functionality now, you will have to download and compile the, [UE4.13 Google VR version]([REDACTED_PROFILE_URL]/UnrealEngine/tree/4.13-googlevr).
-	[/region]
+* [Google Daydream Motion Controllers](Platforms\GoogleVR\HowTo\MotionController)
 
-Now that you have a Google VR Motion Controller component on your Pawn, you can edit and adjust how the Motion Controller's function by searching for **Google VR** in the Blueprint Editor and using the nodes that you find. 
 
 ## Training Streams
 
@@ -143,12 +109,12 @@ These Training Streams help to provide some additional resources on setting up y
     [PARAMLITERAL:videoid]
     6ALnsdQnkVQ
     [/PARAMLITERAL]
-	[PARAMLITERAL:playlist]
-	PLZlv_N0_O1gY7G589Z3I5-Dz7AdFSIWaG
-	[/PARAMLITERAL]
-	[PARAMLITERAL:duration]
-	1:05:05
-	[/PARAMLITERAL]
+    [PARAMLITERAL:playlist]
+    PLZlv_N0_O1gY7G589Z3I5-Dz7AdFSIWaG
+    [/PARAMLITERAL]
+    [PARAMLITERAL:duration]
+    1:05:05
+    [/PARAMLITERAL]
     [PARAMLITERAL:active]
     [/PARAMLITERAL]
     [PARAMLITERAL:linktype]
@@ -166,13 +132,13 @@ These Training Streams help to provide some additional resources on setting up y
     [PARAMLITERAL:videoid]
     eRNtgFo6iU0
     [/PARAMLITERAL]
-	[PARAMLITERAL:playlist]
-	PLZlv_N0_O1gY7G589Z3I5-Dz7AdFSIWaG
-	[/PARAMLITERAL]
-	[PARAMLITERAL:duration]
-	59:50
-	[/PARAMLITERAL]
-	[PARAMLITERAL:active]
+    [PARAMLITERAL:playlist]
+    PLZlv_N0_O1gY7G589Z3I5-Dz7AdFSIWaG
+    [/PARAMLITERAL]
+    [PARAMLITERAL:duration]
+    59:50
+    [/PARAMLITERAL]
+    [PARAMLITERAL:active]
     [/PARAMLITERAL]
     [PARAMLITERAL:linktype]
     block
@@ -189,12 +155,12 @@ These Training Streams help to provide some additional resources on setting up y
     [PARAMLITERAL:videoid]
     utOahIZgKgc
     [/PARAMLITERAL]
-	[PARAMLITERAL:playlist]
-	PLZlv_N0_O1gY7G589Z3I5-Dz7AdFSIWaG
-	[/PARAMLITERAL]
-	[PARAMLITERAL:duration]
-	57:44
-	[/PARAMLITERAL]
+    [PARAMLITERAL:playlist]
+    PLZlv_N0_O1gY7G589Z3I5-Dz7AdFSIWaG
+    [/PARAMLITERAL]
+    [PARAMLITERAL:duration]
+    57:44
+    [/PARAMLITERAL]
     [PARAMLITERAL:active]
     [/PARAMLITERAL]
     [PARAMLITERAL:linktype]
@@ -211,12 +177,12 @@ These Training Streams help to provide some additional resources on setting up y
     [PARAMLITERAL:videoid]
     fcmRGkpWefY
     [/PARAMLITERAL]
-	[PARAMLITERAL:playlist]
-	PLZlv_N0_O1gY7G589Z3I5-Dz7AdFSIWaG
-	[/PARAMLITERAL]
-	[PARAMLITERAL:duration]
-	48:20
-	[/PARAMLITERAL]
+    [PARAMLITERAL:playlist]
+    PLZlv_N0_O1gY7G589Z3I5-Dz7AdFSIWaG
+    [/PARAMLITERAL]
+    [PARAMLITERAL:duration]
+    48:20
+    [/PARAMLITERAL]
     [PARAMLITERAL:active]
     [/PARAMLITERAL]
     [PARAMLITERAL:linktype]
diff --git a/Engine/Documentation/Source/Platforms/GoogleVR/Debugging/GoogleVRDebugging.INT.udn b/Engine/Documentation/Source/Platforms/GoogleVR/Debugging/GoogleVRDebugging.INT.udn
new file mode 100644
--- /dev/null
+++ b/Engine/Documentation/Source/Platforms/GoogleVR/Debugging/GoogleVRDebugging.INT.udn
@@ -0,0 +1,81 @@
+Availability:Public
+Title:Google VR Debugging
+Description:Information on how to debug applications for Google VR.
+Navigation:platform
+Crumbs:%ROOT%
+Platform:Google VR
+Related: Platforms/Mobile/Performance
+Related: Engine/Performance
+Related: Engine/Performance/StatCommands
+Related: Platforms/Mobile/Textures
+Related: Platforms/Mobile/Materials
+Related: Engine/Rendering/Materials/MaterialProperties/BlendModes
+Related: Engine/Rendering/Materials/CustomizedUVs
+Skilllevel: Intermediate
+Version: 4.17
+type:overview
+parent:Platforms/GoogleVR
+tags:VR
+tags:GoogleVR
+
+[TOC(start:1 end:3)]
+
+In the following document we will take a look at the various ways you can go about tracking down performance issues in your Google VR projects. 
+
+## Google Daydream Performance HUD
+
+
+![](GoogleDaydreamPerfHUD.png)
+
+[region:note]
+The Daydream Performance HUD is a performance overlay enabled for all Daydream applications built with Google VR SDK 1.60 and above.
+[/region]
+
+Google Daydream comes with a built in performance overlay that when enabled, will provide quick insight into performance issues  you might encounter. With this overlay enabled you are able to see the following stats:
+
+* **Application frame submit rate and whether frames were blocked on the GPU**
+* **Asynchronous Reprojection statistics**
+* **Process memory usage**
+* **How many degrees (°C) the device is from thermal throttling**
+
+### Enabling Performance HUD
+
+To enable the Performance HUD you will need to do the following.
+
+1. Open VR Settings, available by tapping the gear icon at the top right of the phone's screen after launching a Daydream application.
+
+	![](GVDP_00.png)
+
+1. Tap Build Version **six** times to make the Developer Options item appear.
+
+	[REGION:lightbox]
+	[![](GVDP_01.png)(w:320)](GVDP_01.png)
+	[/REGION]
+
+	[REGION:caption]
+	Click for full image.
+	[/REGION]
+
+1. Select Developer Options > Enable performance heads-up display. Enabling this option also implicitly enables the Enable performance monitoring option.
+
+	[REGION:lightbox]
+	[![](GVDP_02.png)(w:320)](GVDP_02.png)
+	[/REGION]
+
+	[REGION:caption]
+	Click for full image.
+	[/REGION]
+
+4. You should now see the performance overlay the next time you run a VR app like inthe following image.
+
+	[REGION:lightbox]
+	[![](GoogleDaydreamPerfHUD.png)(w:720)](GoogleDaydreamPerfHUD.png)
+	[/REGION]
+
+	[REGION:caption]
+	Click for full image.
+	[/REGION]
+
+	[region:note]
+	For more information on how to use Performance HUD, make sure to check out the offical [Google VR Performance HUD](https://developers.google.com/vr/tools/perfhud) for more detailed inforamtion on using Peformance HUD.
+	[/region]
diff --git a/Engine/Documentation/Source/Platforms/GoogleVR/GoogleVR.INT.udn b/Engine/Documentation/Source/Platforms/GoogleVR/GoogleVR.INT.udn
--- a/Engine/Documentation/Source/Platforms/GoogleVR/GoogleVR.INT.udn
+++ b/Engine/Documentation/Source/Platforms/GoogleVR/GoogleVR.INT.udn
@@ -1,19 +1,17 @@
-Title:Google VR Development	
-Crumbs:%ROOT%, Platforms
-Description:
-Availability:Public
-Navigation:platform
-Platform:GoogleVR
-Version: 4.12
-type:landing
-parent:Platforms/VR
-tags:VR
+Availability: Docs
+Crumbs: %ROOT%
+Title: Google VR Development	
+Description:Information on using UE4 to develop for Google VR.
+Type: Landing
+Version:4.17
+Parent: Platforms/VR
+Order: 1
 tags:Google VR
-tags:Platform
 prereq: Platforms/Android/GettingStarted/
 social-image:googlevr_social.png
 topic-image:GoogleVR_Topic_Image.png
 
+
 ![](googlevr_social.png)
 
 [region:note]
@@ -33,4 +31,9 @@ Google Daydream allows you to enable certain rendering optimizations and use Mot
 [DIR(output:"topic" type:"overview" parent:"Platforms/GoogleVR")]
 
 
- 
\ No newline at end of file
+## How To
+
+ The Google VR How - To's will walk you through how to build a wide range different task with Google VR.
+
+[DIR(output:"topic" type:"how-to" parent:"Platforms/GoogleVR")]
+
diff --git a/Engine/Documentation/Source/Platforms/GoogleVR/HowTo/CameraSetup/HTGoogleVRCameraSetup.INT.udn b/Engine/Documentation/Source/Platforms/GoogleVR/HowTo/CameraSetup/HTGoogleVRCameraSetup.INT.udn
new file mode 100644
--- /dev/null
+++ b/Engine/Documentation/Source/Platforms/GoogleVR/HowTo/CameraSetup/HTGoogleVRCameraSetup.INT.udn
@@ -0,0 +1,222 @@
+Availability: Public
+Crumbs:%ROOT%, Platforms
+Title:Set Up Cameras for Google VR
+Description:Setting up the UE4 camera to work with Google VR.
+Type: how-to
+SkillLevel:Beginner
+Version:4.17
+Parent:Platforms/GoogleVR
+Order:2
+Tags:Basics
+topic-image:Platforms\GoogleVR\HowTo\HT_GoogleVR_Camera_Setup_Topic_Image.png
+prereq: Platforms\GoogleVR\QuickStart
+
+No matter if your Google VR project is a sitting or standing experience, the manner in which the user's camera is set up will be the same. In the following How - To we will go over how up the UE4 camera for use with Google VR experiences.
+
+## Project & Basic Camera Setup for a Google VR Experience 
+
+Below you will find instructions on project setup requirements as well as how to create and setup the needed Blueprints for a Google VR experience.
+
+1. First, create a new UE4 project with the following settings and once completed press the **Create Project** button.
+
+    * Blueprint 
+    * Blank Project
+    * Mobile / Tablet
+    * Scalable 3D or 2D
+    * No Starter Content
+    * Location for the project to be saved
+    * Name for the project
+
+    [REGION:lightbox]
+    [![](HT_GoogleVR_Camera_Setup_00.png)(w:555)](HT_GoogleVR_Camera_Setup_00.png)
+    [/REGION]
+
+    [REGION:caption]
+    Click for full image.
+    [/REGION]
+
+1. When the UE4 Editor loads, from the Main Menu, open up the **Plugins** menu by going to **Edit** > **Plugins**.
+    
+    ![](HT_GoogleVR_Camera_Setup_01.png)
+
+1.  Inside the Plugins window, select the **Virtual Reality** section and then enable the **Google VR** option by clicking on the **Enabled** check mark box.
+ 
+    [REGION:lightbox]
+    [![](HT_GoogleVR_Camera_Setup_02.png)(w:547)](HT_GoogleVR_Camera_Setup_02.png)
+    [/REGION]
+
+    [REGION:caption]
+    Click for full image.
+    [/REGION]
+
+    [region:note]
+    You will need to re-start UE4 once you have enabled the plugin so that it will be loaded on the next UE4 Editor launch.
+    [/region] 
+
+1. After the UE4 Editor has re-started, Right-click in the **Content Browser** and from the displayed menu, go to the **Create Basic Asset** section and select the **Blueprint Class** option.
+
+    [![](Platforms\Oculus\HowTo\Camera\HT_Rift_Camera_Setup_01.png)(w:150)](Platforms\Oculus\HowTo\Camera\HT_Rift_Camera_Setup_01.png)
+
+1. From the **Pick Parent Class** window, select the **Pawn** option, giving it a name of **GoogleVR_Pawn** then double - click on the GoogleVRPawn to open it up in the Blueprint Editor.
+
+    [![](HT_GoogleVR_Camera_Setup_03.png)(w:691)](HT_GoogleVR_Camera_Setup_03.png)
+
+1. Open up the GoogleVR_Pawn Blueprint and go to the **Component** section of the **Viewport** tab. Using the **Add Component** button add the following two components with the following names and make sure that the Camera is a child of the Scene Component.
+    
+    * Scene - VRCameraRoot
+    * Camera - VRCamera
+
+    ![](HT_GoogleVR_Camera_Setup_03.png)
+
+    [region:note]
+    No matter what VR HMD you are using, this is the Epic recommend way to set up a VR camera as it allows you to offset the position of the camera without having to move the actual camera.
+    [/region]
+
+Now that a basic Blueprint for the VR Camera has been set up, in the next section we will go through setting up this Blueprint to be used with Google VR.
+
+## Camera Setup for Google VR Experience 
+
+Below you will find instructions on how to setup your camera for a Google VR experience.
+
+
+1. First, open up the **GoogleVR_Pawn** Blueprint and in the **Event Graph** drag off the **Event Begin Play** node to show the Executable Actions list. In the list search for the **Set Tracking Origin** node and click on it to add it to the Event Graph.
+
+    [REGION:lightbox]
+    [![](Platforms\Oculus\HowTo\Camera\HT_Rift_Camera_Setup_09.png)(w:555)](Platforms\Oculus\HowTo\Camera\HT_Rift_Camera_Setup_09.png)
+    [/REGION]
+
+    [REGION:caption]
+    Click for full image.
+    [/REGION]
+
+1. Make sure the **Origin** of the Set Tracking Origin node is set to **Eye Level**.
+    
+    [REGION:lightbox]
+    [![](Platforms\Oculus\HowTo\Camera\HT_Rift_Camera_Setup_10.png)(w:555)](Platforms\Oculus\HowTo\Camera\HT_Rift_Camera_Setup_10.png)
+    [/REGION]
+
+    [REGION:caption]
+    Click for full image.
+    [/REGION]
+
+1. Create a new **Vector** variable called **CameraHeight** and set the **Z** value to **180**.
+
+    ![](Platforms\GearVR\HowTo\CameraSetup\HT_GearVR_Camera_Setup_03.png)
+
+    [region:note]
+    This number should be set to the real world standing height of the user in Centimeters.
+    [/region]
+
+1. Drag off the output of the Set Tracking Origin node and search for the **Set Relative Location** node, selecting the **SetRelativeLocation(VRCameraRoot)** option.
+
+    [REGION:lightbox]
+    [![](HT_GoogleVR_Camera_Setup_06.png)(w:342)](HT_GoogleVR_Camera_Setup_06.png)
+    [/REGION]
+
+    [REGION:caption]
+    Click for full image.
+    [/REGION]
+
+1. Connect the **CameraHeight** variable to the **New Location** input on the Set Relative Location node and then press the Compile button. When completed your Event Graph should look like the following image.
+
+    [REGION:fullwidth]
+    [INCLUDE:Platforms\GoogleVR\HowTo\CameraSetup\#GoogleVRCamera]
+    [/REGION]
+
+    [region:tip]
+    You can copy the completed Blueprint by clicking on the upper left-hand corner of above and pressing CRTL + C to copy the graph. Once copied, go to your Blueprints event and press CTRL + V to paste it in. 
+    [/region]
+
+1. Drag the **GoogleVR_Pawn** Blueprint from the Content Browser to a level, making sure that it is placed as close to the ground as possible.
+
+    [REGION:lightbox]
+    [![](HT_GoogleVR_Camera_Setup_09.png)(w:555)](HT_GoogleVR_Camera_Setup_09.png)
+    [/REGION]
+
+    [REGION:caption]
+    Click for full image.
+    [/REGION]
+
+1. Select the GoogleVR_Pawn Blueprint and in the **Details** panel under the **Pawn** setting, set the **Auto Possess Player** from **Disabled** to **Player 0**.
+
+    ![](HT_GoogleVR_Camera_Setup_05.png)
+
+1. Finally, go to the **Main Toolbar** and click on the small white triangle next to **Launch Options** and then select your Android device from the list to start the packaging and deployment to that device.
+
+    ![](Platforms\GearVR\HowTo\CameraSetup\HT_GearVR_Camera_Setup_08.png)
+
+## UE4 Project Downloads
+
+Below you will find a link to where you can download the UE4 project that was used to create this example. 
+
+* [**Google VR Camera Example Project**](https://epicgames.box.com/s/duniqrm3wl8u9lkq7igskq4fh7rl9e0f)
+
+
+<!--
+[EXCERPT:GoogleVRCamera]	
+[OBJECT:Snippet]
+	[PARAMLITERAL:id]
+		GoogleVRCamera
+	[/PARAMLITERAL]
+	[PARAMLITERAL:type]
+		blueprint
+	[/PARAMLITERAL]
+	[PARAM:image]
+		[REGION:lightbox]
+		[![](HT_GoogleVR_Camera_Setup_07.png)(w:1000)](HT_GoogleVR_Camera_Setup_07.png)
+		[/REGION]
+	[/PARAM]
+	[PARAM:snippet]
+						Begin Object Class=/Script/BlueprintGraph.K2Node_Event Name="K2Node_Event_3"
+			   EventReference=(MemberParent=Class'/Script/Engine.Actor',MemberName="ReceiveBeginPlay")
+			   bOverrideFunction=True
+			   NodePosX=-64
+			   NodePosY=208
+			   NodeGuid=93C09D814ED1E9D44D261A9686B86DCA
+			   CustomProperties Pin (PinId=7F2EC56943D8143A39FA97A62B1F4E79,PinName="OutputDelegate",Direction="EGPD_Output",PinType.PinCategory="delegate",PinType.PinSubCategory="",PinType.PinSubCategoryObject=None,PinType.PinSubCategoryMemberReference=(MemberParent=Class'/Script/Engine.Actor',MemberName="ReceiveBeginPlay"),PinType.PinValueType=(),PinType.ContainerType=None,PinType.bIsMap=False,PinType.bIsSet=False,PinType.bIsArray=False,PinType.bIsReference=False,PinType.bIsConst=False,PinType.bIsWeakPointer=False,PersistentGuid=00000000000000000000000000000000,bHidden=False,bNotConnectable=False,bDefaultValueIsReadOnly=False,bDefaultValueIsIgnored=False,bAdvancedView=False,bOrphanedPin=False,)
+			   CustomProperties Pin (PinId=FBAC704C41732BA0574711B1FCA29206,PinName="then",Direction="EGPD_Output",PinType.PinCategory="exec",PinType.PinSubCategory="",PinType.PinSubCategoryObject=None,PinType.PinSubCategoryMemberReference=(),PinType.PinValueType=(),PinType.ContainerType=None,PinType.bIsMap=False,PinType.bIsSet=False,PinType.bIsArray=False,PinType.bIsReference=False,PinType.bIsConst=False,PinType.bIsWeakPointer=False,LinkedTo=(K2Node_CallFunction_1 65C430D34EC91388FB4328874913A070,),PersistentGuid=00000000000000000000000000000000,bHidden=False,bNotConnectable=False,bDefaultValueIsReadOnly=False,bDefaultValueIsIgnored=False,bAdvancedView=False,bOrphanedPin=False,)
+			End Object
+			Begin Object Class=/Script/BlueprintGraph.K2Node_CallFunction Name="K2Node_CallFunction_1"
+			   FunctionReference=(MemberParent=Class'/Script/HeadMountedDisplay.HeadMountedDisplayFunctionLibrary',MemberName="SetTrackingOrigin")
+			   NodePosX=176
+			   NodePosY=208
+			   NodeGuid=9A4CBAD3436F40EAADD15FB1E078F81F
+			   CustomProperties Pin (PinId=65C430D34EC91388FB4328874913A070,PinName="execute",PinType.PinCategory="exec",PinType.PinSubCategory="",PinType.PinSubCategoryObject=None,PinType.PinSubCategoryMemberReference=(),PinType.PinValueType=(),PinType.ContainerType=None,PinType.bIsMap=False,PinType.bIsSet=False,PinType.bIsArray=False,PinType.bIsReference=False,PinType.bIsConst=False,PinType.bIsWeakPointer=False,LinkedTo=(K2Node_Event_3 FBAC704C41732BA0574711B1FCA29206,),PersistentGuid=00000000000000000000000000000000,bHidden=False,bNotConnectable=False,bDefaultValueIsReadOnly=False,bDefaultValueIsIgnored=False,bAdvancedView=False,bOrphanedPin=False,)
+			   CustomProperties Pin (PinId=93BC380F44B0A11FE8A023977452EE2B,PinName="then",Direction="EGPD_Output",PinType.PinCategory="exec",PinType.PinSubCategory="",PinType.PinSubCategoryObject=None,PinType.PinSubCategoryMemberReference=(),PinType.PinValueType=(),PinType.ContainerType=None,PinType.bIsMap=False,PinType.bIsSet=False,PinType.bIsArray=False,PinType.bIsReference=False,PinType.bIsConst=False,PinType.bIsWeakPointer=False,LinkedTo=(K2Node_CallFunction_52 FE5D38794B000B7752514CAFCC5316E4,),PersistentGuid=00000000000000000000000000000000,bHidden=False,bNotConnectable=False,bDefaultValueIsReadOnly=False,bDefaultValueIsIgnored=False,bAdvancedView=False,bOrphanedPin=False,)
+			   CustomProperties Pin (PinId=3CDA4D324563AA248BE6FD898AE21B08,PinName="self",PinFriendlyName=NSLOCTEXT("K2Node", "Target", "Target"),PinType.PinCategory="object",PinType.PinSubCategory="",PinType.PinSubCategoryObject=Class'/Script/HeadMountedDisplay.HeadMountedDisplayFunctionLibrary',PinType.PinSubCategoryMemberReference=(),PinType.PinValueType=(),PinType.ContainerType=None,PinType.bIsMap=False,PinType.bIsSet=False,PinType.bIsArray=False,PinType.bIsReference=False,PinType.bIsConst=False,PinType.bIsWeakPointer=False,DefaultObject="/Script/HeadMountedDisplay.Default__HeadMountedDisplayFunctionLibrary",PersistentGuid=00000000000000000000000000000000,bHidden=True,bNotConnectable=False,bDefaultValueIsReadOnly=False,bDefaultValueIsIgnored=False,bAdvancedView=False,bOrphanedPin=False,)
+			   CustomProperties Pin (PinId=61713C8E4CF694BE266DD2BD7B1A2A53,PinName="Origin",PinType.PinCategory="byte",PinType.PinSubCategory="",PinType.PinSubCategoryObject=Enum'/Script/HeadMountedDisplay.EHMDTrackingOrigin',PinType.PinSubCategoryMemberReference=(),PinType.PinValueType=(),PinType.ContainerType=None,PinType.bIsMap=False,PinType.bIsSet=False,PinType.bIsArray=False,PinType.bIsReference=False,PinType.bIsConst=False,PinType.bIsWeakPointer=False,DefaultValue="Eye",PersistentGuid=00000000000000000000000000000000,bHidden=False,bNotConnectable=False,bDefaultValueIsReadOnly=False,bDefaultValueIsIgnored=False,bAdvancedView=False,bOrphanedPin=False,)
+			End Object
+			Begin Object Class=/Script/BlueprintGraph.K2Node_CallFunction Name="K2Node_CallFunction_52"
+			   FunctionReference=(MemberParent=Class'/Script/Engine.SceneComponent',MemberName="K2_SetRelativeLocation")
+			   NodePosX=528
+			   NodePosY=192
+			   NodeGuid=CBF27D16498498DCC1B698BBB05BBB64
+			   CustomProperties Pin (PinId=FE5D38794B000B7752514CAFCC5316E4,PinName="execute",PinType.PinCategory="exec",PinType.PinSubCategory="",PinType.PinSubCategoryObject=None,PinType.PinSubCategoryMemberReference=(),PinType.PinValueType=(),PinType.ContainerType=None,PinType.bIsMap=False,PinType.bIsSet=False,PinType.bIsArray=False,PinType.bIsReference=False,PinType.bIsConst=False,PinType.bIsWeakPointer=False,LinkedTo=(K2Node_CallFunction_1 93BC380F44B0A11FE8A023977452EE2B,),PersistentGuid=00000000000000000000000000000000,bHidden=False,bNotConnectable=False,bDefaultValueIsReadOnly=False,bDefaultValueIsIgnored=False,bAdvancedView=False,bOrphanedPin=False,)
+			   CustomProperties Pin (PinId=1C7A9A54494AE23B7164B7B8725F5D52,PinName="then",Direction="EGPD_Output",PinType.PinCategory="exec",PinType.PinSubCategory="",PinType.PinSubCategoryObject=None,PinType.PinSubCategoryMemberReference=(),PinType.PinValueType=(),PinType.ContainerType=None,PinType.bIsMap=False,PinType.bIsSet=False,PinType.bIsArray=False,PinType.bIsReference=False,PinType.bIsConst=False,PinType.bIsWeakPointer=False,PersistentGuid=00000000000000000000000000000000,bHidden=False,bNotConnectable=False,bDefaultValueIsReadOnly=False,bDefaultValueIsIgnored=False,bAdvancedView=False,bOrphanedPin=False,)
+			   CustomProperties Pin (PinId=0811E84B4E91F1030D8FE3BAC770A247,PinName="self",PinFriendlyName=NSLOCTEXT("K2Node", "Target", "Target"),PinType.PinCategory="object",PinType.PinSubCategory="",PinType.PinSubCategoryObject=Class'/Script/Engine.SceneComponent',PinType.PinSubCategoryMemberReference=(),PinType.PinValueType=(),PinType.ContainerType=None,PinType.bIsMap=False,PinType.bIsSet=False,PinType.bIsArray=False,PinType.bIsReference=False,PinType.bIsConst=False,PinType.bIsWeakPointer=False,LinkedTo=(K2Node_VariableGet_23 14EFDA2C47FF2DF1F58E7A9BA401135E,),PersistentGuid=00000000000000000000000000000000,bHidden=False,bNotConnectable=False,bDefaultValueIsReadOnly=False,bDefaultValueIsIgnored=False,bAdvancedView=False,bOrphanedPin=False,)
+			   CustomProperties Pin (PinId=C2D0A283497566CAFA793CB992D07158,PinName="NewLocation",PinType.PinCategory="struct",PinType.PinSubCategory="",PinType.PinSubCategoryObject=ScriptStruct'/Script/CoreUObject.Vector',PinType.PinSubCategoryMemberReference=(),PinType.PinValueType=(),PinType.ContainerType=None,PinType.bIsMap=False,PinType.bIsSet=False,PinType.bIsArray=False,PinType.bIsReference=False,PinType.bIsConst=False,PinType.bIsWeakPointer=False,DefaultValue="0, 0, 0",AutogeneratedDefaultValue="0, 0, 0",LinkedTo=(K2Node_VariableGet_24 7FE7B76D4888DB151FF54EAB9D25F826,),PersistentGuid=00000000000000000000000000000000,bHidden=False,bNotConnectable=False,bDefaultValueIsReadOnly=False,bDefaultValueIsIgnored=False,bAdvancedView=False,bOrphanedPin=False,)
+			   CustomProperties Pin (PinId=3D1069744CC4C94F298F3980C2A0F289,PinName="bSweep",PinType.PinCategory="bool",PinType.PinSubCategory="",PinType.PinSubCategoryObject=None,PinType.PinSubCategoryMemberReference=(),PinType.PinValueType=(),PinType.ContainerType=None,PinType.bIsMap=False,PinType.bIsSet=False,PinType.bIsArray=False,PinType.bIsReference=False,PinType.bIsConst=False,PinType.bIsWeakPointer=False,DefaultValue="false",AutogeneratedDefaultValue="false",PersistentGuid=00000000000000000000000000000000,bHidden=False,bNotConnectable=False,bDefaultValueIsReadOnly=False,bDefaultValueIsIgnored=False,bAdvancedView=False,bOrphanedPin=False,)
+			   CustomProperties Pin (PinId=63A901EA42073DA211529DBD9BDB4FF8,PinName="SweepHitResult",Direction="EGPD_Output",PinType.PinCategory="struct",PinType.PinSubCategory="",PinType.PinSubCategoryObject=ScriptStruct'/Script/Engine.HitResult',PinType.PinSubCategoryMemberReference=(),PinType.PinValueType=(),PinType.ContainerType=None,PinType.bIsMap=False,PinType.bIsSet=False,PinType.bIsArray=False,PinType.bIsReference=False,PinType.bIsConst=False,PinType.bIsWeakPointer=False,PersistentGuid=00000000000000000000000000000000,bHidden=False,bNotConnectable=False,bDefaultValueIsReadOnly=False,bDefaultValueIsIgnored=False,bAdvancedView=False,bOrphanedPin=False,)
+			   CustomProperties Pin (PinId=03544D8C4FC9FC037BE82CB2780E14A1,PinName="bTeleport",PinType.PinCategory="bool",PinType.PinSubCategory="",PinType.PinSubCategoryObject=None,PinType.PinSubCategoryMemberReference=(),PinType.PinValueType=(),PinType.ContainerType=None,PinType.bIsMap=False,PinType.bIsSet=False,PinType.bIsArray=False,PinType.bIsReference=False,PinType.bIsConst=False,PinType.bIsWeakPointer=False,DefaultValue="false",AutogeneratedDefaultValue="false",PersistentGuid=00000000000000000000000000000000,bHidden=False,bNotConnectable=False,bDefaultValueIsReadOnly=False,bDefaultValueIsIgnored=False,bAdvancedView=False,bOrphanedPin=False,)
+			End Object
+			Begin Object Class=/Script/BlueprintGraph.K2Node_VariableGet Name="K2Node_VariableGet_23"
+			   VariableReference=(MemberName="VRCameraRoot",bSelfContext=True)
+			   NodePosX=272
+			   NodePosY=352
+			   NodeGuid=8372D50849CD3678B5FC15A6DEF7B8FF
+			   CustomProperties Pin (PinId=14EFDA2C47FF2DF1F58E7A9BA401135E,PinName="VRCameraRoot",Direction="EGPD_Output",PinType.PinCategory="object",PinType.PinSubCategory="",PinType.PinSubCategoryObject=Class'/Script/Engine.SceneComponent',PinType.PinSubCategoryMemberReference=(),PinType.PinValueType=(),PinType.ContainerType=None,PinType.bIsMap=False,PinType.bIsSet=False,PinType.bIsArray=False,PinType.bIsReference=False,PinType.bIsConst=False,PinType.bIsWeakPointer=False,LinkedTo=(K2Node_CallFunction_52 0811E84B4E91F1030D8FE3BAC770A247,),PersistentGuid=00000000000000000000000000000000,bHidden=False,bNotConnectable=False,bDefaultValueIsReadOnly=False,bDefaultValueIsIgnored=False,bAdvancedView=False,bOrphanedPin=False,)
+			   CustomProperties Pin (PinId=E2C335C44D268CDCA73404AA7E7D9ABD,PinName="self",PinFriendlyName=NSLOCTEXT("K2Node", "Target", "Target"),PinType.PinCategory="object",PinType.PinSubCategory="",PinType.PinSubCategoryObject=BlueprintGeneratedClass'/Game/Blueprints/GoogleVR_Pawn.GoogleVR_Pawn_C',PinType.PinSubCategoryMemberReference=(),PinType.PinValueType=(),PinType.ContainerType=None,PinType.bIsMap=False,PinType.bIsSet=False,PinType.bIsArray=False,PinType.bIsReference=False,PinType.bIsConst=False,PinType.bIsWeakPointer=False,PersistentGuid=00000000000000000000000000000000,bHidden=True,bNotConnectable=False,bDefaultValueIsReadOnly=False,bDefaultValueIsIgnored=False,bAdvancedView=False,bOrphanedPin=False,)
+			End Object
+			Begin Object Class=/Script/BlueprintGraph.K2Node_VariableGet Name="K2Node_VariableGet_24"
+			   VariableReference=(MemberName="CameraHeight",MemberGuid=ACC6F49F42C088E811CC9BB892E6F87B,bSelfContext=True)
+			   NodePosX=288
+			   NodePosY=432
+			   NodeGuid=BBF93606403186ED65F98F907C091CC5
+			   CustomProperties Pin (PinId=7FE7B76D4888DB151FF54EAB9D25F826,PinName="CameraHeight",Direction="EGPD_Output",PinType.PinCategory="struct",PinType.PinSubCategory="",PinType.PinSubCategoryObject=ScriptStruct'/Script/CoreUObject.Vector',PinType.PinSubCategoryMemberReference=(),PinType.PinValueType=(),PinType.ContainerType=None,PinType.bIsMap=False,PinType.bIsSet=False,PinType.bIsArray=False,PinType.bIsReference=False,PinType.bIsConst=False,PinType.bIsWeakPointer=False,DefaultValue="0, 0, 0",AutogeneratedDefaultValue="0, 0, 0",LinkedTo=(K2Node_CallFunction_52 C2D0A283497566CAFA793CB992D07158,),PersistentGuid=00000000000000000000000000000000,bHidden=False,bNotConnectable=False,bDefaultValueIsReadOnly=False,bDefaultValueIsIgnored=False,bAdvancedView=False,bOrphanedPin=False,)
+			   CustomProperties Pin (PinId=DA87D3494F0E12714F4126B699524EAA,PinName="self",PinFriendlyName=NSLOCTEXT("K2Node", "Target", "Target"),PinType.PinCategory="object",PinType.PinSubCategory="",PinType.PinSubCategoryObject=BlueprintGeneratedClass'/Game/Blueprints/GoogleVR_Pawn.GoogleVR_Pawn_C',PinType.PinSubCategoryMemberReference=(),PinType.PinValueType=(),PinType.ContainerType=None,PinType.bIsMap=False,PinType.bIsSet=False,PinType.bIsArray=False,PinType.bIsReference=False,PinType.bIsConst=False,PinType.bIsWeakPointer=False,PersistentGuid=00000000000000000000000000000000,bHidden=True,bNotConnectable=False,bDefaultValueIsReadOnly=False,bDefaultValueIsIgnored=False,bAdvancedView=False,bOrphanedPin=False,)
+			End Object
+			END OBJECT
+	[/PARAM]
+[/OBJECT]
+[/EXCERPT:GoogleVRCamera]
+-->
\ No newline at end of file
diff --git a/Engine/Documentation/Source/Platforms/GoogleVR/HowTo/GoogleVRHowTo.INT.udn b/Engine/Documentation/Source/Platforms/GoogleVR/HowTo/GoogleVRHowTo.INT.udn
new file mode 100644
--- /dev/null
+++ b/Engine/Documentation/Source/Platforms/GoogleVR/HowTo/GoogleVRHowTo.INT.udn
@@ -0,0 +1,34 @@
+Availability:Public
+Crumbs:%ROOT%
+Title:Google VR How - To
+Description: How To guides providing step-by-step instructions working with the Google VR and UE4.
+Type:how-to
+Version:4.17
+Parent:Parent: Platforms/GoogleVR
+Order:2
+Tags:GoogleVR
+Tags:VR
+Tags:Mobile
+
+The Google VR How To's contain step-by-step guides that will walk you through how to use UE4 and Google VR together.
+
+## Beginner
+
+The Google VR Beginner How - To's will walk you through setting up things like Cameras 
+
+[DIR(end:"1" output:"topic" type:"how-to" tags:"Basics" parent:"Platforms/GoogleVR")]
+
+## Intermediate
+
+The Google VR Intermediate How - To's will walk you through setting up things like the Google Daydream VR Controller.
+
+[DIR(end:"1" output:"topic" type:"how-to" tags:"GoogleVR" parent:"Platforms/GoogleVR")]
+
+## Advanced
+
+The Google VR Advanced How - To's will walk you through debugging your Google VR projects. 
+
+[DIR(end:"1" output:"topic" type:"how-to" tags:"GoogleVR" parent:"Platforms/GoogleVR")]
+
+
+
