Copying //UE4/Dev-Documentation to //UE4/Samples-Main (Source: //UE4/Dev-Documentation @ 3623665)

diff --git a/Engine/Documentation/Source/Platforms/GoogleVR/QuickStart/2/GoogleVRQuickStart_2.INT.udn b/Engine/Documentation/Source/Platforms/GoogleVR/QuickStart/2/GoogleVRQuickStart_2.INT.udn
--- a/Engine/Documentation/Source/Platforms/GoogleVR/QuickStart/2/GoogleVRQuickStart_2.INT.udn
+++ b/Engine/Documentation/Source/Platforms/GoogleVR/QuickStart/2/GoogleVRQuickStart_2.INT.udn
@@ -4,20 +4,22 @@ Crumbs:%ROOT%, Platforms/GoogleVR
 Description:Setting up a UE4 project to work with Google VR.
 Navigation:platform
 Skilllevel: Intermediate
-Version: 4.16
+Version: 4.17
 Parent:Platforms/GoogleVR/QuickStart
 order:2
 type:multi-step
 tags:VR
 tags:Google VR
-
+Checkpoint: GoogleVRQS
 [NAV]
 
+In this part of the Google VR Quick Start, we will go over setting up your UE4 project so that it can be used with Google VR.
+
 [region:note]
 Google has a provided a very detailed document that goes over everything you need to know to get the most out of Daydream. You can find the documents by going to the [Offical Google Daydream Unreal](https://developers.google.com/vr/unreal/) help site.
 [/region]
 
- In this part of the Google VR Quick Start, we will go over setting up your UE4 project so that it can be used with Google VR.
+## Steps
 
 1. Once the UE4 Editor has loaded, from the **Main Toolbar**, open the **Edit** tab and select the **Plugins** menu item.
 
@@ -84,6 +86,7 @@ Google has a provided a very detailed document that goes over everything you nee
 		[PARAMLITERAL:type]
 		simple
 		[/PARAMLITERAL]
+	[/OBJECT]
 
 1. In the **Platforms** section under the **Android** section press the **Configure Now** button in the **APKPackaging** and **Google Play Services** sections to ensure that our project can package correctly.
 
@@ -144,7 +147,7 @@ Google has a provided a very detailed document that goes over everything you nee
 
     [/region]
 
-1. In the **Android SDK** section under the **SDKConfig** section set the **SDK API Level** to **matchndk**. Then set the **NDK API Level** to **android-24** to ensure that your project will be compiled against the most recent Android SDK.
+1. In the **Android SDK** section under the **SDKConfig** section set the **SDK API Level** to **android-24** and set the **NDK API Level** to ** android-19 ** to ensure that your project will be compiled against the most recent Android SDK.
 
     ![](GVRQS_Android_SDK_Options_00.png)
 
@@ -155,6 +158,7 @@ Google has a provided a very detailed document that goes over everything you nee
         * **JDK 1.8**
     [/region]
 
+## End Results
 
 In the next section of the Google VR Quick Start, we will go over what you have to do to package your UE4 project so that it will run on a Google VR device.
 
diff --git a/Engine/Documentation/Source/Platforms/GoogleVR/QuickStart/3/GoogleVRQuickStart_3.INT.udn b/Engine/Documentation/Source/Platforms/GoogleVR/QuickStart/3/GoogleVRQuickStart_3.INT.udn
--- a/Engine/Documentation/Source/Platforms/GoogleVR/QuickStart/3/GoogleVRQuickStart_3.INT.udn
+++ b/Engine/Documentation/Source/Platforms/GoogleVR/QuickStart/3/GoogleVRQuickStart_3.INT.udn
@@ -4,13 +4,13 @@ Crumbs:%ROOT%, Platforms/GoogleVR
 Description:Packing and deploying a UE4 project that can be used with Google VR
 Navigation:platform
 Skilllevel: Intermediate
-Version: 4.16
+Version:4.17
 Parent:Platforms/GoogleVR/QuickStart
 order:3
 type:multi-step
 tags:VR
 tags:Google VR
-
+Checkpoint: GoogleVRQS
 [NAV]
 
 In this section of the Google VR Quick Start, we will go over what you have to do to package your UE4 project so that it will run on a Google VR device.
@@ -19,11 +19,13 @@ In this section of the Google VR Quick Start, we will go over what you have to d
 Before going any further make sure that the Smartphone used for testing has been plugged into your development PC via a USB cable and make also make sure that ADB can see your phone as an authorized device. If you are unsure how to do or are having problems getting ADB to see your device, check out the following page for more information [Android How - To Set Up Your Android Device](https://docs.unrealengine.com/latest/INT/Platforms/Android/GettingStarted/2/).
 [/region]
 
-1. To start the packing process, go to the **Main Toolbar** and from the **File** menu go to **Package Project** > **Android** > **Android (ATC)**.
+## Steps
+
+1. To start the packing process, go to the **Main Toolbar** and from the **File** menu go to **Package Project** > **Android** > **Android (ASTC)**.
 
     ![](GVRQS_Package_Project_00.png)
 
-1. Once Android(ATC) has been selected, you will be presented with a Windows dialogue box. From this dialogue press, the **Make New Folder** option and call the folder **Androids_Builds**. Once the folder has been created, press the **Okay** button to begin the packing process.
+1. Once Android(ASTC) has been selected, you will be presented with a Windows dialogue box. From this dialogue press, the **Make New Folder** option and call the folder **Androids_Builds**. Once the folder has been created, press the **Okay** button to begin the packing process.
 
     ![](GVRQS_Package_Project_Location_00.png)
 
@@ -73,4 +75,8 @@ Before going any further make sure that the Smartphone used for testing has been
     From the displayed menu, you can scan the QR code that is on the side of your VR device, and that will make sure UE4 outputs an image that is compatible with your VR device.
     [/region]
 
+## End Result
+
+When completed you will be able to view the UE4 project using Google VR.
+
 [NAV]
\ No newline at end of file
diff --git a/Engine/Documentation/Source/Platforms/GoogleVR/QuickStart/4/GoogleVRQuickStart_4.INT.udn b/Engine/Documentation/Source/Platforms/GoogleVR/QuickStart/4/GoogleVRQuickStart_4.INT.udn
new file mode 100644
--- /dev/null
+++ b/Engine/Documentation/Source/Platforms/GoogleVR/QuickStart/4/GoogleVRQuickStart_4.INT.udn
@@ -0,0 +1,38 @@
+Availability:Public
+Title:4. On Your Own!
+Crumbs:%ROOT%, Platforms/GoogleVR
+Description:Using the methods provided during this guide, polish up your level with additional content and read up on related topics. 
+Navigation:platform
+Skilllevel: Intermediate
+Version: 4.17
+Parent:Platforms/GoogleVR/QuickStart
+order:4
+type:multi-step
+tags:VR
+tags:Google VR
+Checkpoint: GoogleVRQS
+
+[NAV]
+
+Now that you can deploy a UE4 project to your Android based smartphone and view it with your Google VR HMD try adding the following items.
+
+* Use the [Migrate Tool](Engine/Content/Browser/UserGuide/Migrate/) to move content from one of the mobile content examples into your project for use.
+
+* Adjust the settings of the [Default Device Profiles](Platforms/Mobile/Performance/TipsAndTricks) to ensure that your project will run as fast as possible when deployed to your Android based smartphone.
+
+* Use the [GPU Profiler](Engine/Performance/GPU/) to help track your projects performance while you are building it.
+
+
+
+As for specifics covered in this quick start:
+
+* For more information on setting up the Android SDK, see: [Install the Android SDK](Platforms/Android/GettingStarted/1/)
+
+* For more information on setting up your Android Device for UE4 Development, see: [Set Up Your Android Device](Platforms/Android/GettingStarted/2/)
+
+* For more information on what features Android devices support, see: [Android Device Compatibility](Platforms/Android/DeviceCompatibility/)
+
+* For more information on reducing the packaged size of your project, see: [Reducing APK Package Size](Platforms/Android/ReducingAPKSize/)
+
+
+[NAV]
\ No newline at end of file
diff --git a/Engine/Documentation/Source/Platforms/GoogleVR/QuickStart/GoogleVRQuickStart.INT.udn b/Engine/Documentation/Source/Platforms/GoogleVR/QuickStart/GoogleVRQuickStart.INT.udn
--- a/Engine/Documentation/Source/Platforms/GoogleVR/QuickStart/GoogleVRQuickStart.INT.udn
+++ b/Engine/Documentation/Source/Platforms/GoogleVR/QuickStart/GoogleVRQuickStart.INT.udn
@@ -1,39 +1,36 @@
-Title:Google VR Quick Start
-Crumbs:%ROOT%, Platforms, Platforms/GoogleVR
-Description:Instructions on setting up and using Google VR
-Availability:Public
-Version: 4.12
+Availability: Public
+Crumbs: %ROOT%
+Title: Google VR Quick Start
+Description: Instructions on setting up and using Google VR
+Type: quick start
 SkillLevel: Intermediate
-parent:Platforms/GoogleVR
-type:quick start
-type:landing
-tags:VR
+Version: 4.17
+Parent: Platforms/GoogleVR
+Order:2
 tags:Google VR
+Checkpoint:GoogleVRQS
 topic-image:Platforms\GoogleVR\GoogleVR_Cardboard_Topic_Image.png
 
-The GoogleVR Quick Start will walk you through setting up your UE4 project to work with Google VR.
+![](Google_Cardboard_Headset.png)
 
+## Goals
 
-[OBJECT:TopicButtonList]
-	[PARAM:icon]
-		![](%ROOT%/build_icon.png)(convert:false)
-	[/PARAM]
-	[PARAM:title]
-		Steps
-	[/PARAM]
-	[PARAM:description]
-	[/PARAM]
-	[PARAM:links]
-	* [](Platforms/GoogleVR/QuickStart/1 "%Platforms/GoogleVR/QuickStart/1:description%")
-	* [](Platforms/GoogleVR/QuickStart/2 "%Platforms/GoogleVR/QuickStart/2:description%")
-	* [](Platforms/GoogleVR/QuickStart/3 "%Platforms/GoogleVR/QuickStart/3:description%")
-	[/PARAM]
-[/OBJECT]
+The GoogleVR Quick Start will walk you through setting up your UE4 project to work with Google VR.
 
 [region:note]
 Google has a provided a very detailed document that goes over everything you need to know to get the most out of Daydream. You can find the documents by going to the [Offical Google Daydream Unreal](https://developers.google.com/vr/unreal/) help site.
 [/region]
 
+## Objectives
+* Create a new UE4 project that is targeted towards Google VR development.
+* Setting the required project settings so your project can be used with Google VR.
+* Deploying and viewing your UE4 project with the Google VR HMD of your choice.
+
+[DIR(output:"checkpoint" parent:"Platforms/GoogleVR/QuickStart")]
+
+[REGION:call_to_action]
+ [Click to Start](Platforms/GoogleVR/QuickStart/1 "%Platforms/GoogleVR/QuickStart/1:title%")
+[/REGION]
 
 
 
diff --git a/Engine/Documentation/Source/Platforms/Linux/AdvancedLinuxDeveloper/LinuxCrossCompileLegacy/LinuxCrossCompileLegacy.INT.udn b/Engine/Documentation/Source/Platforms/Linux/AdvancedLinuxDeveloper/LinuxCrossCompileLegacy/LinuxCrossCompileLegacy.INT.udn
--- a/Engine/Documentation/Source/Platforms/Linux/AdvancedLinuxDeveloper/LinuxCrossCompileLegacy/LinuxCrossCompileLegacy.INT.udn
+++ b/Engine/Documentation/Source/Platforms/Linux/AdvancedLinuxDeveloper/LinuxCrossCompileLegacy/LinuxCrossCompileLegacy.INT.udn
@@ -33,14 +33,18 @@ Mac users currently have to resort to [native compilation](https://wiki.unrealen
 
 Please use the following table to download the appropriate toolchain: 
 
-| Toolchain                                                                                                                           | UE4 Version    | 
-| ----------------------------------------------------------------------------------------------------------------------------------- | -------------- | 
-| **-v4** [clang-3.5.0-based](http://cdn.unrealengine.com/qfe/v4_clang-3.5.0_ld-2.24_glibc-2.12.2.zip)                            | <= 4.8         | 
+| Toolchain                                                                                                                       | UE4 Version    | 
+| ------------------------------------------------------------------------------------------------------------------------------- | -------------- | 
+| **-v4** [clang-3.5.0-based](http://cdn.unrealengine.com/qfe/v4_clang-3.5.0_ld-2.24_glibc-2.12.2.zip)                            | 4.8 and earlier| 
 | **-v6** [clang-3.6.0-based](https://s3.amazonaws.com/unrealengine/qfe/v6_clang-3.6.0_ld-2.24_glibc-2.12.2.zip)                  | 4.9 and 4.10   | 
 | **-v7** [clang-3.7.0-based](https://s3.amazonaws.com/unrealengine/CrossToolchain_Linux/v7_clang-3.7.0_ld-2.24_glibc-2.12.2.zip) | 4.11 thru 4.13 | 
 | **-v8** [clang-3.9.0-based](http://cdn.unrealengine.com/qfe/v8_clang-3.9.0-centos7.zip)                                         | 4.14 and 4.15  |
 | **-v9** [clang-4.0.0-based](http://cdn.unrealengine.com/CrossToolchain_Linux/v9_clang-4.0.0-centos7.zip)                        | 4.16           |
 
+[REGION:note]
+The Setup and UnSetup Windows Batch Files (\*.bat) are not included in **-v9** because you only need these files if you're setting up AutoSDKs for UE4. 
+`Setup.bat` and `Unsetup.bat` were included in previous versions of the toolchain because they're part of the AutoSDKs system (which is currently not documented).
+[/REGION]  
 
 We also provide the libraries and toolchains that allow you to compile for [Linux ARM](https://s3.amazonaws.com/unrealengine/qfe/arm-unknown-linux-gnueabihf_v5_clang-3.5.0-ld-2.23.1-glibc-2.13.zip) as well (original Raspberry Pi and up). 
 Although, this will require you to make a (minor) code change in [Unreal Build Tool](Programming/UnrealBuildSystem) (UBT).
diff --git a/Engine/Documentation/Source/Platforms/Mobile/Performance/MobilePerformance.INT.udn b/Engine/Documentation/Source/Platforms/Mobile/Performance/MobilePerformance.INT.udn
--- a/Engine/Documentation/Source/Platforms/Mobile/Performance/MobilePerformance.INT.udn
+++ b/Engine/Documentation/Source/Platforms/Mobile/Performance/MobilePerformance.INT.udn
@@ -1,127 +1,207 @@
-Availability:Public
-Title:Performance Guidelines for Mobile Devices
-Crumbs:%ROOT%, Platforms, Platforms/Mobile
-Description:Optimizing content and features for mobile games.
-Platform:Mobile
-Version: 4.9
-parent:Platforms/Mobile
-order:6
-type:overview
-tags:mobile
+Availability: Public
+Crumbs: %ROOT%
+Title: Performance Guidelines for Mobile Devices
+Description: Optimizing content and features for mobile games. 
+Type: Overview
+SkillLevel:Intermediate
+Version:4.17
+Parent:Platforms/Mobile
+Order:6
+Tags:mobile
 
-[TOC(start:2)]
 
 
+[VAR:spacer]![](%ROOT%/spacer.png)(w:150 convert:false)[/VAR]
+[VAR:spacer2]![](%ROOT%/spacer.png)(w:40 convert:false)[/VAR]
+[VAR:spacer3]![](%ROOT%/spacer.png)(w:100 convert:false)[/VAR]
+[VAR:spacer4]![](%ROOT%/spacer.png)(w:140 convert:false)[/VAR]
+
+[TOC(start:2 end:2)]
+
+When developing for Mobile platforms, there are some things you should keep in mind when targeting specific devices or that are general good practices depending on the 
+list of devices you want to release your title for. This includes developing with a specific lighting tier to get the most out of performance on a
+device and why a particular lighting tier may work best for your target audience. You'll also find some suggestions to keep in mind for any project you work on that will 
+be developed with Mobile in mind.
 
-Below are things to watch out for or suggestions on how to set up your content to maximize performance on mobile platforms. These are listed from most severe to least.
 
-* Make sure lighting is built before running on a device.
-* Make sure most **Post Process** features are disabled.  In fact, the only ones you can use for iPad4 are **Temporal AA** and the Film ones, including vignette.  A bunch of expensive features like Bloom and DOF are enabled by default as we have the same settings between PC and mobile.  **Post Process** features can easily cost 60+ms with the default settings.  We hope to be able to use DOF and light shafts on higher end devices.  The console command `show postprocessing` can be used on devices to toggle **Post Processing** and see how much it is costing.
-* Make sure **Precomputed Visibility** is setup properly.  To do this, you need to place **Precomputed Visibility Volumes** around anywhere that the player can walk, and then build lighting.  You need to make sure the same **Persistent** level is used when building lighting and when running the game (aka do not build the sublevel by itself).  You can verify that it is working by entering `stat initviews` on device or in previewer and make sure the `Statically occluded primitives` is > 0.  Use `r.ShowPrecomputedVisibilityCells 1` to visualize the cells in the editor.
-* Use masked and translucent **Materials** sparingly, and only in places where they cover a small part of the screen.  iOS devices are very optimal in shading opaque surfaces, as they will only shade each pixel once, most of the time.  But for masked and translucency, every single layer will have to be shaded.  The total GPU time of a frame can easily be doubled or more by having overdraw.  Use the **Shader Complexity Viewmode** in **Mobile Previewer** to investigate hot spots.
-* Draw calls of the entire scene should be &lt;=700 for any view.  The areas with poor occlusion, like looking over the landscape, are going to be the biggest challenges for this.  This can be seen with `stat openglrhi` on device, `stat d3d11rhi` in ES2 previewer on PC.
-* Triangle count of the entire scene should be &lt;=500k for any view.  This has been determined to be the poly count needed to hit 30fps on both iPad4 and iPad Air.  This can be seen with `stat openglrhi` on device, `stat d3d11rhi` in ES2 previewer on PC.
-* **Materials** should use as few texture lookups and instructions as possible.  You have 5 texture samplers available, but using all of those makes for a fairly expensive Material.  For now, use the PC instruction count to optimize, and ES2 previewer shader complexity to visualize the total cost.
-* Only use independent texture fetches in Materials.  This means UVs in the pixel shader (BaseColor, Roughness, etc) must not be manipulated in any way, such as scaling.  Use the **CustomizedUVs** feature instead to do the scaling in the vertex shader.  Some special features like environment mapping require math on the UVs and that is ok for special cases. 
-* Prefer square textures as they have less memory waste, and always use power of 2 dimensions (512, 1024, 2048).  Use `listtextures` in the ES2 previewer to see where all the texture memory is going. 
 
 ## Performance Tiers
 
-Unreal Engine 4 (UE4) supports a variety of lighting features on mobile devices. Using these features costs performance and may cause your game to perform poorly on slower mobile devices. While it is possible to mix and match most of UE4's mobile lighting features, it can be useful to categorize these features into the following tiers. When building a mobile game, you should decide which features to use based on the quality of graphics your game requires and the types of devices you need to support. Check the device compatibility pages under [iOS Development](Platforms/iOS/DeviceCompatibility) and [Android Development](Platforms/Android/DeviceCompatibility) for more information on what devices we have tested here at Epic and which tiers we think are most appropriate.
+Unreal Engine 4 (UE4) supports a variety of lighting features on mobile devices. Using these features costs performance and may cause your game to perform poorly on slower mobile devices. While it is 
+possible to mix and match most of UE4's mobile lighting features, it can be useful to categorize these features into the following tiers. When building a mobile game, you should decide which features to
+use based on the quality of graphics your game requires and the types of devices you need to support. Check the compatibility page under [iOS Development](Platforms/iOS/DeviceCompatibility) and [Android Development](Platforms/Android/DeviceCompatibility)
+for more information on what devices have been tested here at Epic and which tiers we think are most appropriate for that device.
 
 ###LDR (Low Dynamic Range)
 
-Low dynamic range or LDR mode is the highest performance tier supported in UE4 and is recommended for games that do not require lighting or **Post Processing** features. To use this mode, you must disable `Mobile HDR` for your project in the Rendering section of the [Project Settings Editor](Engine/UI/ProjectSettings). 
+**Low Dynamic Range** (LDR) mode is the lowest performance tier supported in UE4 and is recommended for games that do not require lighting or Post Processing features. 
 
-**Advantages**
+To use this mode, you must disable **Mobile HDR** for your project in the Rendering section of the [Project Settings](Engine/UI/ProjectSettings). 
 
-* The fastest and lowest overhead mode available for mobile devices. Enables your game to run well on slower mobile devices.
-* Still provides full access to the Material editor for defining custom shaders and even perform simple shading to fake lighting.
+![](DisableMobileHDR.png)
+
+[REGION:simpletable]
+| Advantages | Limitations | Recommendations |
+| --- | --- | --- |
+| [INCLUDE:#ldradvantages] | [INCLUDE:#ldrlimitations] | [INCLUDE:#ldrreccommendations] |
+[/REGION]
 
+<!--- EXCERPTS LDR
 
-**Limitations**
+[INCLUDE:#ldrreccommendations]
 
-* The scene's color is written out in gamma space with each color channel clamped to the range [0,1].
-* Translucent primitives are blended in gamma space. In most cases, this requires that you author your translucent textures and Materials differently than you would in HDR or for a normal PC game.
-* Post Processing features are unavailable in this mode.
+[EXCERPT:ldradvantages]
+* Provides the fastest and lowest overhead mode available for mobile devices, which enables your game to run well on slower mobile devices.
+* Still provides full access to the Material Editor for defining custom shaders and even performing simple shading that can be used to fake lighting.
+[/EXCERPT:ldradvantages]
 
+[EXCERPT:ldrlimitations]
+* The scene's color is written out in gamma space with each color channel clamped to the range of [0,1].
+* Translucent primitives are blended in gamma space. In most cases, this will require you to author your translucent textures and Materials differently than you would in HDR or for a normal PC title.
+* Post Processing Feature are unavailable in this mode.
+[/EXCERPT:ldrlimitations]
 
-**Recommendations**
+[EXCERPT:ldrreccommendations]
+* Make sure all of your Materials have their shading model set to **Unlit** for maximum performance.
+* Placed lighting should not be used in your scene when depending on maximum performance.
+* Consider performing as many operations in the Materials **Vertex Shader** as possible. You can do this by enabling **Customized UVs**, connecting nodes to them and then in the pixel shader using a **Texture Coordinate** node to target the customized UV.
+[/EXCERPT:ldrreccommendations]
+
+--->
 
-* Make sure all of your **Materials** have their shading model set to **Unlit** for maximum performance.
-* Do not place lights in your scene or otherwise depend on lighting for maximum performance.
-* Consider performing as many operations in the **Material's** vertex shader as possible. You can do so by enabling **Customized UVs**, connecting nodes to them, and reading them in the pixel shader using a **Texture Coordinate** node.
 
 
 ### Basic Lighting
 
-In this tier, you leverage static lighting and fully rough **Materials** to create levels with interesting lighting while maximizing performance to reach a broader range of mobile devices.
+In this tier, you will leverage Static lighting and fully rough Materials to create levels with interesting lighting while maximizing performance to reach a broader range of mobile devices. 
+
+To use this mode, you must disable **Mobile HDR** for your project in the Rendering section of the [Project Settings](Engine/UI/ProjectSettings). 
 
-**Advantages**
+![](DisableMobileHDR.png)
 
-* Access to static lighting and global illumination features.
-* Full HDR pipeline with access to some **Post Processing** features like tone mapping.
-* Translucency is blended in linear space allowing you to author content as you normally would for any other UE4 game.
 
+[REGION:simpletable]
+| Advantages | Limitations | Recommendations |
+| --- | --- | --- |
+| [INCLUDE:#BLadvantages] | [INCLUDE:#BLlimitations] | [INCLUDE:#BLreccommendations] |
+[/REGION]
+
+<!--- EXCERPTS Basic Lighting
 
-**Limitations**
+[INCLUDE:#BLreccommendations]
 
-* As you are marking all **Materials** as **Fully Rough**, your **Materials** will not have interesting specular reflections.
-* If you also choose to disable lightmap directionality, normal maps will have no effect.
+[EXCERPT:BLadvantages]
+* Access to Static Lighting and Global Illumination  features.
+* Full HDR pipeline with access to some Post Processing features, like Tone Mapping.
+* Translucency is blended in linear space, enabling you to author content as your normally would for Desktop.
+[/EXCERPT:BLadvantages]
 
+[EXCERPT:BLlimitations]
+* Since all Materials need to be marked as **Fully Rough**, your Materials will not have interesting specular reflection.
+* If you choose to disable **Lightmap Directionality**, Normal Maps will have no effect.
+[/EXCERPT:BLlimitations]
 
-**Recommendations**
+[EXCERPT:BLreccommendations]
+* Author all of your Materials with the **Fully Rough** flag set.
+* Consider disabling the flag for **Lightmap Directionality** in your Materials for additional performance.
+* Use only Static Lights in your maps.
+* Disable some Post Processing features completely, like **Bloom**. Also, stick to the basic set of film and color controls. See [Mobile Post Process](Platforms/Mobile/PostProcessEffects) for more details.
+[/EXCERPT:BLreccommendations]
 
-* Author all of your **Materials** with the fully rough flag set.
-* Consider disabling light map directionality in your **Materials** for additional performance.
-* Use only static lights in your maps.
-* Disable bloom and anti-aliasing. Also, stick to the basic set of film and color controls. See [mobile Post Process](Platforms/Mobile/PostProcessEffects) for more details.
+--->
 
 
-### Full HDR Lighting
 
-In this tier, you take advantage of most of the HDR lighting features available for mobile in UE4 as well as some of the Post Processing features. Using these features requires quite a bit of performance in exchange for high quality lighting features.
 
-**Advantages**
 
-* Access to static lighting and global illumination features.
+
+### Full HDR (High Dynamic Range) Lighting
+
+In this tier, you take advantage of most of the HDR lighting features available for mobile in UE4, as well as some of the Post Processing features. Using these features requires quite a bit of performance 
+in exchange for high-quality lighting features. 
+
+To use this mode, you must enable **Mobile HDR** for your project in the Rendering section of the [Project Settings](Engine/UI/ProjectSettings). 
+
+![](EnableMobileHDR.png)
+
+
+[REGION:simpletable]
+| Advantages | Recommendations |
+| --- | --- | 
+| [INCLUDE:#HDRadvantages] | [INCLUDE:#HDRreccommendations] |
+[/REGION]
+
+<!--- EXCERPTS Basic Lighting
+
+[INCLUDE:#HDRreccommendations]
+
+[EXCERPT:HDRadvantages]
+* Access to Static Lighting and Global Illumination features.
 * Full HDR pipeline with access to some Post Processing features.
-* Translucency is blended in linear space allowing you to author content as you normally would for any other UE4 game.
+* Translucent is blended in linear space, enabling you to author content as you normally would for Desktop. 
 * Realistic specular reflections on surfaces with support for varying roughness.
-* Full support for normal mapped surfaces allows interesting surface detail without adding a lot of polygons to your meshes.
+[/EXCERPT:HDRadvantages]
 
 
-**Recommendations**
+[EXCERPT:HDRreccommendations]
+* Consider enabling **Bloom** to take full advantage of the HDR lighting pipeline.
+* Realistic specular reflections in combination with HDR lighting can lead to specular aliasing. To reduce this effect, enable the Material property for **Normal Curvature to Roughness** to reduce specular aliasing due to high-frequency information in your normal map. You can also consider enabling anti-aliasing via the **Project Settings > Rendering** to help reduce the artifact.
+* Take time to consider the placement of Reflection Capture Actors to achieve the best results. See [Reflections: Placing Reflection Captures](Resources/Showcases/Reflections/#placingreflectioncaptures) for more information.
+* Use only Static Lights and Lightmaps in your scene for the best performance.
+[/EXCERPT:HDRreccommendations]
+
+--->
 
-* Consider enabling bloom to take full advantage of the HDR lighting pipeline.
-* Realistic specular reflections in combination with HDR lighting can lead to specular aliasing. You should take advantage of the normal map to roughness feature to help reduce specular aliasing due to high frequency information in your normal map. You should also consider enabling anti-aliasing in the Post Process volume to help reduce this artifact.
-* Taking time to consider the placement of **Reflection Captures** will help you get the best results. For more information check out the [Reflections: Placing Reflection Captures](Resources/Showcases/Reflections/#placingreflectioncaptures).
-* Use only Static Lights and Lightmaps in your scene for best performance.
 
 
 ### Full HDR Lighting with per-pixel lighting from the Sun
 
-In this tier, you take advantage of all of the HDR lighting features available for mobile in UE4. This tier is the same as the 'Full HDR Lighting' tier and has the same advantages and recommendations except that a single directional light is rendered at much higher quality.
+In this tier, you take advantage of all of the HDR lighting features available for mobile in UE4. This tier is the same as [Full HDR Lighting](#fullhdr_highdynamicrange_lighting) and has the same advantages and 
+recommendations with the exception that here you can add a single Directional Light to your scene that automatically uses per-pixel lighting for higher quality.
+
+To use this mode, you must enable **Mobile HDR** for your project in the Rendering section of the [Project Settings](Engine/UI/ProjectSettings). 
+
+![](EnableMobileHDR.png)
+
+[REGION:simpletable]
+| Advantages | Recommendations |
+| --- | --- | 
+| [INCLUDE:#HDRSunadvantages] | [INCLUDE:#HDRSunreccommendations] |
+[/REGION]
+
+<!--- EXCERPTS Basic Lighting
 
-**Advantages**
+[INCLUDE:#HDRreccommendations]
 
-* All of the features and advantages of the Full HDR Lighting tier.
-* Per-pixel diffuse and specular lighting for a single directional light.
-* High quality precomputed distance field shadows for a single directional light.
+[EXCERPT:HDRSunadvantages]
+* All of the features and advantages listed for the [Full HDR](#fullhdr_highdynamicrange_lighting) tier.
+* Per-pixel diffuse and specular lighting for a single Directional Light. 
+* High-quality precomputed Distance Field shadows for a single Directional Light. 
+[/EXCERPT:HDRSunadvantages]
 
 
-**Recommendations**
+[EXCERPT:HDRSunreccommendations]
+* All of the recommendations and advantages of the [Full HDR](#fullhdr_highdynamicrange_lighting] tier.
+* Use only Static Lights in your except for a single Directional Light, which should be set to Stationary. 
+[/EXCERPT:HDRSunreccommendations]
 
-* All of the recommendations for the Full HDR Lighting tier apply here.
-* Use only static lights in your scene except for a single directional light that you should set to be stationary.
+--->
 
  
 ## Shader Complexity View Mode
 
-You can use the shader complexity view mode in the PC ES2 previewer to get a sense of Material cost.  To use it, 
-just enter `viewmode shadercomplexity` in the console in the **PC ES2 Previewer**. Just like the SM5 shader complexity, 
-green means good performance, bright red means very expensive, white or pink means **extremely** expensive.  Here are some 
-examples from the mobile showcase:
+The [Shader Complexity](https://docs.unrealengine.com/latest/INT/Engine/UI/LevelEditor/Viewports/ViewModes/index.html#shadercomplexity) view mode in the [Mobile Previewer](Platforms/Mobile/Previewer) to get a 
+sense of what your Material costs will be for specific devices you are targeting. The colors indicated on screen will tell you how expensive the Material is for the device you're targeting; green means good 
+performance, bright red means very expensive, and white or pink means that the Material is very expensive. 
+
+To use it, in the Main Viewport click the **View Modes** list  and select **Optimization Viewmodes** and choose **Shader Complexity**, or you can use the keyboard shortcut **Alt + 8**. 
+
+[REGION:lightbox]
+[![](EnableShaderComplexityViewMode.png)(w:358)](EnableShaderComplexityViewMode.png)
+[/REGION]
+[REGION:caption] Click for full image. [/REGION]
+
+These are some examples taken from Mobile Sun Temple to give you an idea of Material cost when viewed using Shader Complexity:
 
 [REGION:imagetable]
 | ![](1.png)(w:310) | ![](2.png)(w:310) | ![](3.png)(w:310) |
@@ -131,97 +211,89 @@ examples from the mobile showcase:
 
 <!--
 [EXCERPT:Image1]
-The pillar Materials are fairly expensive here, and the translucent volumetric sheets are **very** expensive.  In this scene, I asked to have the translucent sheets removed, as they cost too much. 
+The pillar Materials are fairly expensive here, and the translucent volumetric sheets are **very** expensive.  In this scene, I asked to have the translucent sheets should be removed, as they cost too much. 
 [/EXCERPT:Image1]
 [EXCERPT:Image2]
-Here the pillars are quite expensive as they used all five texture lookups and did quite a bit of layering. Otherwise, this is just about efficient enough for 30fps.
+Here the pillars are quite expensive as they use all five texture lookups and do quite a bit of layering. Otherwise, this is just about efficient enough for 30fps.
 [/EXCERPT:Image2]
 [EXCERPT:Image3]
-Here the tree is causing pixels to be **extremely** expensive.  If the player could ever get up close or cover his screen with that Material, the cost would be enormous.
+Here, the tree is causing pixels to be **extremely** expensive. If the player could ever get up close or cover the screen with that Material, the cost would be enormous.
 [/EXCERPT:Image3]
 -->
 
 ## Mobile Content Scale Factor 
+
 **Mobile Content Scale Factor** is a way to scale the resolution of your project to best fit the screen resolution of the mobile device being used to view your project. 
-You can create and store multiple device profiles by creating a new .INI file called **DefaultDeviceProfiles.ini** in your project's config folder. 
-Inside of this file, you can specify what will happen to the resolution of your project by entering the following command. 
+You can [create and store multiple device profiles](Platforms/DeviceProfiles) by creating a new configuration (.ini) file called **DefaultDeviceProfiles.ini** in your project's 
+config folder.
 
-	 r.MobileContentScaleFactor = 
+Inside of this file, you can specify what will happen to the resolution of your project for specific devices by entering the following command followed by one of the values from the iOS and Android sections below:
 
-In the following image, you can see an example of the Mobile Content Scale Factor in action using the command listed above.
+	r.MobileContentScaleFactor
 
+In the following image, you can see an example of the Mobile Content Scale Factor being set in a configuration file:
 
-[OBJECT:ToggleButtonContent]
-[PARAMLITERAL:category]
-OS
-[/PARAMLITERAL]
-[PARAMLITERAL:id]
-windows
-[/PARAMLITERAL]
-[PARAMLITERAL:active]
-[/PARAMLITERAL]
-[PARAM:content]
 [REGION:lightbox]
 [![](T_MCF_INI_File.png)(w:358 h:505)](T_MCF_INI_File.png)
 [/REGION]
-[/PARAM]
-[/OBJECT]
-
-[OBJECT:ToggleButtonContent]
-[PARAMLITERAL:category]
-OS
-[/PARAMLITERAL]
-[PARAMLITERAL:id]
-mac
-[/PARAMLITERAL]
-[PARAMLITERAL:active]
-[/PARAMLITERAL]
-[PARAM:content]
-[REGION:lightbox]
-[![](T_MCF_INI_File_Mac.png)(w:358 h:505)](T_MCF_INI_File_Mac.png)
-[/REGION]
-[/PARAM]
-[/OBJECT]
 [REGION:caption]
 Click for full image.
 [/REGION]
 
-The image above was taken from the UE4 Tappy Chicken project and shows what will happen to the resolution of Tappy Chicken when it is played on various mobile devices.
-The top section of the file handles resolution scaling for iOS devices, and the bottom section handles resolution scaling for Android devices.
-Make sure to note that each **r.MobileContentScaleFactor** has a number after it.
-The number that is used after the command is what determines if the projects resolution should be scaled up or down.
+This particular configuration file is from the Tappy Chicken project and shows what will happen to the resolution of Tappy Chicken when it is played on various mobile devices. The top section of
+this file handles the resolution scaling for iOS devices and the bottom section handles resolution scaling for Android devices. Take note that each **r.MobileContentScaleFactor** has a number after it. The number
+is used after the command to determine if the project's resolution should be scaled up or down. 
 
 
 ### Mobile Content Scale Factor for iOS
-When using Mobile Content Scale Factor for iOS devices inputting the following number after the r.MobileContentScaleFactor will produce the following results. Please note that the actual resolution for any scale factor other than 0.0 will be corrected to match the aspect ratio of the screen and clamped to native resolution.
 
-[region:note]
-The iOS scale factor is directly related to Apple's scale factor system.
-[/region]
+For iOS devices, inputting the following numbers will produce the following results:
+ 
+[REGION:note]
+The iOS scale factor is directly related to Apple's scale factor system and the actual resolution for any scale factor other than 0.0 will be corrected to match the aspect ratio of the screen and clamped to 
+the native resolution of the device.
+[/REGION]
+
+[REGION:simpletable]
+| Value | Result |
+|------ | ------ |
+| **0.0** | This will use the native resolution of the device. |
+| **1.0** | On a Retina device is the non-retina resolution. |
+| **2.0** | Full native resolution on iPhone 5S, iPad Air, etc. |
+| **3.0** | Full native resolution for iPhone 6+. |
+[/REGION]
+
 
-|Command| Value|Result|
-|-------|------|------|
-|r.MobileContentScaleFactor|0.0|This will use the native resolution of the device.|
-|r.MobileContentScaleFactor|1.0|On a Retina device is the non-retina resolution.|
-|r.MobileContentScaleFactor|2.0|Full native resolution on iPhone 5S, iPad Air, etc.|
-|r.MobileContentScaleFactor|3.0|Full native resolution for iPhone 6+.|
 
 ### Mobile Content Scale Factor for Android
-When using the Mobile Content Scale Factor for Android devices inputting the following number after the r.MobileContentScaleFactor will produce the following results. Please note that inputting value other than **0.0** will use this value as the scale factor for a standard 1280x720 or 720x1280 resolution depending on orientation.
 
-|Command| Value|Result|
-|-------|------|------|
-|r.MobileContentScaleFactor|0.0|This will use the native resolution of the device.|
-|r.MobileContentScaleFactor|1.0|This will try and give you a resolution of 1280 x 720 for Landscape and 720 x 1280 for portrait.|
-|r.MobileContentScaleFactor|2.0|This will try and give you a resolution of 2560 x 1440 for Landscape and 1440 x 2560 for portrait.|
+For Android devices, inputting the following numbers will produce the following results: 
+
+[REGION:note]
+Please note that inputting values other than 0.0 will use this value as the scale factor for a standard 1280x720 or 720x1280 resolution depending on the orientation of the device. 
+[/REGION]
 
+[REGION:simpletable]
+| Value | Result |
+| ----- | ------ |
+| **0.0** | This will use the native resolution of the device.|
+| **1.0** | This will try and give you a resolution of 1280 x 720 for Landscape and 720 x 1280 for portrait.|
+| **2.0** | This will try and give you a resolution of 2560 x 1440 for Landscape and 1440 x 2560 for portrait.|
+[/REGION]
 
 ## Mobile Material Quality Settings
-When building Content for an Unreal Engine 4 (UE4) project that will run on both low and high end mobile devices you will often run into issues with certain features working on one set of devices but not working on another set of devices. You also might encounter issues with your artwork looking correct on one device but completely different on another. While there are a lot of ways to address problems like this, many of theses ways are both time and resource intensive and also extremely error prone. To address these issues UE4 has the Material Quality Level system. This system allows you to build one single Material that can then be used on a wide range of different devices, giving you full control over which devices use which features. In the following sections we will go over these systems and how you can use them in your UE4 projects:
+
+When building content for a UE4 project that will run on both low and high-end mobile devices, you will often run into issues where a feature or your artwork is working on one set of devices 
+but not another. While there are a lot of ways to address issues like this, many of these are both time and resource intensive which can be error prone at times. To address these types of issues
+UE4 has the **Material Quality Level** system. This system enables you to build one single Material that can then be used on a wide range of devices, giving you full control over which devices 
+use which features. 
+
+In the following sections, we'll cover these systems and how you can use them in your own UE4 projects:
 
 ### Previewing Material Quality Level
-You can view what the different Material Quality level settings will look like inside of the UE4 Editor by going to the
- **Main Toolbar** and then going to **Settings** > **Material Quality Level** and selecting the quality level you want to preview.
+
+You can view what the different Material Quality Level settings will look like inside the Editor by going to the **Main Toolbar** and then then selecting **Settings** > **Material Quality Level** to choose
+the level you want to preview.
 
 [REGION:lightbox]
 [![](MQS_Material_Preview_Level.png)(w:332 h:324)](MQS_Material_Preview_Level.png)
@@ -237,57 +309,35 @@ The following images show what a Material will look like when the Material Quali
 [/REGION]
 
 ### Setting Material Quality Level
-You can set the Material Quality Level in the following ways inside UE4:
 
-**From the Console**: You can set the quality level that your Materials will use in the Editor by pressing the Backtick key to open up the UE4 console and entering `r.MaterialQualityLevel` followed by one of the following numbers.
+To set the Material Quality Level for your devices, you can do so in the following ways: 
+
+##### From the Console 
+
+Bring up the console by hitting the **`** (backtick) key and entering **r.MaterialQualityLevel** followed by one of the following values:
 
 ![](Platforms\Mobile\Performance\MQS_Console_Command.png)
 
+[REGION:simpletable]
 |Material Quality Level Number| Setting|
-|-------------------------------------|----------|
-|0| Low|
-|1|High|
-|2|Medium|
+|-----------------------------|----------|
+| **0** | Low|
+| **1** |High|
+| **2** |Medium|
+[/REGION]
 
 [region:note]
-Please note that 1 for high and 2 for medium **IS** correct and not a typo!
+Please note that 1 for high and 2 for medium is correct and not a typo.
 [/region]
 
-**From Device Profiles**: You can also adjust the Material Quality by doing the following:
+##### From Device Profiles
 
-1. From the Main Toolbar go to **Windows** > **Developer Tools** and then click on **Device Profiles**.  
+Using a configuration (.ini) file, you can set the Material Quality for the devices you plan to target using the following method from within the Editor:
 
-	[OBJECT:ToggleButtonContent]
-	[PARAMLITERAL:category]
-	OS
-	[/PARAMLITERAL]
-	[PARAMLITERAL:id]
-	windows
-	[/PARAMLITERAL]
-	[PARAMLITERAL:active]
-	[/PARAMLITERAL]
-	[PARAM:content]
+1. From the Main Toolbar go to **Windows** > **Developer Tools** and then click on **Device Profiles**.  
 	[REGION:lightbox]
 	[![](MQS_Device_Profiles.png)(w:217 h:371)](MQS_Device_Profiles.png)
 	[/REGION]
-	[/PARAM]
-	[/OBJECT]
-
-	[OBJECT:ToggleButtonContent]
-	[PARAMLITERAL:category]
-	OS
-	[/PARAMLITERAL]
-	[PARAMLITERAL:id]
-	mac
-	[/PARAMLITERAL]
-	[PARAMLITERAL:active]
-	[/PARAMLITERAL]
-	[PARAM:content]
-	[REGION:lightbox]
-	[![](MQS_Device_Profiles_Mac.png)(w:217 h:371)](MQS_Device_Profiles_Mac.png)
-	[/REGION]
-	[/PARAM]
-	[/OBJECT]
 	[REGION:caption]
 	Click for full image.
 	[/REGION]
@@ -315,8 +365,9 @@ Please note that 1 for high and 2 for medium **IS** correct and not a typo!
 	[/REGION]
 
 ### Previewing Different Quality Levels
-You preview from the UE4 Editor how your project will look when running on different hardware by adjusting the **Preview Rendering Level**. You can adjust this by going to the **Main Toolbar** and then go to **Settings** > **Preview Rendering Level** and select the preview level you want to use.
 
+By adjusting the **Preview Rendering Level**, you can use the Unreal Engine 4 Editor to see how your project will look when running on different hardware. You can adjust this by going to the 
+**Main Toolbar** and then select **Settings** > **Preview Rendering Level** and select the preview level you want to use.
 
 [REGION:lightbox]
 [![](MQS_Material_Quality_Level.png)(w:332 h:324)](MQS_Material_Quality_Level.png)
@@ -326,11 +377,14 @@ Click for full image.
 [/REGION]
 
 ### Material Quality Level Switch
-The **Quality Switch** Material Expression allows you to define different levels of complexity within a single Material. For example, if you have a complex math operation, or heavy texture reads that will work on higher end mobile devices but not on lower end mobile devices you can use this node to specify simplified versions to be displayed on lower end mobile devices that do not have the power to display higher end graphics.
+
+The **Quality Switch** Material Expression enables you to define different levels of complexity within a single Material. For example, if you have a complex math operation, or heavy texture reads that will work on 
+higher-end mobile devices but not on lower-end ones. You can use this node to specify simplified versions to be displayed on lower-end mobile devices that do not have the power to display higher-end graphics. 
 
 ![](image08.png)
 
-To use the Quality Switch in your Materials, you will need to first add the Quality Switch Material Expression node to the Material Graph and then plug it's output into any of the inputs on the Main Material node you want to change when the quality level is changed. In the following image, the M_Brick_Clay_Old Material that can be found with the Starter Content was set up so that when the Quality Level is switched the complexity of the Material is reduced. 
+To use the Quality Switch in your Materials, you will need to first add the **Quality Switch Material Expression** node to the Material Graph and then plug its output into any of the inputs on the Main Material
+node. In the following image, the M_Brick_Clay_Old Material (that can be found in the Starter Content) is setup so that when the Quality Level is switched, the complexity of the Material will be reduced.
 
 [REGION:lightbox]
 [![](MQS_Material_Setup.png)(w:789 h:412)](MQS_Material_Setup.png)
@@ -340,10 +394,16 @@ Click for full image.
 [/REGION]
 
 ### Mobile Material Options 
-If you open up any **Material** and click on the Main Material Node then look in the Details Panel under the **Mobile** section you will find the following two properties.
 
-* **Fully Rough**: Forces a Material to be completely rough saving a number of instructions and a single texture sampler.
-* **Use Lightmap Directionality**: Uses lightmap directionally and per pixel normals. If disabled, the lightmaps will be flatter but cheaper.
+If you open up any **Material** and click on the Main Material Node then look in the Details Panel under the **Mobile** section you will find the following two properties:
+
+[REGION:simpletable]
+| Property | Setting |
+| -------- | ------- |
+| **Fully Rough** | Forces a Material to be completely rough saving a number of instructions and a single texture sampler. |
+| **Use Lightmap Directionality** | Uses lightmap directionality and per-pixel normals. If disabled, the lightmaps will be flatter looking visually but cheaper to use. |
+[/REGION]
+
 
 [REGION:lightbox]
 [![](MQS_Mobile_Material_Options.png)(w:556 h:410)](MQS_Mobile_Material_Options.png)
@@ -352,130 +412,74 @@ If you open up any **Material** and click on the Main Material Node then look in
 Click for full image.
 [/REGION]
 
-These two properties can be enabled to help reduce the rendering cost of that Material when viewed on less powerful mobile devices by completely eliminating the rendering path for those features. However, this is, an all or nothing option as you either have these options enabled or disabled for this Material on all devices.
+These two properties can be enabled to help reduce the rendering cost of that Material when viewed on less powerful mobile devices by completely eliminating the rendering path for those features. 
+However, this is, an all-or-nothing option as you either have these options enabled or disabled for this Material on all devices.
 
 
 ### Mobile Material Rendering Overrides
-You can override the available rendering option a platform will use by adjusting the **ES2 Quality Settings**. You can find the ES2 Quality Setting for both Android and iOS by going to your **Project Settings** > **Platforms** and then selecting either **Android ES2 Quality** or **iOS ES2 Quality Settings**. To use the overrides, you must first click on the **Enable Quality Override** option and then select the option you want to override by clicking on it. Once all of the options you want have been selected, press the **Update preview shaders** button to re-compile all Materials to work with the options you have specified. 
-
-[OBJECT:ToggleButtonContent]
-[PARAMLITERAL:category]
-OS
-[/PARAMLITERAL]
-[PARAMLITERAL:id]
-windows
-[/PARAMLITERAL]
-[PARAMLITERAL:active]
-[/PARAMLITERAL]
-[PARAM:content]
-[REGION:lightbox]
-[![](MQS_Rendering_Overrides.png)(w:613 h:330)](MQS_Rendering_Overrides.png)
-[/REGION]
-[/PARAM]
-[/OBJECT]
 
-[OBJECT:ToggleButtonContent]
-[PARAMLITERAL:category]
-OS
-[/PARAMLITERAL]
-[PARAMLITERAL:id]
-mac
-[/PARAMLITERAL]
-[PARAMLITERAL:active]
-[/PARAMLITERAL]
-[PARAM:content]
-[REGION:lightbox]
-[![](MQS_Rendering_Overrides_Mac.png)(w:613 h:330)](MQS_Rendering_Overrides_Mac.png)
-[/REGION]
-[/PARAM]
-[/OBJECT]
-[REGION:caption]
-Click for full image.
-[/REGION]
+You can override the available rendering option a platform will use by adjusting the the Material Quality settings for different target devices located under the **Project Settings** > **Platforms**.
+
+
+![](PS_QualityLevels.png)
+
+
+**Android:**
+* Open GLES 2
+* Open GLES 3.1
+* Vulkan 
+
+**iOS:**
+* Open GLES 2
+* Metal
+
+You can override the available rendering option a platform will use by adjusting the **Quality Options** listed. In order to use the overrides, you must first click on the **Enable Quality Overrides** option
+and then select the option you want to override for a given quality level. 
+
+![](EnableQualityLevelOverrides.png)
+
+Once all of the options you want have been selected, press the **Update Preview Shaders** button to recompile all the Materials to work with the options specified.
+
+![](QualityLevelsUpdateShaders.png)
+
+
 
 ## Mobile Rendering Options
-Inside of your **Project's Settings** menu under the **Mobile** section you will find a number of properties that control  how lighting will be handled inside of your mobile project. In the following section we will take a look at these settings and what effect they will have on your project.
-
-### Accessing your Project Settings
-To access your project settings you will first need to go to the **Main Toolbar** and then select **Edit** > **Project Settings**. 
-
-[OBJECT:ToggleButtonContent]
-[PARAMLITERAL:category]
-OS
-[/PARAMLITERAL]
-[PARAMLITERAL:id]
-windows
-[/PARAMLITERAL]
-[PARAMLITERAL:active]
-[/PARAMLITERAL]
-[PARAM:content]
+
+Inside of your **Project's Settings** menu under the **Rendering** section, you will find a number of properties that control how lighting and optimization will be handled inside of 
+your mobile project. In the following sections we will take a look at these settings and what effect they will have on your project. 
+
+### Accessing your Project's Settings
+To access your project's settings you will first need to go to the **Main Toolbar** and then select **Edit** > **Project Settings**. 
+
 ![](MRO_Project_Settings.png)
-[/PARAM]
-[/OBJECT]
 
-[OBJECT:ToggleButtonContent]
-[PARAMLITERAL:category]
-OS
-[/PARAMLITERAL]
-[PARAMLITERAL:id]
-mac
-[/PARAMLITERAL]
-[PARAMLITERAL:active]
-[/PARAMLITERAL]
-[PARAM:content]
-![](MRO_Project_Settings_Mac.png)
-[/PARAM]
-[/OBJECT]
+Once the Project Settings menu is open, locate the **Engine** section and then click on the **Rendering** category. Use the following sections to adjust the settings below: 
 
-Once the Project Settings menu is open, locate the **Engine** section and then click on the **Rendering** category and look for the **Mobile** heading. 
-
-[OBJECT:ToggleButtonContent]
-[PARAMLITERAL:category]
-OS
-[/PARAMLITERAL]
-[PARAMLITERAL:id]
-windows
-[/PARAMLITERAL]
-[PARAMLITERAL:active]
-[/PARAMLITERAL]
-[PARAM:content]
-[REGION:lightbox]
-[![](MRO_Rendering_Settings.png)(w:430)](MRO_Rendering_Settings.png)
-[/REGION]
-[/PARAM]
-[/OBJECT]
+* Mobile
+* Lighting
+* Mobile Shader Permutation
 
-[OBJECT:ToggleButtonContent]
-[PARAMLITERAL:category]
-OS
-[/PARAMLITERAL]
-[PARAMLITERAL:id]
-mac
-[/PARAMLITERAL]
-[PARAMLITERAL:active]
-[/PARAMLITERAL]
-[PARAM:content]
 [REGION:lightbox]
-[![](MRO_Rendering_Settings_Mac.png)(w:430)](MRO_Rendering_Settings_Mac.png)
+[![](MRO_Rendering_Settings.png)(w:625)](MRO_Rendering_Settings.png)
 [/REGION]
-[/PARAM]
-[/OBJECT]
+
 [REGION:caption]
 Click for full image.
 [/REGION]
 
 ### Mobile Rendering Options
-Under the Mobile section you will find a number of different options that can be enabled or disabled. In the following section we will take a look each one of the properties and the effect they will have on your mobile project.
+
+Under the **Mobile** category, you can set properties that directly affect the lighting and look of your mobile project. Read on below to explore these settings for your own projects.
 
 ![](MRO_Mobile_Settings.png)
 
 |Property Name| Description|
 |--------------------|-------|
-|Mobile HDR|If true, mobile renders in full HDR. Disable this setting for games that do not require lighting features for better performance on slow devices.|
-|Max Dynamic Point Lights|The number of dynamic point lights to support on mobile devices. Setting this to 0 for games which do not require dynamic point lights will reduce the number of shader generated. Changing this setting requires restarting the editor.|
-|Use Shared Dynamic Point Light Shaders|If this setting is enabled, the same shader will be used for any number of dynamic point lights(up to the maximum specified above)hitting a surface. This is slightly slower but reduced the number of shaders generated. Changing this setting requires restarting the editor. |
-|Enabled Combined Static and CSM Shadowing|Allows primitives to receive both static and CSM shadows from a stationary light. Disabling will free a mobile texture sampler.|
-
+| **Mobile HDR** |If true, mobile renders in full HDR. Disable this setting for games that do not require lighting features for better performance on slow devices.|
+| **Disable Vertex Fogging in mobile shaders** | If true, vertex fog will be omitted from all mobile shaders. If your game does not use fog, you should choose this setting to increase shading performance. |
+| **Maximum number of CSM cascades to render** | The maximum number of cascades with which to render dynamic directional light shadows when using the mobile renderer. |
+| **Mobile MSAA** | This is the Multi-Sample Anti-Aliasing (MSAA) setting to use on mobile. MSAA is currently supported using Metal on iOS, and on Android devices with the required supported using ES2 or ES 3.1. If MSAA is not available, the current default AA method will be used. |
 
 <!-----
 [EXCERPT:MMQ]
@@ -490,7 +494,7 @@ Under the Mobile section you will find a number of different options that can be
 	MMQ
 	[/PARAMLITERAL]
 	[PARAM:caption]
-	**Dragging the slider will show what happens when changing the Material Quality Level from Low to Medium and High **
+	Dragging the slider will show what happens when changing the Material Quality Level from Low to Medium to High 
 	[/PARAM]
 	[PARAM:image_list]
 	![](MMQ_1.png)
@@ -499,4 +503,76 @@ Under the Mobile section you will find a number of different options that can be
 	[/PARAM]
 [/OBJECT]
 [/EXCERPT:MMQ]
--->
\ No newline at end of file
+-->
+
+
+## Shader Reduction Settings for Mobile
+
+Under the **Lighting** and **Mobile Shader Permutation Reduction** section, you will find a number of different option that can be enabled or disabled to effect shaders generated for
+Materials, which in-turn will reduce app size, save memory and improve load times. 
+
+
+![](MobileOptimizations.png)
+
+[REGION:raw]
+|Property Name| %spacer3% | Description|
+|------------ | --------- |-------|
+| **Lighting** || The listed options found under the **Lighting** section of the Project Settings that should be adjusted for mobile performance.  [INCLUDE:#lighting] |
+| **Mobile Shader Permutation Reduction** || The listed options found under **Mobile Shader Permutation Reduction** that should be adjusted for mobile performance. [INCLUDE:#shaderpermutations] |
+[/REGION]
+
+<!--- EXCERPTS ---
+
+[INCLUDE:#shaderpermutations]
+
+[EXCERPT:lighting]
+[REGION:simpletable]
+|--------------------| --- |-------|
+| **Allow Static Lighting** || Whether to allow any static lighting to be generated and used, like lightmaps and shadow maps. Games that only use dynamic lighting or that are unlit should set this to 0 to save some static lighting overhead. |
+| | %spacer4% | | 
+[/REGION]
+[/EXCERPT:lighting]
+
+[EXCERPT:shaderpermutations]
+[REGION:simpletable]
+| ------------------ | ----- | -------|
+| **Support Combined Static and CSM Shadowing** || Allows primitives to receive both static and CSM shadows from a stationary light. Disabling will free a mobile texture sampler. |
+| **Support Distance Field Shadowing** || Generate shaders for primitives to receive distance field shadows from Stationary Directional Lights. |
+| **Support Movable Directional Light** || Generate shaders for primitives to receive movable directional lights. |
+| **Max Movable Point Lights** || The number of dynamic point lights to support on mobile devices. Setting this to 0 for games which do not require dynamic point lights will reduce the number of shaders generated. |
+| **Use Shared Movable Point Light Shaders** || If this setting is enabled, the same shader will be used for any number of dynamic point lights(up to the maximum specified above)hitting a surface. This is slightly slower but reduced the number of shaders generated. |
+|| %spacer3% |||
+[/REGION]
+[/EXCERPT:shaderpermutations]
+
+----->
+
+[REGION:note]
+Changing any of these settings will require you to restart the editor before they take effect. 
+[/REGION]
+
+[REGION:warning]
+If you disable a shader permutation that your lighting setup requires, your lighting or shadowing will not render correctly and there will be an on-screen message indicating 
+the setting that needs to be re-enabled for it to work properly.
+[/REGION]
+
+
+
+## Additional Suggestions
+
+The following suggestions will help to set up your content for optimal performance on mobile platforms. These are listed from most severe to least impactful. 
+
+
+
+* Make sure lighting is built before running on a device.
+* Most Post Processing features are disabled for mobile platforms as they are too expensive, like like vignette or screen space reflections. Some GPU intensive features like Bloom and Depth of Field are enabled by default so that PC and mobile look the same but many features can cost 60 milliseconds or more with the default settings on some devices. The console command **showflag.PostProcessing 0** can be used to turn off these features to get a general idea of their costs.
+* Make sure to use Precomputed Visibility and that it is set up properly. To do this, place **Precomputed Visibility Volumes** around areas where the player can walk or jump and then build lighting. You need to make sure to place these in the Persistent level that is used when building lighting and when running the game, so do not build a sublevel by itself. You can verify that Precomputed Visibility is working by entering the console command **Stat Initviews** on device or in the level previewer and making sure that **Statically Occluded Primitives** is greater than 0. Use the console command **r.ShowPrecomputedVisibilityCells 1** to visualize the cells in the editor. 
+* Use masked and transparent Materials sparingly. Only use them in places where they cover a small part of the screen. iOS devices are very optimal in shading opaque surfaces, as they will only shade each pixel once, but for masked and translucency, every single layer will have to be shaded. The total GPU time of a frame can be doubled or more by having overdraw.  Use the **Shader Complexity** view mode to investigate any hot spots.
+* Draw calls of the entire scene should be &lt;=700 for any single view. Areas with poor occlusion, like looking over a large landscape, is going to be biggest challenges for this. This can be seen with **Stat OpenGLRHI** on device or **Stat D3D11RHI** in the Previewer on PC.
+* Triangle count of the entire scene should be &lt;=500k for any view. This has been determined to be the maximum poly count that can hit 30fps on both iPad4 and iPad Air. This can be seen with **Stat OpenGLRHI** on device or **Stat D3D11RHI** in the Previewer on PC.
+* Materials should use as few texture lookups and instructions as possible. You have five texture samplers available, but using all of those makes for a fairly expensive Materials. For now, use the PC instruction
+count to optimize and the Previewer shader complexity to visualize the total cost. 
+* Only use independent texture fetches in Materials. This means UVs in the pixel shader (BaseColor, Roughness, etc.) must not be manipulated in any way, such as scaling. Use the **CustomizedUVs** feature instead to do the scaling in the Vertex Shader. Some special features like environment mapping require math on the UVs and that is OK for special cases. 
+* It is recommended to use square textures as they have less memory waste. Always use power of two dimensions for these (256x256, 512x512, 1024x1024). Use the console command **ListTextures** in the Previewer to see where all the texture memory is going. 
+
+
diff --git a/Engine/Documentation/Source/Platforms/Mobile/PostProcessEffects/MobilePostProcessEffects.INT.udn b/Engine/Documentation/Source/Platforms/Mobile/PostProcessEffects/MobilePostProcessEffects.INT.udn
--- a/Engine/Documentation/Source/Platforms/Mobile/PostProcessEffects/MobilePostProcessEffects.INT.udn
+++ b/Engine/Documentation/Source/Platforms/Mobile/PostProcessEffects/MobilePostProcessEffects.INT.udn
@@ -3,7 +3,7 @@ Title:Post Process Effects on Mobile Platforms
 Crumbs:
 Description:Using Post Process Effects on mobile games.
 Platform:Mobile
-Version: 4.14
+Version: 4.16
 parent:Platforms/Mobile
 order:5
 type:overview
@@ -16,37 +16,9 @@ reference-image:Mobile_PostProcessEffect.gif
 Mobile Post Processing is implemented as a separate code path from PC/Console Post Processing. This is to reflect the slower dependent texture reads, missing hardware features, extra render target resolve costs, and slower performance in general. 
 To enable Post Processing for your mobile applications, please make sure that **Mobile HDR** is enabled in your **Project Settings** as most of the settings below require Mobile HDR to be enabled for the effect(s) to work.
 
-[OBJECT:ToggleButtonContent]
-[PARAMLITERAL:category]
-OS
-[/PARAMLITERAL]
-[PARAMLITERAL:id]
-windows
-[/PARAMLITERAL]
-[PARAMLITERAL:active]
-[/PARAMLITERAL]
-[PARAM:content]
 [region: lightbox]
 [![](MobileHDR_Enabled.png)(w:400)](MobileHDR_Enabled.png)
 [/region]
-[/PARAM]
-[/OBJECT]
-
-[OBJECT:ToggleButtonContent]
-[PARAMLITERAL:category]
-OS
-[/PARAMLITERAL]
-[PARAMLITERAL:id]
-mac
-[/PARAMLITERAL]
-[PARAMLITERAL:active]
-[/PARAMLITERAL]
-[PARAM:content]
-[region: lightbox]
-[![](MobileHDR_Enabled_Mac.png)(w:400)](MobileHDR_Enabled_Mac.png)
-[/region]
-[/PARAM]
-[/OBJECT]
 
 [region:caption]
 Click for full image.
@@ -66,10 +38,10 @@ To use the Tonemapper Post Processing options on Mobile Devices, you will need t
 
 
 ### Enabling Tonemapper Film via the Console
-You can enable the Tonemapper in the UE4 Editor by first opening up the console by pressing the **Backtick key(`)** and then entering the following command to enable the Tonemapper.
+You can enable the Tonemapper in the UE4 Editor by first opening up the console by pressing the **Backtick key(`)** and then entering the following command:
 * **r.tonemapperfilm 1** 
 
-To disable the Tonemapper, open up the console and enter the following command to disable the Tonemapper.
+To disable the Tonemapper, open up the console and enter the following command:
 *  **r.tonemapperfilm 0**. 
 
 [region:note]
@@ -85,78 +57,19 @@ Click for full image.
 [/REGION]
 
 ### Enabling Tonemapper Film via .INI Files
-To enable or disable the Tonemapper project wide you will need to add **r.tonemapperfilm=1** to the **Rendering Settings** in your project's **DefaultEngine.ini**.
-
-
+To enable or disable the Tonemapper project wide you will need to add **r.tonemapperfilm=1** to the **Rendering Settings** section of your project's **DefaultEngine.ini** file.
 
-[OBJECT:ToggleButtonContent]
-[PARAMLITERAL:category]
-OS
-[/PARAMLITERAL]
-[PARAMLITERAL:id]
-windows
-[/PARAMLITERAL]
-[PARAMLITERAL:active]
-[/PARAMLITERAL]
-[PARAM:content]
 [REGION:lightbox]
 [![](TM_Default_Engine_INI.png)(w:422 h:199)](TM_Default_Engine_INI.png)
 [/REGION]
-[/PARAM]
-[/OBJECT]
-
-[OBJECT:ToggleButtonContent]
-[PARAMLITERAL:category]
-OS
-[/PARAMLITERAL]
-[PARAMLITERAL:id]
-mac
-[/PARAMLITERAL]
-[PARAMLITERAL:active]
-[/PARAMLITERAL]
-[PARAM:content]
-[REGION:lightbox]
-[![](TM_Default_Engine_INI_Mac.png)(w:422 h:199)](TM_Default_Engine_INI_Mac.png)
-[/REGION]
-[/PARAM]
-[/OBJECT]
 	
 [REGION:caption]
 Click for full image.
 [/REGION]
 
-[OBJECT:ToggleButtonContent]
-[PARAMLITERAL:category]
-OS
-[/PARAMLITERAL]
-[PARAMLITERAL:id]
-windows
-[/PARAMLITERAL]
-[PARAMLITERAL:active]
-[/PARAMLITERAL]
-[PARAM:content]
 [REGION:lightbox]
 [![](TM_INI_Location.png)(w:422 h:199)](TM_INI_Location.png)
 [/REGION]
-[/PARAM]
-[/OBJECT]
-
-[OBJECT:ToggleButtonContent]
-[PARAMLITERAL:category]
-OS
-[/PARAMLITERAL]
-[PARAMLITERAL:id]
-mac
-[/PARAMLITERAL]
-[PARAMLITERAL:active]
-[/PARAMLITERAL]
-[PARAM:content]
-[REGION:lightbox]
-[![](TM_INI_Location_Mac.png)(w:422 h:199)](TM_INI_Location_Mac.png)
-[/REGION]
-[/PARAM]
-[/OBJECT]
-
 	
 [REGION:caption]
 Click for full image.
@@ -184,79 +97,109 @@ In the following section we will go over what Post Process Effects and settings
 Click for full image.
 [/REGION]
 
-### White Balance
-|Post Process Property| Supported|R.tonemapperfilm Value| 
-|--------------------|-----------|--------------|
+### Color Grading
+
+In the **Color Grading** section, you'll find the supported properties that enable you to make color corrections to the image.
+
+|Post Process Property| Supported |R.tonemapperfilm Value| 
+|---------------------|-----------|----------------------|
+| [REGION:tablesection] White Balance [/REGION] |||
 |**Temp**|Yes|1|
 |**Tint**|Yes|1|
-
-### Color Grading
-|Post Process Property| Supported|R.tonemapperfilm Value| 
-|--------------------|-----------|--------------|
+| [REGION:tablesection] Global [/REGION] |||
+|**Saturation**|Yes|1|
+|**Contrast**|Yes|1|
+|**Gamma**|Yes|1|
+|**Gain**|Yes|1|
+|**Offset**|Yes|1|
+|**Scene Color Tint**|Yes|1|
+|**Color Grading LUT Intensity**|Yes|1|
+|**Color Grading LUT**|Yes|1|
+| [REGION:tablesection] Shadows [/REGION] |||
+|**Saturation**|Yes|1|
+|**Contrast**|Yes|1|
+|**Gamma**|Yes|1|
+|**Gain**|Yes|1|
+|**Offset**|Yes|1|
+|**Shadows Max**|Yes|1|
+| [REGION:tablesection] Midtones [/REGION] |||
+|**Saturation**|Yes|1|
+|**Contrast**|Yes|1|
+|**Gamma**|Yes|1|
+|**Gain**|Yes|1|
+|**Offset**|Yes|1|
+| [REGION:tablesection] Highlights [/REGION] |||
 |**Saturation**|Yes|1|
 |**Contrast**|Yes|1|
 |**Gamma**|Yes|1|
 |**Gain**|Yes|1|
 |**Offset**|Yes|1|
+|**Highlights Min**|Yes|1|
+
+
+### Tonemapper
+
+In the **Tonemapper** section, you'll find the supported properties that enable you to make adjustments to the image color range.
+
+
+|Post Process Property| Supported|R.tonemapperfilm Value| 
+|--------------------|-----------|--------------|
+|**Slope**|Yes|1|
+|**Toe**|Yes|1|
+|**Shoulder**|Yes|1|
+|**Black Clip**|Yes|1|
+|**White Clip**|Yes|1|
+
+### Mobile Tonemapper
+
+In the **Mobile Tonemapper** section, you'll find the supported properties specific to mobile that enable you to make adjustments to the image color range.
 
-### Film
 |Post Process Property| Supported|R.tonemapperfilm Value| 
 |--------------------|-----------|--------------|
 |**Tint**|Yes|0|
+|**Saturation**|Yes|0|
+|**Contrast**|Yes|0|
 |**Tint Shadow**|Yes|0|
 |**Tint Shadow Blend**|Yes|0|
 |**Tint Shadow Amount**|Yes|0|
-|**Saturation**|Yes|0|
 |**Channel Mixer Red**|Yes|0|
 |**Channel Mixer Green**|Yes|0|
 |**Channel Mixer Blue**|Yes|0|
-|**Contrast**|Yes|0|
 |**Crush Shadows**|Yes|0|
 |**Crush Highlights**|Yes|0|
 |**Dynamic Range**|Yes|0|
-|**Slope**|Yes|1|
-|**Toe**|Yes|1|
-|**Shoulder**|Yes|1|
-|**Black Clip**|Yes|1|
-|**White Clip**|Yes|1|
 
-### Scene Color
+
+### Lens
+
+In the **Lens** section, you'll find the supported properties that enable you to make different visual adjustments to the image for features like Bloom and Depth of Field.
+
 |Post Process Property| Supported|R.tonemapperfilm Value| 
 |--------------------|-----------|--------------|
-|**Scene Color Tint**|Yes|1|
-|**Fringe Intensity**|Yes|0 & 1|
+| [REGION:tablesection] Image Effects [/REGION] |||
+|**Chromatic Aberation**|Yes|0 & 1|
 |**Vignette Intensity**|Yes|0 & 1|
 |**Grain Jitter**|Yes|0 & 1|
-|**Grain Intensity**|Yes|0 & 1|
 |**Color Grading Intensity**|Yes|1|
-|**Color Grading**|Yes|1|
-
-### Bloom
-|Post Process Property| Supported|R.tonemapperfilm Value| 
-|--------------------|-----------|--------------|
+|[REGION:tablesection] Bloom [/REGION] |||
+|**Method: Standard**|Yes|0 & 1|
+|**Method: Convolution**|No|n/a|
 |**Intensity**|Yes|0 & 1|
 |**Threshold**|Yes|0 & 1|
 |**Size Scale**|Yes|0 & 1|
-|**#1 Size**|No|n/a|
-|**#2 Size**|No|n/a|
-|**#3 Size**|No|n/a|
-|**#4 Size**|No|n/a|
-|**#5 Size**|No|n/a|
-|**#1 Tint**|No|n/a|
-|**#2 Tint**|No|n/a|
-|**#3 Tint**|No|n/a|
-|**#4 Tint**|No|n/a|
-|**#5 Tint**|No|n/a|
-|**#6 Tint**|No|n/a|
+|**#123456 Size**|No|n/a|
+|**#123456 Tint**|No|n/a|
+|**Convolution Scale**|No|n/a|
+|**Convolution Center**|No|n/a|
+|**Convolution Boost**|No|n/a|
+|**Convolution Buffer**|No|n/a|
+|[REGION:tablesection] Dirt Mask [/REGION] |||
+|**Dirt Mask Texture**|Yes|1|
 |**Dirt Mask Intensity**|Yes|1|
 |**Dirt Mask Tint**|Yes|1|
-|**Dirt Mask**|Yes|1|
-
-### Auto Exposure
-|Post Process Property| Supported|R.tonemapperfilm Value| 
-|--------------------|-----------|--------------|
-|**Auto Exposure Histogram**|Yes|0 & 1|
-|**Auto Exposure Basic**|Yes|0 & 1|
+|[REGION:tablesection] Auto Exposure [/REGION] |||
+|**Auto Exposure Histogram**|No|n/a|
+|**Auto Exposure Basic**|No|n/a|
 |**Low Percent**|No|n/a|
 |**High Percent**|No|n/a|
 |**Min Brightness**|No|n/a|
@@ -266,10 +209,18 @@ Click for full image.
 |**Exposure Bias**|Yes|0 & 1|
 |**Histogram Log Min**|No|n/a|
 |**Histogram Log Max**|No|n/a|
-
-### Depth Of Field
-|Post Process Property| Supported|R.tonemapperfilm Value| 
-|--------------------|-----------|--------------|
+|**Low Percent**|No|n/a|
+|**High Percent**|No|n/a|
+|**Histogram Log Min**|No|n/a|
+|**Histogram Log Max**|No|n/a|
+|[REGION:tablesection] Lens Flares [/REGION] |||
+|**Intensity**|No|n/a|
+|**Tint**|No|n/a|
+|**Bokeh Size**|No|n/a|
+|**Threshold**|No|n/a|
+|**Bokeh Shape**|No|n/a|
+|**Tints #012345678**|No|n/a|
+|[REGION:tablesection] Depth of Field [/REGION] |||
 |**BokehDOF**|Yes|0 |
 |**GaussianDOF**|Yes|0 |
 |**CircleDOF**|No|n/a |
@@ -288,31 +239,78 @@ Click for full image.
 |**Shape**|No|n/a|
 |**Occlusion**|No|n/a|
 |**Color Threshold**|No|n/a|
+|**Size Threshold**|No|n/a|
 |**Sky Distance**|No|n/a|
 |**Vignette Size**|No|n/a|
 
+### Rendering Features
+
+In the **Rendering Features** section, you'll find the supported properties that enable you to make visual adjustments to the image for specific Post Process Materials, like Custom Depth and Custom Stencil. 
 
-### Misc
 |Post Process Property| Supported|R.tonemapperfilm Value| 
 |--------------------|-----------|--------------|
-|**Screen Percentage**|No|n/a |
-|**AA Methods**|Yes|0 & 1|
-[region:note]
-The AA Method supports "None" (the fastest) and "TemporalAA" which provides a special mobile temporal AA which provides something around 2xSGSSAA with a little bit of judder in motion
-[/region]
+| [REGION:tablesection] Post Process Materials [/REGION] |||
+|**Array**|Yes|1|
+| [REGION:tablesection] Ambient Cubemap [/REGION] |||
+|**Tint**|No|n/a|
+|**Intensity**|No|n/a|
+|**Cubemap Texture**|No|n/a|
+| [REGION:tablesection] Ambient Occlusion [/REGION] |||
+|**Intensity**|No|n/a|
+|**Radius**|No|n/a|
+|**Static Fraction**|No|n/a|
+|**Radius in Worldspace**|No|n/a|
+|**Fade Out Distance**|No|n/a|
+|**Fade Out Radius**|No|n/a|
+|**Power**|No|n/a|
+|**Bias**|No|n/a|
+|**Quality**|No|n/a|
+|**Mip Blend**|No|n/a|
+|**Mip Scale**|No|n/a|
+|**Mip Threshold**|No|n/a|
+| [REGION:tablesection] Global Illumination [/REGION] |||
+|**Indirect Lighting Color**|Yes|0 & 1|
+|**Indirect Lighting Intensity**|Yes|0 & 1|
+| [REGION:tablesection] Motion Blur [/REGION] |||
+|**Amount**|No|n/a|
+|**Max**|No|n/a|
+|**Per Object Size**|No|n/a|
+| [REGION:tablesection] Light Propagation Volume [/REGION] |||
+|**Intensity**|No|n/a|
+|**Size**|No|n/a|
+|**Secondary Occlusion Intensity**|No|n/a|
+|**Occlusion Intensity**|No|n/a|
+|**Diffuse Occlusion Exponent**|No|n/a|
+|**Specular Occlusion Exponent**|No|n/a|
+|**Light Injection Bias**|No|n/a|
+|**Secondary Bounce Intensity**|No|n/a|
+|**Geometry Volume Bias**|No|n/a|
+|**Emissive Injection Intensity**|No|n/a|
+|**Occlusion Radius**|No|n/a|
+|**Diffuse Occlusion Intensity**|No|n/a|
+|**Specular Occlusion Intensity**|No|n/a|
+|**Fade Range**|No|n/a|
+|**Directional Occlusion Fade Range**|No|n/a|
+| [REGION:tablesection] Screen Space Reflections [/REGION] |||
+|**Intensity**|No|n/a|
+|**Quality**|No|n/a|
+|**Max Roughness**|No|n/a|
+| [REGION:tablesection] Misc [/REGION] |||
+|**Screen Percentage**|No|n/a|
+
 
 ## Custom Post Process Materials
 
-Unreal Engine ships with a variety of Post Process effects, allowing you to customize the look and feel of your mobile applications. There are times, though, when you'll want to create 
-custom Post Process Materials, enabling your own effects (such as "VCR Static").
+The Unreal Engine ships with a variety of Post Process effects, enabling you to customize the look and feel of your mobile applications. There are times, though, when you'll want to create 
+custom Post Process Materials to implement your own effects (such as "VCR Static").
 
 [REGION:asyncgif]
 ![image alt text](Mobile_PostProcessEffect.gif)
 [/REGION]
 
-Adding effects to your mobile application requires the development of custom Post Process materials. When creating custom Post Process materials for mobile applications, you should note that there 
-are some key things to consider. Specifically, it's important to note that custom Post Process materials for mobile applications can only fetch from PostProcessInput0 (Scene Color) with the 
-following **Blendable Locations**:
+Adding effects to your mobile application requires the development of custom Post Process materials. When creating custom Post Process Materials for mobile applications, you should note that there 
+are some key things to consider. Specifically, it's important to note that custom Post Process Materials for mobile applications can only fetch from **PostProcessInput0 (Scene Color)**, **Scene Depth**,
+**Custom Depth**, and **Custom Stencil** with the following **Blendable Locations**.
 
 * **Before Tonemapping**
 * **After Tonemapping**
@@ -320,66 +318,45 @@ following **Blendable Locations**:
 To learn more about Blendable Locations, read through our [](Engine/Rendering/PostProcessEffects/PostProcessMaterials) reference page. 
 Currently, you cannot create custom Post Process materials for older Android devices that require 'mosaic' mode for HDR rendering. Finally, pixel depth information is not yet supported.
 
-### Custom Depth
+As this features requires Post-Processing to work correctly, **Mobile HDR** must be enabled. To enable Mobile HDR go to **Project Settings** > **Engine** > **Rendering** > **Mobile** and make sure that **Mobile HDR** is checked.
 
-You can now use the **Custom Depth** Post Process option with the Mobile rendering path. Custom Post-Process Materials can now sample from **Scene Depth**, **Custom Depth** as well as **Scene Color**. 
+[REGION:lightbox]
+[![](MobileHDR_Enabled.png)(w:500)](MobileHDR_Enabled.png)
+[/REGION]
 
-![](Custom_Scene_Depth.png)
+[REGION:caption]
+Click for full image.
+[/REGION]
 
-As this features requires Post-Processing to work correctly, **Mobile HDR** must be enabled. To enable Mobile HDR go to **Project Settings** > **Engine** > **Rendering** > **Mobile** and make sure that **Mobile HDR** is checked.
 
-[OBJECT:ToggleButtonContent]
-[PARAMLITERAL:category]
-OS
-[/PARAMLITERAL]
-[PARAMLITERAL:id]
-windows
-[/PARAMLITERAL]
-[PARAMLITERAL:active]
-[/PARAMLITERAL]
-[PARAM:content]
-[region: lightbox]
-[![](MobileHDR_Enabled.png)(w:400)](MobileHDR_Enabled.png)
-[/region]
-[/PARAM]
-[/OBJECT]
+You now use **Custom Depth** Post Process option with the Mobile Rendering path. Custom Post Process Materails can now sample from **Scene Depth**, **Custom Depth**, **Scene Color**, and **Custom Stencil**.
 
-[OBJECT:ToggleButtonContent]
-[PARAMLITERAL:category]
-OS
-[/PARAMLITERAL]
-[PARAMLITERAL:id]
-mac
-[/PARAMLITERAL]
-[PARAMLITERAL:active]
-[/PARAMLITERAL]
-[PARAM:content]
-[region: lightbox]
-[![](MobileHDR_Enabled_Mac.png)(w:400)](MobileHDR_Enabled_Mac.png)
-[/region]
-[/PARAM]
-[/OBJECT]
+[REGION:imagetable]
+| [INCLUDE:#customdepth] | [INCLUDE:#customstencil] |
+|---|---|
+| Custom Depth | Custom Stencil  |
+[/REGION]
 
-[region:caption]
-Click for full image.
-[/region]
+<!---
+[EXCERPT:customdepth]
+[REGION:raw]
+![](Custom_Scene_Depth.png)(w:575)
+[/REGION]
+[/EXCERPT:customdepth]
+
+[EXCERPT:customstencil]
+[REGION:raw]
+![](Custom_Stencil.png)(w:575)
+[/REGION]
+[/EXCERPT:customstencil]
+--->
 
-[region:warning]
-Please note that this feature also does not currently work while **Mobile MSAA** is enabled so Mobile MSAA will have to be disabled for you to see the effect on the device.
-[/region]
 
 ## General Post Process Performance Tips
 
-* For best performance and quality, you should stick to using just Bloom and TemporalAA.
+* For best performance and quality, it is best to stick to using just Bloom and TemporalAA.
 
 [REGION:note]
-For better performance, we have implemented a special workaround which is controlled by the console variable "r.RenderTargetSwitchWorkaround". In **BaseDeviceProfiles.ini**, you can see it is enabled
-for several devices. Note that there will be a small change in appearance to the bloom effect on devices with certain performance characteristics. The console variables help text describes this in more detail:
-
-	r.RenderTargetSwitchWorkaround
-
-	Workaround needed on some mobile platforms to avoid a performance drop related to switching render targets.
-	Only enabled on some hardware. This affects the bloom quality a bit. It runs slower than the normal code path but
-	still faster as it avoids the many render target switches. (Default: 0)
-	We want this enabled (1) on all 32-bit iOS devices (implemented through DeviceProfiles).
+For better performance, we have implemented a special workaround which is controlled by the console variable "r.RenderTargetSwitchWorkaround". In **BaseDeviceProfiles.ini**, you can see that it is enabled
+for several devices. Note that there will be a small change in appearance to the bloom effect on devices with certain performance characteristics. 
 [/REGION]
diff --git a/Engine/Documentation/Source/Platforms/Mobile/Previewer/MobilePreviewer.INT.udn b/Engine/Documentation/Source/Platforms/Mobile/Previewer/MobilePreviewer.INT.udn
--- a/Engine/Documentation/Source/Platforms/Mobile/Previewer/MobilePreviewer.INT.udn
+++ b/Engine/Documentation/Source/Platforms/Mobile/Previewer/MobilePreviewer.INT.udn
@@ -1,52 +1,157 @@
-Availability:Public
-Title:Mobile Previewer
-Crumbs:%ROOT%, Platforms, Platforms/Mobile
-Description: Previewing games on mobile platforms will look in the UE4 editor.
-Platform:Mobile
-version: 4.9
-parent:Platforms/Mobile
-order:7
-type:overview
-tags:mobile
+Availability: Public
+Crumbs: %ROOT%
+Title: Mobile Previewer
+Description: Previewing your games inside the UE4 Editor is based on the selected Mobile Platforms.
+Type: Overview
+skilllevel:Beginner
+Version:4.16
+Parent:Platforms/Mobile
+Related:Platforms/Mobile
+Related:Engine/Rendering
+Related:Engine/UI/LevelEditor
+Order:7
+Tags:Mobile
+Topic-image:MobilePreview_topic.png
+Social-Image:T_Rendering_Level_Examples.png
 
-![](T_Rendering_Level_Examples.png)
+[VAR:spacer]![](%ROOT%/spacer.png)(w:150 convert:false)[/VAR]
+[VAR:spacer2]![](%ROOT%/spacer.png)(w:40 convert:false)[/VAR]
+
+[TOC(start:2 end:2)]
 
+![](T_Rendering_Level_Examples.png)
 [region:caption]
 1:Mobile / HTML5 - Open GL ES2,
 2:Shader Model 4.0 - DX10/ OpenGL 3.3+,
 3:Shader Model 5.0 - DX11/ OpenGL 4.3+.
 [/region]
 
-With the release of Unreal Engine (UE4) 4.8 you can now preview what your scene will look like on mobile devices directly in the UE4 editor viewport.
-When enabling different preview levels the Materials in your scene will be re-compiled to emulate the look and feature set of the render you selected. 
-You can also switch back and forth between rendering levels seamlessly inside of UE4 without having to re-start the editor.
-To change the which rendering level is used in the current UE4 Editor session, you will need to do the following. 
+In Unreal Engine 4 (UE4), you can use the **Mobile Previewer** to preview what your scene will look like on different mobile devices directly in the UE4 Editor viewport. When enabling
+the different preview rendering levels, the Materials in your scene will be recompiled to best emulate the look and feature set of the renderer preview that you selected. The Mobile
+Previewer enables you to switch back and forth between rendering levels seamlessly without having to restart the editor.
 
-[region:note]
-While you do not need to re-start the Editor for the changes to take effect the first time you change the Preview Rendering Level the editor will need time to re-compile shaders.
-[/region]
 
-1. From the  Main Toolbar use the mouse and click on the **Settings** icon.
+## Using the Mobile Previewer
 
-1. Under the **Scalability** section move the mouse over the **Preview Rendering Level** section to expose the different rendering levels options. 
+The **Mobile Previewer** enables you to quickly change between different renderers for your current UE4 session so that you can get an idea of what your game will look like on the device(s) that 
+you're targeting while allowing you to work in the editor. To learn how to change to a different renderer preview, follow these steps:
 
-1. Selecting the rendering option you would like to preview in editor by clicking on it.
+[REGION:note]
+While you do not need to restart the Editor for the new preview rendering level to take effect, the first time you change to a preview rendering level the editor will need to time to recompile the 
+shaders. Subsequent changes to previously used rendering levels should be nearly instantaneous. 
+[/REGION]
 
-1. Once a rendering level has been selected the Materials in the viewport will automatically update to reflect the new rendering method. 
-![](T_PRL.png)
+1. From the Main Toolbar, select the **Settings** button to expand the listed menu items. Under the **Scalability** section, hover over the **Preview Rendering Level** to expose the different 
+rendering levels options you can select from.
 
-[region:note]
-The Mobile previewer is intended to match mobile devices as closely as possible but is not indicative of what your project is going to look like on the target device.
-You should always make sure to fully test your project on your target device and only use the mobile previewer to see if your work is headed in the right direction.
-[/region]
+	![](PreviewRenderingLevel.png)(w:475)
+1. Hover the mouse over the section rendering level you'd like to preview and then left-mouse click to select it. In this example, **Android Vulkan Preview** was selected.
+
+	![](SelectTargetPlatform.png)
+
+	You will see the **Changing Preview Rendering Level** progress bar pop-up. Wait for this to finish and recompile shaders.
+
+	![](ChangingPreviewRenderingLevel.png)(w:450)
+
+1. Once a rendering level has been selected, the Materials in the viewport will automatically update to reflect the new rendering method using the Material Qualities enabled or disabled for for 
+that target platform. (See the [Platform Material Quality Settings](#platformmaterialqualitysettings) section of this page for more details on how to adjust these further.)
 
-Versions of UE4 prior to 4.8 do not have the ability to change the editor viewport rendering method to simulate mobile rendering.
-To see what your project could look like on mobile, you have to launch a standalone game using the **Enable Mobile Preview** option. 
-To enable the Mobile Preview mode you will need to do the following.
+	[OBJECT:ComparisonSlider]
+		[PARAM:before]
+		![Android Vulkan Preview](VulkanPreview.png)(w:700)
+		[/PARAM]
+		[PARAM:after]
+		![Desktop Shader Model 5 (SM5) Preview](Sm5Preview.png)(w:700)
+		[/PARAM]
+	[/OBJECT]
+	
 
-1. On the **Main Tool Bar** click on the triangle that is to the right of the **Play** button to expose the different play modes.
+[REGION:note]
+The Mobile Previewer is intended to match mobile devices as closely as possible but it may not always be indicative of what your project is going to look like on the target device. 
+You should always make sure to fully test your project on your target device and only use the mobile previewer to see if your work is headed in the right direction. 
+[/REGION]
 
-1. Click on the one that says **Mobile Preview** and your project will then be launched in a new window that simulates what your project should look like on a mobile device.
+### Using Mobile Previewer in 4.8 or Earlier
+
+In version 4.8 of Unreal Engine, you do not have the ability to change the editor viewport rendering method to simulate mobile rendering. To see what your project could like on 
+mobile, you have to launch a Standalone game using the **Enable Mobile Preview** option.
+
+To enable Mobile Preview mode you will need to do the following:
+
+1. From the Main Toolbar, select the arrow that is next to the **Play** button to expose the different play modes. 
+1. Select **Mobile Preview** and your project will then be launched in a new window that simulates what your project should look like on a mobile device. 
 
 	![](T_Activate_Mobile_Preview.png)
 
+
+## Enabling High-End Mobile Preview Platforms
+
+For some platforms, like the high-end ones for Android OPENGLES 3.1, Android Vulkan, and iOS Metal, you can choose to enable or disable these individually as a Preview Rendering Level that will populate
+the preview list for **High-End Mobile / Metal**. When these target platforms are enabled or disabled, they will be added or removed from the list accordingly. 
+
+Use the steps below to learn how to enable or disable these target platforms from being used with the Mobile Previewer. 
+
+1. Start by opening the **Project Settings** window by going to the file menu > **Edit** > **Window** > **Project Settings**.
+
+	![](OpenProjectSettings.png)(w:300)
+
+1. In the **Project Settings** under the **Platforms** section, select either the **Android** or **iOS** section. 
+	
+	[REGION:lightbox]
+	[![](PS_SelectPlatform.png)(w:900)](PS_SelectPlatform.png)
+	[/REGION]
+	[REGION:caption] Click image for full size view. [/REGION]
+
+	For **Android OPENGLES 3.1** and **VULKAN** previews, in the **Android** section under **Build**, set the options for **Support OpenGLES 3.1** or **Support Vulkan**. 
+	
+	![](AndroidTarget.png)(w:650)
+	
+	For **iOS Metal**, in the **iOS** section under **Rendering**, set the option for **Support Forward Rendering with Metal (A7 and up devices)**. 
+	
+	![](MetalTarget.png)(w:650)
+
+1. Once any of these have been enabled or disabled you will see the list of supported platforms update in the Preview Rendering Levels list.
+	
+	[REGION:raw]
+	![](HighEndPreviewEnabled.png)
+	[/REGION]
+
+
+
+
+## Platform Material Quality Settings
+
+In the **Project Settings** under the **Platforms** category, you can select different platform **Material Quality** sections to enable or disable specific features for the target platform.
+
+[REGION:lightbox]
+[![](MaterialQualityProjectSettings.png)(w:900)](MaterialQualityProjectSettings.png)
+[/REGION]
+[REGION:caption] Click image for full size view. [/REGION]
+
+For any of these changes to take effect, be sure to click the **Update Preview Shaders** button.
+
+![](UpdatePreviewShadersButton.png)
+
+
+## Preview Rendering Level Selections
+
+When you select your Preview Rendering Level, you will have a few options select from. Use the table below to select the option that best fits your target device. 
+
+[REGION:raw]
+| Device Target | %spacer% | Description |
+| --- | --- | --- |
+| [REGION:tablesection] High-End Mobile / Metal [/REGION] |||
+| **Default High-End Mobile** || This will emulate the default high-end mobile material quality settings without using any material quality overrides specified in the Project Settings. |
+| **Android GLES 3.1** || Provides a preview emulation of supported Android OpenGL ES3.1 quality settings. The Material Quality settings can be set in the **Project Settings** > **Android Material Quality - ES31** section.  |
+| **Android Vulkan** || Provides a preview emulation of supported Android Vulkan quality settings. The Material Quality settings can be set in the **Project Settings** > **Android Material Quality - Vulkan** section. |
+| **iOS Metal** || Provides a preview emulation of supported iOS Metal quality settings. The Material Quality settings can be set in the **Project Settings** > **iOS Material Quality - Metal** section. |
+| [REGION:tablesection] Mobile / HTML5 [/REGION] |||
+| **Default Mobile / HTML5** || This will emulate the default mobile material quality settings without using any material quality overrides specified in the **Project Settings**. |
+| **Android** || Provides a preview emulation of supported Android OpenGL ES2 quality settings. The Material Quality settings can be set in the **Project Settings** > **Android Material Quality - ES2** section. |
+| **iOS ES2** || Provides a preview emulation of supported iOS OpenGL ES2 quality settings. The Material Quality settings can be set in the **Project Settings** > **iOS Material Quality - ES2** section. |
+[/REGION]
+
+[REGION:note]
+Some of the available Preview Rendering Levels are not visible until you've enabled them as a target platform via the Project Settings, namely, Android OpenGLES 3.1 and Android Vulkan. iOS Metal defaults
+to on, however, it will also be removed as an available preview option if its target platform is disabled.
+[/REGION]
diff --git a/Engine/Documentation/Source/Platforms/Oculus/BestPractices/OculusRift.INT.udn b/Engine/Documentation/Source/Platforms/Oculus/BestPractices/OculusRift.INT.udn
--- a/Engine/Documentation/Source/Platforms/Oculus/BestPractices/OculusRift.INT.udn
+++ b/Engine/Documentation/Source/Platforms/Oculus/BestPractices/OculusRift.INT.udn
@@ -1,53 +1,38 @@
 Availability:Public
 Title:Oculus Rift Best Practices
 Crumbs:%ROOT%, Platforms, Platforms/Oculus
-Description:Information over developing for Oculus Rift.
+Description:Best practices to follow when developing UE4 projects for the Oculus Rift.
 Navigation:platform
-Version: 4.11
+Version: 4.17
 parent:Platforms/Oculus
 type:overview
 tags:VR
 tags:Oculus Rift
+topic-image:Platforms\Oculus\oculus_bp.png
 
-[TOC(start:2)]
+The Oculus Rift is a Head Mounted Display (HMD) which allows the user to deeply immerse themselves into the worlds and games they are viewing or playing. The following guide will help to highlight some information that you need to know about when developing content for the Oculus Rift and Unreal Engine 4 (UE4).
 
-The Oculus Rift is a Head Mounted Display (HMD) which allows the user to deeply immerse themselves into the worlds and 
-games they are viewing or playing. The following guide will help to highlight some information that you need to know 
-about when developing content for the Oculus Rift and Unreal Engine 4 (UE4).
+## Oculus HMD Target Frame Rate
 
+Below you will find the frame rate that your UE4 project must meet when run on an Oculus Rift HMD device. 
+
+| HMD Device  | Target Frame Rate |
+| --- | --- |
+| Rift Retail | 90 FPS|
+| DK2 | 75 FPS|
+| DK1 | 60 FPS|
 
 ## Oculus Rift Sanity Check
 
-Unreal Engine will automatically use an Oculus Rift if it is plugged in and the plug-in is enabled. If for some 
-reason the Oculus is not working first check that the plug-In is enabled. You can find the Plugin under Head Mounted 
-Displays that is under the [Plug-in](Programming/Plugins) menu.  
+Unreal Engine will automatically use an Oculus Rift HMD if it is plugged in and the Oculus plug-in is enabled. If for some reason the Oculus HMD is not working first check that the Oculus plug-In is enabled. You can find the Oculus plugin under the **Virtual Reality** section in the [Plug-in](Programming/Plugins) menu.  
 
-[REGION:fullwidth]
-![](T_VR_Plug-In_On_00.png)
+[REGION:lightbox]
+[![](T_VR_Plug-In_On_00.png)(w:491)](T_VR_Plug-In_On_00.png)
 [/REGION]
 
-<!-- If that did not fix your issues the next thing to check is to make sure that the Oculus Rift has been detected. To 
-do this, you need to do the following steps:
-
-* First run a [Standalone Game](Engine/UI/LevelEditor/InEditorTesting/Settings#playinstandalonegame). 
-* Once the Standalone game has been started, bring up the Unreal Console by using the Backtick (`) key. 
-* Then type **showlog**.
-* Next type **ovrversion**.
-* The current LibOVR version and build date should be printed in the log window like in the following picture:
-	[REGION:fullwidth]
-	![](T_OculusRiftPluginVersionCheck.png)
-	[/REGION]
-* If the message **Command not recognized: ovrversion** is displayed instead check the following:
-	* That the plug-in is enabled. 
-	* That you are running UE4 on the Windows platform.
-	* That your Oculus Rift is plugged in.
-	* That your Oculus Rift is functioning correctly.
-
-[REGION:tip]
-If you do all of the above and it still does not work it could be that you have a hardware issues. If this happens to 
-you please contact [Oculus VR Support](https://support.oculus.com/help/oculus/885981024820727/?ref=hc_fnav). 
+[REGION:caption]
+Click for full image.
 [/REGION]
--->
 
 ## Checking your work with the Oculus Rift
 
@@ -67,6 +52,4 @@ The Oculus Rift HMD will not work with any of the UE4 Editor viewports or the de
 
 ## Oculus Reference
 
-Oculus has put a great deal of time and energy into researching how to give the user the best VR experience possible when 
-using the Rift. If you have not done so already, it is highly recommend that you check out their [website]( https://developer.oculus.com/) 
-as there is a lot of useful information that will help you to develop and deliver the best VR user experience possible.
+Oculus has put a great deal of time and energy into researching how to give the user the best VR experience possible when using the Rift. If you have not done so already, it is highly recommend that you check out their [website]( https://developer.oculus.com/) as there is a lot of useful information that will help you to develop and deliver the best VR user experience possible.
