Copying //UE4/Dev-Documentation to //UE4/Samples-Main (Source: //UE4/Dev-Documentation @ 3623665)

diff --git a/Engine/Documentation/Source/Platforms/Oculus/HowTo/GardianSystemChangeColor/HTOGS_ChangeColors.INT.udn b/Engine/Documentation/Source/Platforms/Oculus/HowTo/GardianSystemChangeColor/HTOGS_ChangeColors.INT.udn
new file mode 100644
--- /dev/null
+++ b/Engine/Documentation/Source/Platforms/Oculus/HowTo/GardianSystemChangeColor/HTOGS_ChangeColors.INT.udn
@@ -0,0 +1,84 @@
+Availability: Public
+Crumbs: %ROOT%, Platforms
+Title: Change the Oculus Rift Guardian System Color
+Description: Changing the color of the Rift Guardian system.
+Type: how-to
+Skilllevel: Intermediate
+Version: 4.17
+Parent: Platforms/Oculus
+Order: 2
+tags: Oculus
+topic-image:Platforms\Oculus\HowTo\GuardianSystem\HTGuardian_System_Topic_Image.png
+
+
+The Oculus Guardian system is used to display the boundaries of the VR interaction area. As a tracked device approaches the boundaries, the Oculus Runtime automatically provides visual cues notifying the user of this. In the following How - To we will take a look at changing the color, the Oculus Guardian system uses to display the VR interaction area.
+
+[REGION:note]
+The following How - To is going to make use of the Standing VR Pawn with added Touch Controller and Guardian support. If you are not sure how to set this up, the following two How - To's will help you out.
+* [Set Up Cameras for the Oculus Rift](Platforms\Oculus\HowTo\StandingCamera)
+* [Setup Touch Controllers for the Oculus Rift](Platforms\Oculus\HowTo\TouchControllers)
+* [Set Up Guardian System](Platforms\Oculus\HowTo\GuardianSystem)
+[/REGION]
+
+[REGION:note]
+For the Guardian system to work, you will need to make sure you have set it up using the Oculus App. For more information on how to do this, please see the official [Oculus Guardian System](https://developer.oculus.com/documentation/pcsdk/latest/concepts/dg-guardian-system/) Setup pages.
+[/REGION]
+
+[REGION:warning]
+It is **NOT** possible nor advisable to disable the Guardian system from inside UE4. You can, however, adjust how UE4 will react to the user getting close to the boundaries.  
+[/REGION]
+
+## Changing the color of the Oculus Rift Guardian System
+
+[REGION:note]
+Make sure that your Pawn has a **OculusRiftBoundy** added to it or the following will not work. If you are not certain how to do this, please see the [Setting Up Guardian System](Platforms\Oculus\HowTo\GuardianSystem) pages for more information on how to do this.
+[/REGION]
+
+1. Create a new **Variable** called **Oculus Rift Boundary Color** and set its type to **Linear Color** and its color to **Red**.
+
+    ![](HTGuardian_System_04.png)
+
+1. In the **Event Graph** add a **Event Begin Play** and **Set Tracking Origin** node. Make sure that the Set Tracking Origin node has had its **Origin** set to **Floor Level** and then connect the Event Begin Play to the Set Tracking Origin node. 
+
+    [REGION:lightbox]
+    [![](HTGuardian_System_02.png)(w:492)](HTGuardian_System_02.png)
+    [/REGION]
+
+    [REGION:caption]
+    Click for full image.
+    [/REGION]
+
+1. Right - click on the Event Graph and from the menu search for **Set Outer Boundary Color**, clicking on it to add it to the graph. 
+
+    ![](HTGuardian_System_03.png)
+
+1. Drag both the **Oculus Rift Boundary Color** variable and the **Oculus Rift Boundary** component into the Event Graph. Connect the Oculus Rift Boundary Color variable to in the **In Boundary Color** on the Set Boundary Color node and connect the Oculus Rift Boundary to the **Target** input. 
+
+    [REGION:lightbox]
+    [![](HTGuardian_System_05.png)(w:492)](HTGuardian_System_05.png)
+    [/REGION]
+
+    [REGION:caption]
+    Click for full image.
+    [/REGION]
+
+1. Connect the output of the Set Tracking Origin node to the input of the Set Outer Boundary Color node and when completed your Event Graph should look like the following.
+
+    [REGION:lightbox]
+    [![](HTGuardian_System_06.png)(w:492)](HTGuardian_System_06.png)
+    [/REGION]
+
+    [REGION:caption]
+    Click for full image.
+    [/REGION]
+
+1. Now when you run the level and put the Oculus Rift on when the boundary is displayed, it will be the color you set in the Oculus Rift Boundary Color variable. 
+
+    ![](HTGuardian_System_07.png)
+
+
+## UE4 Project Downloads
+Below you will find a link to where you can download the UE4 project that was used to create this example.
+
+* [**Oculus Rift Guardian System Example Project**](https://epicgames.box.com/s/s4vvkb2i0ohtice8vtqude6i8ih7hy8i)
+
diff --git a/Engine/Documentation/Source/Platforms/Oculus/HowTo/GardianSystemShowHide/HTOGC_ShowHideBoundryHit.INT.udn b/Engine/Documentation/Source/Platforms/Oculus/HowTo/GardianSystemShowHide/HTOGC_ShowHideBoundryHit.INT.udn
new file mode 100644
--- /dev/null
+++ b/Engine/Documentation/Source/Platforms/Oculus/HowTo/GardianSystemShowHide/HTOGC_ShowHideBoundryHit.INT.udn
@@ -0,0 +1,248 @@
+Availability: Public
+Crumbs: %ROOT%, Platforms
+Title: Detect Oculus Rift Guardian System Activation
+Description: Detecting if the user has activated the Oculus Rift Guardian system.
+Type: how-to
+Skilllevel: Intermidate
+Version: 4.17
+Parent: Platforms/Oculus
+Order: 2
+tags: VR
+topic-image:Platforms\Oculus\HowTo\GuardianSystem\HTGuardian_System_Topic_Image.png
+
+The Oculus Guardian system is used to display the boundaries of the VR interaction area. As a tracked device approaches the boundaries, the Oculus Runtime automatically provides visual cues notifying the user of this. In the following How - To we will take a look at how to display the interaction as well as other visual cues to notify the users that one of their devices is outside or are close to going outside the interaction area.
+
+[REGION:note]
+The following How - To is going to make use of the Standing VR Pawn with added Touch Controller and Guardian support. If you are not sure how to set this up, the following two How - To's will help you out.
+* [Set Up Standing Cameras for the Oculus Rift](Platforms\Oculus\HowTo\StandingCamera)
+* [Set Up Touch Controllers for the Oculus Rift](Platforms\Oculus\HowTo\TouchControllers)
+* [Set Up Oculus Rift Guardian System](Platforms\Oculus\HowTo\GuardianSystem)
+[/REGION]
+
+[REGION:note]
+For the Guardian system to work, you will need to make sure you have set it up using the Oculus App. For more information on how to do this, please see the official [Oculus Guardian System](https://developer.oculus.com/documentation/pcsdk/latest/concepts/dg-guardian-system/) Setup pages.
+[/REGION]
+
+[REGION:warning]
+It is **NOT** possible nor advisable to disable the Guardian system from inside UE4. You can, however, adjust how UE4 will react to the user getting close to the boundaries.  
+[/REGION]
+
+## Interacting with the Oculus Rift Guardian System
+
+1. To display the Oculus Rift Boundary, we have to know when one of the user's devices is close to it. The node **Is Outer Boundary Displayed** will allow us to do just this but we need a way to check this each tick. For us to be able to do this in UE4 we will first need to add the following nodes to the Event Graph:
+
+    * **Event Tick**  
+    * **Branch**
+    * **Is Outer Boundary Displayed**
+    * **Oculus Rift Boundary**
+
+1. Once the nodes have been added, we will need to connect them so that the Is Outer Boundary Displayed node is only called when a user's device gets close to the boundary; otherwise, we do not want anything to happen. To achieve this, setup up your Event graph to match the following image:
+
+    [REGION:lightbox]
+    [![](HTGuardian_System_00.png)(w:405)](HTGuardian_System_00.png)
+    [/REGION]
+
+    [REGION:caption]
+    Click for full image.
+    [/REGION]
+
+1. Next, we want to display not only the boundaries that the user setup in the Oculus App but also the square/rectangular interaction area that fits within those boundaries. To do this, we are going to look at all the points within the play area using the **Get Play Area Points** node. Then to determine the size of the square/rectangular interaction we will go through each point using a **For Loop** and at each point, we draw a point and connect each point with a line display this information to the user.  To achieve this, setup up your Event graph to match the following image:
+
+    [REGION:lightbox]
+    [Include:Platforms\Oculus\HowTo\GardianSystemShowHide\#DrawPlayArea]
+    [/REGION]
+
+    [REGION:caption]
+    Click on image to copy Blueprint code.
+    [/REGION]
+
+1. When that is completed, make sure to connect the For Loop to the True output of the Branch node and do not forget to Compile and Save your Blueprint. When completed your Blueprint Event graph should look like the following image:
+
+    [REGION:lightbox]
+    [![](HTGuardian_System_02.png)(w:1000)](HTGuardian_System_02.png)
+    [/REGION]
+
+    [REGION:caption]
+    Click for full image.
+    [/REGION]
+
+1. Now, put on your Oculus Rift HMD on, pick up your Touch controllers and launch the project using VR Preview. Once the project launches, slowly move one of the Touch controllers to towards the Guardian boundary. When the Guardian boundary is displayed, you should also see the square/rectangular interaction area being displayed like in the following video.
+
+    [OBJECT:EmbeddedVideo]
+        [PARAMLITERAL:videoid]
+            09yYAc_kLXM
+        [/PARAMLITERAL]
+        [PARAMLITERAL:width]
+        100
+        [/PARAMLITERAL]
+        [PARAMLITERAL:height]
+        56.25
+        [/PARAMLITERAL]
+        [PARAMLITERAL:units]
+        %
+        [/PARAMLITERAL]
+        [PARAMLITERAL:type]
+        simple
+        [/PARAMLITERAL]
+    [/OBJECT]
+
+
+## UE4 Project Downloads
+Below you will find a link to where you can download the UE4 project that was used to create this example.
+
+* [**Oculus Rift Guardian System Example Project**](https://epicgames.box.com/s/s4vvkb2i0ohtice8vtqude6i8ih7hy8i)
+
+
+
+<!--
+[EXCERPT:DrawPlayArea]    
+[OBJECT:Snippet]
+    [PARAMLITERAL:id]
+        DrawPlayArea
+    [/PARAMLITERAL]
+    [PARAMLITERAL:type]
+        blueprint
+    [/PARAMLITERAL]
+    [PARAM:image]
+        ![](HTGuardian_System_01.png)
+    [/PARAM]
+    [PARAM:snippet]
+            BEGIN OBJECT
+                Begin Object Class=K2Node_CallFunction Name="K2Node_CallFunction_175"
+                   bIsPureFunc=True
+                   FunctionReference=(MemberParent=Class'/Script/OculusRift.OculusRiftBoundaryComponent',MemberName="GetPlayAreaPoints")
+                   NodePosX=832
+                   NodePosY=-260
+                   NodeGuid=ED8B9D104C03A3A9BE5248ACB92A9045
+                   CustomProperties Pin (PinId=A893351045D87FA4ACF13DA19AA91AA7,PinName="self",PinFriendlyName=NSLOCTEXT("K2Node", "Target", "Target"),PinToolTip="Target\nOculus Rift Boundary Component Reference",PinType.PinCategory="object",PinType.PinSubCategory="",PinType.PinSubCategoryObject=Class'/Script/OculusRift.OculusRiftBoundaryComponent',PinType.PinSubCategoryMemberReference=(),PinType.PinValueType=(),PinType.bIsMap=False,PinType.bIsSet=False,PinType.bIsArray=False,PinType.bIsReference=False,PinType.bIsConst=False,PinType.bIsWeakPointer=False,LinkedTo=(K2Node_VariableGet_85 B683B98F40D442BC4C0F20A33CB8A44E,),PersistentGuid=00000000000000000000000000000000,bHidden=False,bNotConnectable=False,bDefaultValueIsReadOnly=False,bDefaultValueIsIgnored=False,bAdvancedView=False,)
+                   CustomProperties Pin (PinId=A0ED07004023EA0AA18F56AE3DE49023,PinName="ReturnValue",PinToolTip="Return Value\nArray of Vector s\n\nArray of 3D points (clockwise order) defining play area at floor height Points are returned in Unreal world coordinate space.",Direction="EGPD_Output",PinType.PinCategory="struct",PinType.PinSubCategory="",PinType.PinSubCategoryObject=ScriptStruct'/Script/CoreUObject.Vector',PinType.PinSubCategoryMemberReference=(),PinType.PinValueType=(),PinType.bIsMap=False,PinType.bIsSet=False,PinType.bIsArray=True,PinType.bIsReference=False,PinType.bIsConst=False,PinType.bIsWeakPointer=False,LinkedTo=(K2Node_CallArrayFunction_14 47C5006C4C77AA941F3B1EBFC411B000,K2Node_CallArrayFunction_17 01395B804D226B8004689BA1CD742F72,K2Node_CallArrayFunction_16 01395B804D226B8004689BA1CD742F72,),PersistentGuid=00000000000000000000000000000000,bHidden=False,bNotConnectable=False,bDefaultValueIsReadOnly=False,bDefaultValueIsIgnored=False,bAdvancedView=False,)
+                End Object
+                Begin Object Class=K2Node_VariableGet Name="K2Node_VariableGet_85"
+                   VariableReference=(MemberName="OculusRiftBoundary",bSelfContext=True)
+                   NodePosX=640
+                   NodePosY=-222
+                   NodeGuid=5A7082144FA890B349267186192A81ED
+                   CustomProperties Pin (PinId=B683B98F40D442BC4C0F20A33CB8A44E,PinName="OculusRiftBoundary",Direction="EGPD_Output",PinType.PinCategory="object",PinType.PinSubCategory="",PinType.PinSubCategoryObject=Class'/Script/OculusRift.OculusRiftBoundaryComponent',PinType.PinSubCategoryMemberReference=(),PinType.PinValueType=(),PinType.bIsMap=False,PinType.bIsSet=False,PinType.bIsArray=False,PinType.bIsReference=False,PinType.bIsConst=False,PinType.bIsWeakPointer=False,LinkedTo=(K2Node_CallFunction_175 A893351045D87FA4ACF13DA19AA91AA7,),PersistentGuid=00000000000000000000000000000000,bHidden=False,bNotConnectable=False,bDefaultValueIsReadOnly=False,bDefaultValueIsIgnored=False,bAdvancedView=False,)
+                   CustomProperties Pin (PinId=5B4341774C523895808B11BCE7244F62,PinName="self",PinFriendlyName=NSLOCTEXT("K2Node", "Target", "Target"),PinType.PinCategory="object",PinType.PinSubCategory="",PinType.PinSubCategoryObject=BlueprintGeneratedClass'/Game/VRPawn.VRPawn_C',PinType.PinSubCategoryMemberReference=(),PinType.PinValueType=(),PinType.bIsMap=False,PinType.bIsSet=False,PinType.bIsArray=False,PinType.bIsReference=False,PinType.bIsConst=False,PinType.bIsWeakPointer=False,PersistentGuid=00000000000000000000000000000000,bHidden=True,bNotConnectable=False,bDefaultValueIsReadOnly=False,bDefaultValueIsIgnored=False,bAdvancedView=False,)
+                End Object
+                Begin Object Class=K2Node_CallArrayFunction Name="K2Node_CallArrayFunction_14"
+                   bIsPureFunc=True
+                   FunctionReference=(MemberParent=Class'/Script/Engine.KismetArrayLibrary',MemberName="Array_Length")
+                   NodePosX=1168
+                   NodePosY=-340
+                   NodeGuid=2F3BD21B4C03FA55185638892C1809B2
+                   CustomProperties Pin (PinId=40AFF52B4465D86FC494DB8502D8D15B,PinName="self",PinFriendlyName=NSLOCTEXT("K2Node", "Target", "Target"),PinToolTip="Target\nKismet Array Library Reference",PinType.PinCategory="object",PinType.PinSubCategory="",PinType.PinSubCategoryObject=Class'/Script/Engine.KismetArrayLibrary',PinType.PinSubCategoryMemberReference=(),PinType.PinValueType=(),PinType.bIsMap=False,PinType.bIsSet=False,PinType.bIsArray=False,PinType.bIsReference=False,PinType.bIsConst=False,PinType.bIsWeakPointer=False,DefaultObject="/Script/Engine.Default__KismetArrayLibrary",PersistentGuid=00000000000000000000000000000000,bHidden=True,bNotConnectable=False,bDefaultValueIsReadOnly=False,bDefaultValueIsIgnored=False,bAdvancedView=False,)
+                   CustomProperties Pin (PinId=47C5006C4C77AA941F3B1EBFC411B000,PinName="TargetArray",PinToolTip="Target Array\nArray of Vector s\n\nThe array to get the length of *",PinType.PinCategory="struct",PinType.PinSubCategory="",PinType.PinSubCategoryObject=ScriptStruct'/Script/CoreUObject.Vector',PinType.PinSubCategoryMemberReference=(),PinType.PinValueType=(),PinType.bIsMap=False,PinType.bIsSet=False,PinType.bIsArray=True,PinType.bIsReference=True,PinType.bIsConst=True,PinType.bIsWeakPointer=False,LinkedTo=(K2Node_CallFunction_175 A0ED07004023EA0AA18F56AE3DE49023,),PersistentGuid=00000000000000000000000000000000,bHidden=False,bNotConnectable=False,bDefaultValueIsReadOnly=False,bDefaultValueIsIgnored=True,bAdvancedView=False,)
+                   CustomProperties Pin (PinId=FA1FAA9B417BE094A04382A683CFD740,PinName="ReturnValue",PinToolTip="Return Value\nInteger\n\nThe length of the array",Direction="EGPD_Output",PinType.PinCategory="int",PinType.PinSubCategory="",PinType.PinSubCategoryObject=None,PinType.PinSubCategoryMemberReference=(),PinType.PinValueType=(),PinType.bIsMap=False,PinType.bIsSet=False,PinType.bIsArray=False,PinType.bIsReference=False,PinType.bIsConst=False,PinType.bIsWeakPointer=False,AutogeneratedDefaultValue="0",LinkedTo=(K2Node_CallFunction_554 3EEC9DF74A50D22A7105AB811D1F7B77,K2Node_CallFunction_555 699967334582B6AF92086FAE81646852,),PersistentGuid=00000000000000000000000000000000,bHidden=False,bNotConnectable=False,bDefaultValueIsReadOnly=False,bDefaultValueIsIgnored=False,bAdvancedView=False,)
+                End Object
+                Begin Object Class=K2Node_MacroInstance Name="K2Node_MacroInstance_15"
+                   MacroGraphReference=(MacroGraph=EdGraph'/Engine/EditorBlueprintResources/StandardMacros.StandardMacros:ForLoop',GraphBlueprint=Blueprint'/Engine/EditorBlueprintResources/StandardMacros.StandardMacros',GraphGuid=55C904AF4B45FE1761FB55A8DB9FB801)
+                   NodePosX=1520
+                   NodePosY=-480
+                   NodeGuid=ED513F3547E389BD3763B3B538C8AB7A
+                   CustomProperties Pin (PinId=C8A32D6840D2DDAA8A995C9F36947115,PinName="execute",PinType.PinCategory="exec",PinType.PinSubCategory="",PinType.PinSubCategoryObject=None,PinType.PinSubCategoryMemberReference=(),PinType.PinValueType=(),PinType.bIsMap=False,PinType.bIsSet=False,PinType.bIsArray=False,PinType.bIsReference=False,PinType.bIsConst=False,PinType.bIsWeakPointer=False,PersistentGuid=00000000000000000000000000000000,bHidden=False,bNotConnectable=False,bDefaultValueIsReadOnly=False,bDefaultValueIsIgnored=False,bAdvancedView=False,)
+                   CustomProperties Pin (PinId=C94FA3FA43CEBB68D28E4085C5E14C59,PinName="FirstIndex",PinType.PinCategory="int",PinType.PinSubCategory="",PinType.PinSubCategoryObject=None,PinType.PinSubCategoryMemberReference=(),PinType.PinValueType=(),PinType.bIsMap=False,PinType.bIsSet=False,PinType.bIsArray=False,PinType.bIsReference=False,PinType.bIsConst=False,PinType.bIsWeakPointer=False,DefaultValue="0",PersistentGuid=00000000000000000000000000000000,bHidden=False,bNotConnectable=False,bDefaultValueIsReadOnly=False,bDefaultValueIsIgnored=False,bAdvancedView=False,)
+                   CustomProperties Pin (PinId=C0CFB214415D38BA0350F992C93090B5,PinName="LastIndex",PinType.PinCategory="int",PinType.PinSubCategory="",PinType.PinSubCategoryObject=None,PinType.PinSubCategoryMemberReference=(),PinType.PinValueType=(),PinType.bIsMap=False,PinType.bIsSet=False,PinType.bIsArray=False,PinType.bIsReference=False,PinType.bIsConst=False,PinType.bIsWeakPointer=False,DefaultValue="0",LinkedTo=(K2Node_CallFunction_554 BB5C6FC44F3E025D33427990DEBE0984,),PersistentGuid=00000000000000000000000000000000,bHidden=False,bNotConnectable=False,bDefaultValueIsReadOnly=False,bDefaultValueIsIgnored=False,bAdvancedView=False,)
+                   CustomProperties Pin (PinId=585876994C3082C336E4F08AB6707724,PinName="LoopBody",Direction="EGPD_Output",PinType.PinCategory="exec",PinType.PinSubCategory="",PinType.PinSubCategoryObject=None,PinType.PinSubCategoryMemberReference=(),PinType.PinValueType=(),PinType.bIsMap=False,PinType.bIsSet=False,PinType.bIsArray=False,PinType.bIsReference=False,PinType.bIsConst=False,PinType.bIsWeakPointer=False,LinkedTo=(K2Node_CallFunction_970 5C92A48C4EC3E621FB3186B24B7F49C9,),PersistentGuid=00000000000000000000000000000000,bHidden=False,bNotConnectable=False,bDefaultValueIsReadOnly=False,bDefaultValueIsIgnored=False,bAdvancedView=False,)
+                   CustomProperties Pin (PinId=485524BA49C6E2F23DF0DC9247CD6B93,PinName="Index",Direction="EGPD_Output",PinType.PinCategory="int",PinType.PinSubCategory="",PinType.PinSubCategoryObject=None,PinType.PinSubCategoryMemberReference=(),PinType.PinValueType=(),PinType.bIsMap=False,PinType.bIsSet=False,PinType.bIsArray=False,PinType.bIsReference=False,PinType.bIsConst=False,PinType.bIsWeakPointer=False,LinkedTo=(K2Node_CommutativeAssociativeBinaryOperator_0 78D1F70B4F17912EEE2A86ADA26C7372,K2Node_CallArrayFunction_16 7B92B2804499E72162447D8B5307E76C,),PersistentGuid=00000000000000000000000000000000,bHidden=False,bNotConnectable=False,bDefaultValueIsReadOnly=False,bDefaultValueIsIgnored=False,bAdvancedView=False,)
+                   CustomProperties Pin (PinId=E4E357844650ABA010F53FA044955EC6,PinName="Completed",Direction="EGPD_Output",PinType.PinCategory="exec",PinType.PinSubCategory="",PinType.PinSubCategoryObject=None,PinType.PinSubCategoryMemberReference=(),PinType.PinValueType=(),PinType.bIsMap=False,PinType.bIsSet=False,PinType.bIsArray=False,PinType.bIsReference=False,PinType.bIsConst=False,PinType.bIsWeakPointer=False,PersistentGuid=00000000000000000000000000000000,bHidden=False,bNotConnectable=False,bDefaultValueIsReadOnly=False,bDefaultValueIsIgnored=False,bAdvancedView=False,)
+                End Object
+                Begin Object Class=K2Node_CallFunction Name="K2Node_CallFunction_554"
+                   bIsPureFunc=True
+                   FunctionReference=(MemberParent=Class'/Script/Engine.KismetMathLibrary',MemberName="Subtract_IntInt")
+                   NodePosX=1376
+                   NodePosY=-411
+                   NodeGuid=0ED49692453023A374CA65966AF281CD
+                   CustomProperties Pin (PinId=6C25DD954E5DD954CD13298437BB51E8,PinName="self",PinFriendlyName=NSLOCTEXT("K2Node", "Target", "Target"),PinToolTip="Target\nKismet Math Library Reference",PinType.PinCategory="object",PinType.PinSubCategory="",PinType.PinSubCategoryObject=Class'/Script/Engine.KismetMathLibrary',PinType.PinSubCategoryMemberReference=(),PinType.PinValueType=(),PinType.bIsMap=False,PinType.bIsSet=False,PinType.bIsArray=False,PinType.bIsReference=False,PinType.bIsConst=False,PinType.bIsWeakPointer=False,DefaultObject="/Script/Engine.Default__KismetMathLibrary",PersistentGuid=00000000000000000000000000000000,bHidden=True,bNotConnectable=False,bDefaultValueIsReadOnly=False,bDefaultValueIsIgnored=False,bAdvancedView=False,)
+                   CustomProperties Pin (PinId=3EEC9DF74A50D22A7105AB811D1F7B77,PinName="A",PinToolTip="A\nInteger",PinType.PinCategory="int",PinType.PinSubCategory="",PinType.PinSubCategoryObject=None,PinType.PinSubCategoryMemberReference=(),PinType.PinValueType=(),PinType.bIsMap=False,PinType.bIsSet=False,PinType.bIsArray=False,PinType.bIsReference=False,PinType.bIsConst=False,PinType.bIsWeakPointer=False,DefaultValue="0",AutogeneratedDefaultValue="0",LinkedTo=(K2Node_CallArrayFunction_14 FA1FAA9B417BE094A04382A683CFD740,),PersistentGuid=00000000000000000000000000000000,bHidden=False,bNotConnectable=False,bDefaultValueIsReadOnly=False,bDefaultValueIsIgnored=False,bAdvancedView=False,)
+                   CustomProperties Pin (PinId=A4AA99B74190D88F1C1D088F55D529D2,PinName="B",PinToolTip="B\nInteger",PinType.PinCategory="int",PinType.PinSubCategory="",PinType.PinSubCategoryObject=None,PinType.PinSubCategoryMemberReference=(),PinType.PinValueType=(),PinType.bIsMap=False,PinType.bIsSet=False,PinType.bIsArray=False,PinType.bIsReference=False,PinType.bIsConst=False,PinType.bIsWeakPointer=False,DefaultValue="1",AutogeneratedDefaultValue="1",PersistentGuid=00000000000000000000000000000000,bHidden=False,bNotConnectable=False,bDefaultValueIsReadOnly=False,bDefaultValueIsIgnored=False,bAdvancedView=False,)
+                   CustomProperties Pin (PinId=BB5C6FC44F3E025D33427990DEBE0984,PinName="ReturnValue",PinToolTip="Return Value\nInteger",Direction="EGPD_Output",PinType.PinCategory="int",PinType.PinSubCategory="",PinType.PinSubCategoryObject=None,PinType.PinSubCategoryMemberReference=(),PinType.PinValueType=(),PinType.bIsMap=False,PinType.bIsSet=False,PinType.bIsArray=False,PinType.bIsReference=False,PinType.bIsConst=False,PinType.bIsWeakPointer=False,AutogeneratedDefaultValue="0",LinkedTo=(K2Node_MacroInstance_15 C0CFB214415D38BA0350F992C93090B5,),PersistentGuid=00000000000000000000000000000000,bHidden=False,bNotConnectable=False,bDefaultValueIsReadOnly=False,bDefaultValueIsIgnored=False,bAdvancedView=False,)
+                End Object
+                Begin Object Class=K2Node_CommutativeAssociativeBinaryOperator Name="K2Node_CommutativeAssociativeBinaryOperator_0"
+                   bIsPureFunc=True
+                   FunctionReference=(MemberParent=Class'/Script/Engine.KismetMathLibrary',MemberName="Add_IntInt")
+                   NodePosX=1776
+                   NodePosY=-384
+                   NodeGuid=4C4A458E4B25AE3BAB9B6D86838C2363
+                   CustomProperties Pin (PinId=670DE9E1405EDBAC586D2F8A6DEB5DB3,PinName="self",PinFriendlyName=NSLOCTEXT("K2Node", "Target", "Target"),PinToolTip="Target\nKismet Math Library Reference",PinType.PinCategory="object",PinType.PinSubCategory="",PinType.PinSubCategoryObject=Class'/Script/Engine.KismetMathLibrary',PinType.PinSubCategoryMemberReference=(),PinType.PinValueType=(),PinType.bIsMap=False,PinType.bIsSet=False,PinType.bIsArray=False,PinType.bIsReference=False,PinType.bIsConst=False,PinType.bIsWeakPointer=False,DefaultObject="/Script/Engine.Default__KismetMathLibrary",PersistentGuid=00000000000000000000000000000000,bHidden=True,bNotConnectable=False,bDefaultValueIsReadOnly=False,bDefaultValueIsIgnored=False,bAdvancedView=False,)
+                   CustomProperties Pin (PinId=78D1F70B4F17912EEE2A86ADA26C7372,PinName="A",PinToolTip="A\nInteger",PinType.PinCategory="int",PinType.PinSubCategory="",PinType.PinSubCategoryObject=None,PinType.PinSubCategoryMemberReference=(),PinType.PinValueType=(),PinType.bIsMap=False,PinType.bIsSet=False,PinType.bIsArray=False,PinType.bIsReference=False,PinType.bIsConst=False,PinType.bIsWeakPointer=False,DefaultValue="0",AutogeneratedDefaultValue="0",LinkedTo=(K2Node_MacroInstance_15 485524BA49C6E2F23DF0DC9247CD6B93,),PersistentGuid=00000000000000000000000000000000,bHidden=False,bNotConnectable=False,bDefaultValueIsReadOnly=False,bDefaultValueIsIgnored=False,bAdvancedView=False,)
+                   CustomProperties Pin (PinId=668EB32E49EC473CA96DF788B99D242E,PinName="B",PinToolTip="B\nInteger",PinType.PinCategory="int",PinType.PinSubCategory="",PinType.PinSubCategoryObject=None,PinType.PinSubCategoryMemberReference=(),PinType.PinValueType=(),PinType.bIsMap=False,PinType.bIsSet=False,PinType.bIsArray=False,PinType.bIsReference=False,PinType.bIsConst=False,PinType.bIsWeakPointer=False,DefaultValue="1",AutogeneratedDefaultValue="1",PersistentGuid=00000000000000000000000000000000,bHidden=False,bNotConnectable=False,bDefaultValueIsReadOnly=False,bDefaultValueIsIgnored=False,bAdvancedView=False,)
+                   CustomProperties Pin (PinId=2D424B8A45D377433884CABF17899C0D,PinName="ReturnValue",PinToolTip="Return Value\nInteger",Direction="EGPD_Output",PinType.PinCategory="int",PinType.PinSubCategory="",PinType.PinSubCategoryObject=None,PinType.PinSubCategoryMemberReference=(),PinType.PinValueType=(),PinType.bIsMap=False,PinType.bIsSet=False,PinType.bIsArray=False,PinType.bIsReference=False,PinType.bIsConst=False,PinType.bIsWeakPointer=False,AutogeneratedDefaultValue="0",LinkedTo=(K2Node_CallFunction_555 2FFCF9834C0A41D14C1EA0BE10F4B229,),PersistentGuid=00000000000000000000000000000000,bHidden=False,bNotConnectable=False,bDefaultValueIsReadOnly=False,bDefaultValueIsIgnored=False,bAdvancedView=False,)
+                End Object
+                Begin Object Class=K2Node_CallFunction Name="K2Node_CallFunction_555"
+                   bIsPureFunc=True
+                   FunctionReference=(MemberParent=Class'/Script/Engine.KismetMathLibrary',MemberName="Percent_IntInt")
+                   NodePosX=1968
+                   NodePosY=-240
+                   NodeGuid=8D5D65624BED73ED64F1B78D621D69D8
+                   CustomProperties Pin (PinId=5BCE4019473987CC992D75868E15C8E8,PinName="self",PinFriendlyName=NSLOCTEXT("K2Node", "Target", "Target"),PinToolTip="Target\nKismet Math Library Reference",PinType.PinCategory="object",PinType.PinSubCategory="",PinType.PinSubCategoryObject=Class'/Script/Engine.KismetMathLibrary',PinType.PinSubCategoryMemberReference=(),PinType.PinValueType=(),PinType.bIsMap=False,PinType.bIsSet=False,PinType.bIsArray=False,PinType.bIsReference=False,PinType.bIsConst=False,PinType.bIsWeakPointer=False,DefaultObject="/Script/Engine.Default__KismetMathLibrary",PersistentGuid=00000000000000000000000000000000,bHidden=True,bNotConnectable=False,bDefaultValueIsReadOnly=False,bDefaultValueIsIgnored=False,bAdvancedView=False,)
+                   CustomProperties Pin (PinId=2FFCF9834C0A41D14C1EA0BE10F4B229,PinName="A",PinToolTip="A\nInteger",PinType.PinCategory="int",PinType.PinSubCategory="",PinType.PinSubCategoryObject=None,PinType.PinSubCategoryMemberReference=(),PinType.PinValueType=(),PinType.bIsMap=False,PinType.bIsSet=False,PinType.bIsArray=False,PinType.bIsReference=False,PinType.bIsConst=False,PinType.bIsWeakPointer=False,DefaultValue="0",AutogeneratedDefaultValue="0",LinkedTo=(K2Node_CommutativeAssociativeBinaryOperator_0 2D424B8A45D377433884CABF17899C0D,),PersistentGuid=00000000000000000000000000000000,bHidden=False,bNotConnectable=False,bDefaultValueIsReadOnly=False,bDefaultValueIsIgnored=False,bAdvancedView=False,)
+                   CustomProperties Pin (PinId=699967334582B6AF92086FAE81646852,PinName="B",PinToolTip="B\nInteger",PinType.PinCategory="int",PinType.PinSubCategory="",PinType.PinSubCategoryObject=None,PinType.PinSubCategoryMemberReference=(),PinType.PinValueType=(),PinType.bIsMap=False,PinType.bIsSet=False,PinType.bIsArray=False,PinType.bIsReference=False,PinType.bIsConst=False,PinType.bIsWeakPointer=False,DefaultValue="0",AutogeneratedDefaultValue="1",LinkedTo=(K2Node_CallArrayFunction_14 FA1FAA9B417BE094A04382A683CFD740,),PersistentGuid=00000000000000000000000000000000,bHidden=False,bNotConnectable=False,bDefaultValueIsReadOnly=False,bDefaultValueIsIgnored=False,bAdvancedView=False,)
+                   CustomProperties Pin (PinId=B21F99C242DC61162987FF99AF34760F,PinName="ReturnValue",PinToolTip="Return Value\nInteger",Direction="EGPD_Output",PinType.PinCategory="int",PinType.PinSubCategory="",PinType.PinSubCategoryObject=None,PinType.PinSubCategoryMemberReference=(),PinType.PinValueType=(),PinType.bIsMap=False,PinType.bIsSet=False,PinType.bIsArray=False,PinType.bIsReference=False,PinType.bIsConst=False,PinType.bIsWeakPointer=False,AutogeneratedDefaultValue="0",LinkedTo=(K2Node_CallArrayFunction_17 7B92B2804499E72162447D8B5307E76C,),PersistentGuid=00000000000000000000000000000000,bHidden=False,bNotConnectable=False,bDefaultValueIsReadOnly=False,bDefaultValueIsIgnored=False,bAdvancedView=False,)
+                End Object
+                Begin Object Class=K2Node_CallArrayFunction Name="K2Node_CallArrayFunction_17"
+                   bIsPureFunc=True
+                   FunctionReference=(MemberParent=Class'/Script/Engine.KismetArrayLibrary',MemberName="Array_Get")
+                   NodePosX=2096
+                   NodePosY=-256
+                   NodeGuid=790749664302FAAEE07771AF385841D0
+                   CustomProperties Pin (PinId=8E4FE52F403A37B7A080048D77300710,PinName="self",PinFriendlyName=NSLOCTEXT("K2Node", "Target", "Target"),PinToolTip="Target\nKismet Array Library Reference",PinType.PinCategory="object",PinType.PinSubCategory="",PinType.PinSubCategoryObject=Class'/Script/Engine.KismetArrayLibrary',PinType.PinSubCategoryMemberReference=(),PinType.PinValueType=(),PinType.bIsMap=False,PinType.bIsSet=False,PinType.bIsArray=False,PinType.bIsReference=False,PinType.bIsConst=False,PinType.bIsWeakPointer=False,DefaultObject="/Script/Engine.Default__KismetArrayLibrary",PersistentGuid=00000000000000000000000000000000,bHidden=True,bNotConnectable=False,bDefaultValueIsReadOnly=False,bDefaultValueIsIgnored=False,bAdvancedView=False,)
+                   CustomProperties Pin (PinId=01395B804D226B8004689BA1CD742F72,PinName="TargetArray",PinToolTip="Target Array\nArray of Vector s\n\nThe array to get an item from *",PinType.PinCategory="struct",PinType.PinSubCategory="",PinType.PinSubCategoryObject=ScriptStruct'/Script/CoreUObject.Vector',PinType.PinSubCategoryMemberReference=(),PinType.PinValueType=(),PinType.bIsMap=False,PinType.bIsSet=False,PinType.bIsArray=True,PinType.bIsReference=True,PinType.bIsConst=True,PinType.bIsWeakPointer=False,LinkedTo=(K2Node_CallFunction_175 A0ED07004023EA0AA18F56AE3DE49023,),PersistentGuid=00000000000000000000000000000000,bHidden=False,bNotConnectable=False,bDefaultValueIsReadOnly=False,bDefaultValueIsIgnored=True,bAdvancedView=False,)
+                   CustomProperties Pin (PinId=7B92B2804499E72162447D8B5307E76C,PinName="Index",PinToolTip="Index\nInteger\n\nThe index in the array to get an item from *",PinType.PinCategory="int",PinType.PinSubCategory="",PinType.PinSubCategoryObject=None,PinType.PinSubCategoryMemberReference=(),PinType.PinValueType=(),PinType.bIsMap=False,PinType.bIsSet=False,PinType.bIsArray=False,PinType.bIsReference=False,PinType.bIsConst=False,PinType.bIsWeakPointer=False,DefaultValue="0",AutogeneratedDefaultValue="0",LinkedTo=(K2Node_CallFunction_555 B21F99C242DC61162987FF99AF34760F,),PersistentGuid=00000000000000000000000000000000,bHidden=False,bNotConnectable=False,bDefaultValueIsReadOnly=False,bDefaultValueIsIgnored=False,bAdvancedView=False,)
+                   CustomProperties Pin (PinId=807DD0AC440EB715E49E72ABE52ABBB5,PinName="Item",PinToolTip="Item\nVector ",Direction="EGPD_Output",PinType.PinCategory="struct",PinType.PinSubCategory="",PinType.PinSubCategoryObject=ScriptStruct'/Script/CoreUObject.Vector',PinType.PinSubCategoryMemberReference=(),PinType.PinValueType=(),PinType.bIsMap=False,PinType.bIsSet=False,PinType.bIsArray=False,PinType.bIsReference=False,PinType.bIsConst=False,PinType.bIsWeakPointer=False,AutogeneratedDefaultValue="0",LinkedTo=(K2Node_CallFunction_971 EB02B76B45333DD21AF3A2B5B4F3E367,),PersistentGuid=00000000000000000000000000000000,bHidden=False,bNotConnectable=False,bDefaultValueIsReadOnly=False,bDefaultValueIsIgnored=False,bAdvancedView=False,)
+                End Object
+                Begin Object Class=K2Node_CallArrayFunction Name="K2Node_CallArrayFunction_16"
+                   bIsPureFunc=True
+                   FunctionReference=(MemberParent=Class'/Script/Engine.KismetArrayLibrary',MemberName="Array_Get")
+                   NodePosX=2096
+                   NodePosY=-352
+                   NodeGuid=57FA58A741291929787D30A674147416
+                   CustomProperties Pin (PinId=8E4FE52F403A37B7A080048D77300710,PinName="self",PinFriendlyName=NSLOCTEXT("K2Node", "Target", "Target"),PinToolTip="Target\nKismet Array Library Reference",PinType.PinCategory="object",PinType.PinSubCategory="",PinType.PinSubCategoryObject=Class'/Script/Engine.KismetArrayLibrary',PinType.PinSubCategoryMemberReference=(),PinType.PinValueType=(),PinType.bIsMap=False,PinType.bIsSet=False,PinType.bIsArray=False,PinType.bIsReference=False,PinType.bIsConst=False,PinType.bIsWeakPointer=False,DefaultObject="/Script/Engine.Default__KismetArrayLibrary",PersistentGuid=00000000000000000000000000000000,bHidden=True,bNotConnectable=False,bDefaultValueIsReadOnly=False,bDefaultValueIsIgnored=False,bAdvancedView=False,)
+                   CustomProperties Pin (PinId=01395B804D226B8004689BA1CD742F72,PinName="TargetArray",PinToolTip="Target Array\nArray of Vector s\n\nThe array to get an item from *",PinType.PinCategory="struct",PinType.PinSubCategory="",PinType.PinSubCategoryObject=ScriptStruct'/Script/CoreUObject.Vector',PinType.PinSubCategoryMemberReference=(),PinType.PinValueType=(),PinType.bIsMap=False,PinType.bIsSet=False,PinType.bIsArray=True,PinType.bIsReference=True,PinType.bIsConst=True,PinType.bIsWeakPointer=False,LinkedTo=(K2Node_CallFunction_175 A0ED07004023EA0AA18F56AE3DE49023,),PersistentGuid=00000000000000000000000000000000,bHidden=False,bNotConnectable=False,bDefaultValueIsReadOnly=False,bDefaultValueIsIgnored=True,bAdvancedView=False,)
+                   CustomProperties Pin (PinId=7B92B2804499E72162447D8B5307E76C,PinName="Index",PinToolTip="Index\nInteger\n\nThe index in the array to get an item from *",PinType.PinCategory="int",PinType.PinSubCategory="",PinType.PinSubCategoryObject=None,PinType.PinSubCategoryMemberReference=(),PinType.PinValueType=(),PinType.bIsMap=False,PinType.bIsSet=False,PinType.bIsArray=False,PinType.bIsReference=False,PinType.bIsConst=False,PinType.bIsWeakPointer=False,DefaultValue="0",AutogeneratedDefaultValue="0",LinkedTo=(K2Node_MacroInstance_15 485524BA49C6E2F23DF0DC9247CD6B93,),PersistentGuid=00000000000000000000000000000000,bHidden=False,bNotConnectable=False,bDefaultValueIsReadOnly=False,bDefaultValueIsIgnored=False,bAdvancedView=False,)
+                   CustomProperties Pin (PinId=807DD0AC440EB715E49E72ABE52ABBB5,PinName="Item",PinToolTip="Item\nVector ",Direction="EGPD_Output",PinType.PinCategory="struct",PinType.PinSubCategory="",PinType.PinSubCategoryObject=ScriptStruct'/Script/CoreUObject.Vector',PinType.PinSubCategoryMemberReference=(),PinType.PinValueType=(),PinType.bIsMap=False,PinType.bIsSet=False,PinType.bIsArray=False,PinType.bIsReference=False,PinType.bIsConst=False,PinType.bIsWeakPointer=False,AutogeneratedDefaultValue="0",LinkedTo=(K2Node_CallFunction_970 3949361F44C6842D1E1E52A6D30EBB44,K2Node_CallFunction_971 459C46264271A2417FAABDAED1C20724,),PersistentGuid=00000000000000000000000000000000,bHidden=False,bNotConnectable=False,bDefaultValueIsReadOnly=False,bDefaultValueIsIgnored=False,bAdvancedView=False,)
+                End Object
+                Begin Object Class=K2Node_CallFunction Name="K2Node_CallFunction_970"
+                   FunctionReference=(MemberParent=Class'/Script/Engine.KismetSystemLibrary',MemberName="DrawDebugPoint")
+                   NodePosX=2320
+                   NodePosY=-480
+                   NodeGuid=13548C7549C296F678A27F8D67E9C576
+                   EnabledState=DevelopmentOnly
+                   CustomProperties Pin (PinId=5C92A48C4EC3E621FB3186B24B7F49C9,PinName="execute",PinToolTip="\nExec",PinType.PinCategory="exec",PinType.PinSubCategory="",PinType.PinSubCategoryObject=None,PinType.PinSubCategoryMemberReference=(),PinType.PinValueType=(),PinType.bIsMap=False,PinType.bIsSet=False,PinType.bIsArray=False,PinType.bIsReference=False,PinType.bIsConst=False,PinType.bIsWeakPointer=False,LinkedTo=(K2Node_MacroInstance_15 585876994C3082C336E4F08AB6707724,),PersistentGuid=00000000000000000000000000000000,bHidden=False,bNotConnectable=False,bDefaultValueIsReadOnly=False,bDefaultValueIsIgnored=False,bAdvancedView=False,)
+                   CustomProperties Pin (PinId=5CCCAED249A381CA5387F0B65B573277,PinName="then",PinToolTip="\nExec",Direction="EGPD_Output",PinType.PinCategory="exec",PinType.PinSubCategory="",PinType.PinSubCategoryObject=None,PinType.PinSubCategoryMemberReference=(),PinType.PinValueType=(),PinType.bIsMap=False,PinType.bIsSet=False,PinType.bIsArray=False,PinType.bIsReference=False,PinType.bIsConst=False,PinType.bIsWeakPointer=False,LinkedTo=(K2Node_CallFunction_971 2C80F14646574BCD89C88681C0CD95CC,),PersistentGuid=00000000000000000000000000000000,bHidden=False,bNotConnectable=False,bDefaultValueIsReadOnly=False,bDefaultValueIsIgnored=False,bAdvancedView=False,)
+                   CustomProperties Pin (PinId=597B36884742B881A517BE9B578D5621,PinName="self",PinFriendlyName=NSLOCTEXT("K2Node", "Target", "Target"),PinToolTip="Target\nKismet System Library Reference",PinType.PinCategory="object",PinType.PinSubCategory="",PinType.PinSubCategoryObject=Class'/Script/Engine.KismetSystemLibrary',PinType.PinSubCategoryMemberReference=(),PinType.PinValueType=(),PinType.bIsMap=False,PinType.bIsSet=False,PinType.bIsArray=False,PinType.bIsReference=False,PinType.bIsConst=False,PinType.bIsWeakPointer=False,DefaultObject="/Script/Engine.Default__KismetSystemLibrary",PersistentGuid=00000000000000000000000000000000,bHidden=True,bNotConnectable=False,bDefaultValueIsReadOnly=False,bDefaultValueIsIgnored=False,bAdvancedView=False,)
+                   CustomProperties Pin (PinId=508C8C4B4E004C569A21219E8F775F02,PinName="WorldContextObject",PinToolTip="World Context Object\nObject Reference",PinType.PinCategory="object",PinType.PinSubCategory="",PinType.PinSubCategoryObject=Class'/Script/CoreUObject.Object',PinType.PinSubCategoryMemberReference=(),PinType.PinValueType=(),PinType.bIsMap=False,PinType.bIsSet=False,PinType.bIsArray=False,PinType.bIsReference=False,PinType.bIsConst=False,PinType.bIsWeakPointer=False,PersistentGuid=00000000000000000000000000000000,bHidden=True,bNotConnectable=False,bDefaultValueIsReadOnly=False,bDefaultValueIsIgnored=False,bAdvancedView=False,)
+                   CustomProperties Pin (PinId=3949361F44C6842D1E1E52A6D30EBB44,PinName="Position",PinToolTip="Position\nVector ",PinType.PinCategory="struct",PinType.PinSubCategory="",PinType.PinSubCategoryObject=ScriptStruct'/Script/CoreUObject.Vector',PinType.PinSubCategoryMemberReference=(),PinType.PinValueType=(),PinType.bIsMap=False,PinType.bIsSet=False,PinType.bIsArray=False,PinType.bIsReference=False,PinType.bIsConst=True,PinType.bIsWeakPointer=False,DefaultValue="0, 0, 0",AutogeneratedDefaultValue="0, 0, 0",LinkedTo=(K2Node_CallArrayFunction_16 807DD0AC440EB715E49E72ABE52ABBB5,),PersistentGuid=00000000000000000000000000000000,bHidden=False,bNotConnectable=False,bDefaultValueIsReadOnly=False,bDefaultValueIsIgnored=False,bAdvancedView=False,)
+                   CustomProperties Pin (PinId=371502C8439CAF413C6B25BA4B87520B,PinName="Size",PinToolTip="Size\nFloat",PinType.PinCategory="float",PinType.PinSubCategory="",PinType.PinSubCategoryObject=None,PinType.PinSubCategoryMemberReference=(),PinType.PinValueType=(),PinType.bIsMap=False,PinType.bIsSet=False,PinType.bIsArray=False,PinType.bIsReference=False,PinType.bIsConst=False,PinType.bIsWeakPointer=False,DefaultValue="10",AutogeneratedDefaultValue="0.0",PersistentGuid=00000000000000000000000000000000,bHidden=False,bNotConnectable=False,bDefaultValueIsReadOnly=False,bDefaultValueIsIgnored=False,bAdvancedView=False,)
+                   CustomProperties Pin (PinId=B803AC1A40FE7FB0C1911C903B70ED08,PinName="PointColor",PinToolTip="Point Color\nLinear Color Structure",PinType.PinCategory="struct",PinType.PinSubCategory="",PinType.PinSubCategoryObject=ScriptStruct'/Script/CoreUObject.LinearColor',PinType.PinSubCategoryMemberReference=(),PinType.PinValueType=(),PinType.bIsMap=False,PinType.bIsSet=False,PinType.bIsArray=False,PinType.bIsReference=False,PinType.bIsConst=False,PinType.bIsWeakPointer=False,DefaultValue="(R=0.000000,G=1.000000,B=0.085370,A=1.000000)",PersistentGuid=00000000000000000000000000000000,bHidden=False,bNotConnectable=False,bDefaultValueIsReadOnly=False,bDefaultValueIsIgnored=False,bAdvancedView=False,)
+                   CustomProperties Pin (PinId=9A97752C4C102DDC57B4E6A567881E77,PinName="Duration",PinToolTip="Duration\nFloat",PinType.PinCategory="float",PinType.PinSubCategory="",PinType.PinSubCategoryObject=None,PinType.PinSubCategoryMemberReference=(),PinType.PinValueType=(),PinType.bIsMap=False,PinType.bIsSet=False,PinType.bIsArray=False,PinType.bIsReference=False,PinType.bIsConst=False,PinType.bIsWeakPointer=False,DefaultValue="0.000000",AutogeneratedDefaultValue="0.000000",PersistentGuid=00000000000000000000000000000000,bHidden=False,bNotConnectable=False,bDefaultValueIsReadOnly=False,bDefaultValueIsIgnored=False,bAdvancedView=False,)
+                End Object
+                Begin Object Class=K2Node_CallFunction Name="K2Node_CallFunction_971"
+                   FunctionReference=(MemberParent=Class'/Script/Engine.KismetSystemLibrary',MemberName="DrawDebugLine")
+                   NodePosX=2608
+                   NodePosY=-304
+                   NodeGuid=1CCED43144BB32BA3C6332B6C3872B5A
+                   EnabledState=DevelopmentOnly
+                   CustomProperties Pin (PinId=2C80F14646574BCD89C88681C0CD95CC,PinName="execute",PinToolTip="\nExec",PinType.PinCategory="exec",PinType.PinSubCategory="",PinType.PinSubCategoryObject=None,PinType.PinSubCategoryMemberReference=(),PinType.PinValueType=(),PinType.bIsMap=False,PinType.bIsSet=False,PinType.bIsArray=False,PinType.bIsReference=False,PinType.bIsConst=False,PinType.bIsWeakPointer=False,LinkedTo=(K2Node_CallFunction_970 5CCCAED249A381CA5387F0B65B573277,),PersistentGuid=00000000000000000000000000000000,bHidden=False,bNotConnectable=False,bDefaultValueIsReadOnly=False,bDefaultValueIsIgnored=False,bAdvancedView=False,)
+                   CustomProperties Pin (PinId=37772E68420821E6D8FE989492B623A9,PinName="then",PinToolTip="\nExec",Direction="EGPD_Output",PinType.PinCategory="exec",PinType.PinSubCategory="",PinType.PinSubCategoryObject=None,PinType.PinSubCategoryMemberReference=(),PinType.PinValueType=(),PinType.bIsMap=False,PinType.bIsSet=False,PinType.bIsArray=False,PinType.bIsReference=False,PinType.bIsConst=False,PinType.bIsWeakPointer=False,PersistentGuid=00000000000000000000000000000000,bHidden=False,bNotConnectable=False,bDefaultValueIsReadOnly=False,bDefaultValueIsIgnored=False,bAdvancedView=False,)
+                   CustomProperties Pin (PinId=AD09356A4D6E12DEB6874EAB7A066595,PinName="self",PinFriendlyName=NSLOCTEXT("K2Node", "Target", "Target"),PinToolTip="Target\nKismet System Library Reference",PinType.PinCategory="object",PinType.PinSubCategory="",PinType.PinSubCategoryObject=Class'/Script/Engine.KismetSystemLibrary',PinType.PinSubCategoryMemberReference=(),PinType.PinValueType=(),PinType.bIsMap=False,PinType.bIsSet=False,PinType.bIsArray=False,PinType.bIsReference=False,PinType.bIsConst=False,PinType.bIsWeakPointer=False,DefaultObject="/Script/Engine.Default__KismetSystemLibrary",PersistentGuid=00000000000000000000000000000000,bHidden=True,bNotConnectable=False,bDefaultValueIsReadOnly=False,bDefaultValueIsIgnored=False,bAdvancedView=False,)
+                   CustomProperties Pin (PinId=353C68CE4F2B2D72ADDAEBA325715F00,PinName="WorldContextObject",PinToolTip="World Context Object\nObject Reference",PinType.PinCategory="object",PinType.PinSubCategory="",PinType.PinSubCategoryObject=Class'/Script/CoreUObject.Object',PinType.PinSubCategoryMemberReference=(),PinType.PinValueType=(),PinType.bIsMap=False,PinType.bIsSet=False,PinType.bIsArray=False,PinType.bIsReference=False,PinType.bIsConst=False,PinType.bIsWeakPointer=False,PersistentGuid=00000000000000000000000000000000,bHidden=True,bNotConnectable=False,bDefaultValueIsReadOnly=False,bDefaultValueIsIgnored=False,bAdvancedView=False,)
+                   CustomProperties Pin (PinId=459C46264271A2417FAABDAED1C20724,PinName="LineStart",PinToolTip="Line Start\nVector ",PinType.PinCategory="struct",PinType.PinSubCategory="",PinType.PinSubCategoryObject=ScriptStruct'/Script/CoreUObject.Vector',PinType.PinSubCategoryMemberReference=(),PinType.PinValueType=(),PinType.bIsMap=False,PinType.bIsSet=False,PinType.bIsArray=False,PinType.bIsReference=False,PinType.bIsConst=True,PinType.bIsWeakPointer=False,DefaultValue="0, 0, 0",AutogeneratedDefaultValue="0, 0, 0",LinkedTo=(K2Node_CallArrayFunction_16 807DD0AC440EB715E49E72ABE52ABBB5,),PersistentGuid=00000000000000000000000000000000,bHidden=False,bNotConnectable=False,bDefaultValueIsReadOnly=False,bDefaultValueIsIgnored=False,bAdvancedView=False,)
+                   CustomProperties Pin (PinId=EB02B76B45333DD21AF3A2B5B4F3E367,PinName="LineEnd",PinToolTip="Line End\nVector ",PinType.PinCategory="struct",PinType.PinSubCategory="",PinType.PinSubCategoryObject=ScriptStruct'/Script/CoreUObject.Vector',PinType.PinSubCategoryMemberReference=(),PinType.PinValueType=(),PinType.bIsMap=False,PinType.bIsSet=False,PinType.bIsArray=False,PinType.bIsReference=False,PinType.bIsConst=True,PinType.bIsWeakPointer=False,DefaultValue="0, 0, 0",AutogeneratedDefaultValue="0, 0, 0",LinkedTo=(K2Node_CallArrayFunction_17 807DD0AC440EB715E49E72ABE52ABBB5,),PersistentGuid=00000000000000000000000000000000,bHidden=False,bNotConnectable=False,bDefaultValueIsReadOnly=False,bDefaultValueIsIgnored=False,bAdvancedView=False,)
+                   CustomProperties Pin (PinId=F46FC3744B82648A65ABC8BB8C6323BC,PinName="LineColor",PinToolTip="Line Color\nLinear Color Structure",PinType.PinCategory="struct",PinType.PinSubCategory="",PinType.PinSubCategoryObject=ScriptStruct'/Script/CoreUObject.LinearColor',PinType.PinSubCategoryMemberReference=(),PinType.PinValueType=(),PinType.bIsMap=False,PinType.bIsSet=False,PinType.bIsArray=False,PinType.bIsReference=False,PinType.bIsConst=False,PinType.bIsWeakPointer=False,DefaultValue="(R=1.000000,G=0.014920,B=0.000000,A=1.000000)",PersistentGuid=00000000000000000000000000000000,bHidden=False,bNotConnectable=False,bDefaultValueIsReadOnly=False,bDefaultValueIsIgnored=False,bAdvancedView=False,)
+                   CustomProperties Pin (PinId=C2B19A7E453BE596D4D31EAC0782E5DD,PinName="Duration",PinToolTip="Duration\nFloat",PinType.PinCategory="float",PinType.PinSubCategory="",PinType.PinSubCategoryObject=None,PinType.PinSubCategoryMemberReference=(),PinType.PinValueType=(),PinType.bIsMap=False,PinType.bIsSet=False,PinType.bIsArray=False,PinType.bIsReference=False,PinType.bIsConst=False,PinType.bIsWeakPointer=False,DefaultValue="0",AutogeneratedDefaultValue="0.000000",PersistentGuid=00000000000000000000000000000000,bHidden=False,bNotConnectable=False,bDefaultValueIsReadOnly=False,bDefaultValueIsIgnored=False,bAdvancedView=False,)
+                   CustomProperties Pin (PinId=3C5AC041439E0A0D435858B3A3A854CF,PinName="Thickness",PinToolTip="Thickness\nFloat",PinType.PinCategory="float",PinType.PinSubCategory="",PinType.PinSubCategoryObject=None,PinType.PinSubCategoryMemberReference=(),PinType.PinValueType=(),PinType.bIsMap=False,PinType.bIsSet=False,PinType.bIsArray=False,PinType.bIsReference=False,PinType.bIsConst=False,PinType.bIsWeakPointer=False,DefaultValue="1",AutogeneratedDefaultValue="0.000000",PersistentGuid=00000000000000000000000000000000,bHidden=False,bNotConnectable=False,bDefaultValueIsReadOnly=False,bDefaultValueIsIgnored=False,bAdvancedView=False,)
+                End Object
+                Begin Object Class=EdGraphNode_Comment Name="EdGraphNode_Comment_3"
+                   NodePosX=608
+                   NodePosY=-560
+                   NodeWidth=2192
+                   NodeHeight=480
+                   NodeComment="Draw the play area"
+                   NodeGuid=5D41106C4E1F47307BF07AAC129F72E9
+                End Object
+            END OBJECT
+    [/PARAM]
+[/OBJECT]
+[/EXCERPT:DrawPlayArea]
+-->
\ No newline at end of file
diff --git a/Engine/Documentation/Source/Platforms/Oculus/HowTo/GuardianSystem/HTOculusGuardianSystem.INT.udn b/Engine/Documentation/Source/Platforms/Oculus/HowTo/GuardianSystem/HTOculusGuardianSystem.INT.udn
new file mode 100644
--- /dev/null
+++ b/Engine/Documentation/Source/Platforms/Oculus/HowTo/GuardianSystem/HTOculusGuardianSystem.INT.udn
@@ -0,0 +1,69 @@
+Availability: Public
+Crumbs: %ROOT%, 
+Title: Setup the Oculus Rift Guardian System
+Description: Setting up the Rift Guardian System for room scale VR.
+Type: how-to
+Skilllevel: Beginner
+Version: 4.17
+Parent: Platforms/Oculus
+Order: 2
+tags: Basics
+topic-image:HTGuardian_System_Topic_Image.png
+
+The Oculus Guardian system is used to display the boundaries of the VR interaction area. As a tracked device approaches the boundaries, the Oculus Runtime automatically provides visual cues notifying the user of this. In the following How - To we will take a look at adding the Oculus Guardian system to your UE4 projects player Pawn.
+
+[REGION:note]
+The following How - To is going to make use of the Standing VRPawn with added Motion Controller support. If you are not sure how to set this up the following two How - To's will help you out:
+* [Set Up Standing Cameras for the Oculus Rift](Platforms\Oculus\HowTo\StandingCamera)
+* [Set Up Touch Controllers for the Oculus Rift](Platforms\Oculus\HowTo\TouchControllers)
+[/REGION]
+
+[REGION:note]
+For the Guardian system to work, you will need to make sure you have set it up using the Oculus App. For more information on how to do this, please see the official [Oculus Guardian System](https://developer.oculus.com/documentation/pcsdk/latest/concepts/dg-guardian-system/) Setup pages.
+[/REGION]
+
+[REGION:warning]
+It is **NOT** possible nor advisable to disable the Guardian system from inside UE4. You can, however, adjust how UE4 will react to the user getting close to the boundaries.  
+[/REGION]
+
+## Oculus Rift Guardian System Setup
+
+To add the Oculus Rift Guardian system to your UE4 player pawn, you will need to do the following:
+
+1. Open up your projects player Pawn Blueprint and make sure the **Components** tab is visible.  
+
+1. Click on the **Add Component** button and from the displayed list search for the **OculusRiftBoundy** component. When located, click on it to add it to the Component list.
+
+    ![](HTGuardian_System_00.png)
+
+1. When completed, your player pawn should look like the following image.
+
+    ![](HTGuardian_System_01.png)
+
+1. Now, launch your project in VR and put on your Rift HMD and when you approach the bounds of your VR interaction area, you should see something similar to the following video.
+
+    [OBJECT:EmbeddedVideo]
+        [PARAMLITERAL:videoid]
+        AGgUNDZoObQ
+        [/PARAMLITERAL]
+        [PARAMLITERAL:width]
+        100
+        [/PARAMLITERAL]
+        [PARAMLITERAL:height]
+        56.25
+        [/PARAMLITERAL]
+        [PARAMLITERAL:units]
+        %
+        [/PARAMLITERAL]
+        [PARAMLITERAL:type]
+        simple
+        [/PARAMLITERAL]
+    [/OBJECT]
+
+## UE4 Project Downloads
+
+Below you will find a link to where you can download the UE4 project that was used to create this example. 
+
+* [**Oculus Rift Guardian System Set Up Project**](https://epicgames.box.com/s/s4vvkb2i0ohtice8vtqude6i8ih7hy8i)
+
+
diff --git a/Engine/Documentation/Source/Platforms/Oculus/HowTo/OculusHowTo.INT.udn b/Engine/Documentation/Source/Platforms/Oculus/HowTo/OculusHowTo.INT.udn
new file mode 100644
--- /dev/null
+++ b/Engine/Documentation/Source/Platforms/Oculus/HowTo/OculusHowTo.INT.udn
@@ -0,0 +1,31 @@
+Availability: Public
+Crumbs: %ROOT%,  Platforms
+Title: Oculus How To
+Description: How To guides providing step-by-step instructions working with the Rift and UE4.
+Type: how-to
+Version: 4.17
+Parent: Platforms/Oculus
+Order: 1
+Tags: Platform
+topic-image:Oculus_How_To_Header.png
+
+The Oculus How To's contain step-by-step guides that will walk you through how to use UE4 and the Oculus Rift together.
+
+## Beginner
+
+The Oculus Rift Beginner How - To's will walk you through setting up your UE4 project to work with the Rift HMD, Touch Controllers and the Guardian system.
+
+[DIR(end:"1" output:"topic" type:"how-to" tags:"Basics" parent:"Platforms/Oculus")]
+
+## Intermediate
+
+The Oculus Rift Intermediate How - To's will walk you through adjusting the color of the Guardian System, using Touch Controllers to pick objects up and also how to use Stereo Layers to display information to the user.
+
+[DIR(end:"1" output:"topic" type:"how-to" tags:"Oculus" parent:"Platforms/Oculus")]
+
+## Advanced
+
+The Oculus Rift Advanced How - To's will walk you through detecting if the users has activated the Guardian system and how to setup and use loading Splash Screens. 
+
+[DIR(end:"1" output:"topic" type:"how-to" tags:"VR" parent:"Platforms/Oculus")]
+
diff --git a/Engine/Documentation/Source/Platforms/Oculus/HowTo/SeatedCamera/HTOculusSeatedCamera.INT.udn b/Engine/Documentation/Source/Platforms/Oculus/HowTo/SeatedCamera/HTOculusSeatedCamera.INT.udn
new file mode 100644
--- /dev/null
+++ b/Engine/Documentation/Source/Platforms/Oculus/HowTo/SeatedCamera/HTOculusSeatedCamera.INT.udn
@@ -0,0 +1,230 @@
+Availability: Public
+Crumbs: %ROOT%, Platforms
+Title:Set Up a Seated Camera for the Oculus Rift
+Description:Setting up the UE4 camera to work with a seated Oculus Rift experience.
+Type: how-to
+Skilllevel:Beginner
+Version:4.17
+Parent:Platforms/Oculus/HowTo
+Order: 2
+tags:Basics
+topic-image:HT_Rift_Camera_Setup_Topic_Image.png
+
+One of the first things you will need to think about when starting to develop any UE4 powered VR experience to be used with the Oculus Rift is deciding whether the experience will be a seated one or a standing one. In the following How - To we will take a look at how to setup your UE4 projects VR camera for a seated Oculus Rift experience.
+
+## Pawn Setup for a Seated Rift Experience 
+
+Below you will find instructions on how to setup up any Pawn to work with a seated Oculus Rift experience.
+
+1. First, open up or create a new Pawn Blueprint and then go to the **Component** section of the **Viewport** tab. From there add the following two components with the following names, making sure that the VRCamera is a child of VRCameraRoot:
+    
+    * **Scene** - VRCameraRoot
+    * **Camera** - VRCamera
+
+    [REGION:lightbox]
+    [![](HT_Rift_Camera_Setup_03.png)(w:224)](HT_Rift_Camera_Setup_03.png)
+    [/REGION]
+
+    [REGION:caption]
+    Click for full image.
+    [/REGION]
+
+    When 
+
+    [region:note]
+    No matter what VR HMD you are using, this is the Epic recommend way to set up a VR camera as it allows you to offset the position of the camera without having to move the actual camera.
+    [/region]
+
+With that completed, the initial setup for the VR Camera is now completed, in the next sections, we will go through setting the needed Blueprint logic for a seated Oculus Rift experience.
+
+## Blueprint Setup for Seated Rift Experience 
+
+Below you will find instructions on how to setup the Blueprint logic for a seated Oculus Rift experience.  
+
+1. First, open up the Pawn Blueprint if not already open and in the **Event Graph** drag off the **Event Begin Play** node to show the Executable Actions list. In the list search for the **Set Tracking Origin** node and click on it to add it to the Event Graph.
+
+    [REGION:lightbox]
+    [![](HT_Rift_Camera_Setup_09.png)(w:555)](HT_Rift_Camera_Setup_09.png)
+    [/REGION]
+
+    [REGION:caption]
+    Click for full image.
+    [/REGION]
+
+1. The Set Tracking Origin node has two options, **Floor Level** and **Eye Level**. For a seated experience you will need to set the **Origin** of the Set Tracking Origin node to **Eye Level**.
+    
+    [REGION:lightbox]
+    [![](HT_Rift_Camera_Setup_10.png)(w:555)](HT_Rift_Camera_Setup_10.png)
+    [/REGION]
+
+    [REGION:caption]
+    Click for full image.
+    [/REGION]
+
+1. Next, in the **Variables**  section of the **My Blueprint** tab, create a new **Vector** variable called **RiftCameraHeight** and set the **Z** value to **121**.
+
+    ![](HT_Rift_Camera_Setup_05.png)
+
+    [region:note]
+    For a seated Rift experience, you will always want to set the height of the camera to match the real world seated height of the user in Centimeters. 
+    [/region]
+
+1. Then, drag off the output of the Set Tracking Origin node and search for the **Set Relative Location** node, selecting the **SetRelativeLocation(VRCameraRoot)** option.
+
+    [REGION:lightbox]
+    [![](HT_Rift_Camera_Setup_11.png)(w:472)](HT_Rift_Camera_Setup_11.png)
+    [/REGION]
+
+    [REGION:caption]
+    Click for full image.
+    [/REGION]
+
+1. Connect the **RiftCameraHeight** variable to the **New Location** input on the Set Relative Location node and then press the Compile button and when completed your Event Graph should look like the following image.
+
+    [REGION:fullwidth]
+    [INCLUDE:Platforms\Oculus\HowTo\SeatedCamera\#OculusSittingVRCamera]
+    [/REGION]
+
+    [region:tip]
+    You can copy the completed Blueprint by clicking on the upper left-hand corner of above and pressing CRTL + C to copy the graph. Once copied, go to your Blueprints event and press CTRL + V to paste it in. 
+    [/region]
+
+1. Drag the Pawn Blueprint from the Content Browser to a level, making sure that it is placed at 0,0,0 in the level.
+    [REGION:lightbox]
+    [![](HT_Rift_Camera_Setup_06.png)(w:555)](HT_Rift_Camera_Setup_06.png)
+    [/REGION]
+
+    [REGION:caption]
+    Click for full image.
+    [/REGION]
+
+1. Select the Pawn Blueprint that was put in the level and in the **Details** panel under the **Pawn** setting, set the **Auto Possess Player** from **Disabled** to **Player 0**.
+
+    [REGION:lightbox]
+    [![](HT_Rift_Camera_Setup_07.png)(w:225)](HT_Rift_Camera_Setup_07.png)
+    [/REGION]
+
+    [REGION:caption]
+    Click for full image.
+    [/REGION]
+
+1. Finally, go to the **Main Toolbar** and change the **Play Mode** to **VR Preview** and then press the **Play** button. When you put your Oculus Rift HMD on and view the level while sitting down, you should see something similar to the following video. 
+
+    [OBJECT:EmbeddedVideo]
+        [PARAMLITERAL:videoid]
+        ErQoJ2n5EpI
+        [/PARAMLITERAL]
+        [PARAMLITERAL:width]
+        100
+        [/PARAMLITERAL]
+        [PARAMLITERAL:height]
+        56.25
+        [/PARAMLITERAL]
+        [PARAMLITERAL:units]
+        %
+        [/PARAMLITERAL]
+        [PARAMLITERAL:type]
+        simple
+        [/PARAMLITERAL]
+    [/OBJECT]
+
+## UE4 Project Downloads
+Below you will find a link to where you can download the UE4 project that was used to create this example. 
+
+* [**Oculus Rift Seated VR Camera Example Project**](https://epicgames.box.com/s/jrhmrpvtk84jrrw05m5ma5rvcs35zm3g)
+
+<!--
+[EXCERPT:OculusSittingVRCamera]	
+[OBJECT:Snippet]
+	[PARAMLITERAL:id]
+		OculusSittingVRCamera
+	[/PARAMLITERAL]
+	[PARAMLITERAL:type]
+		blueprint
+	[/PARAMLITERAL]
+	[PARAM:image]
+		[REGION:lightbox]
+		[![](HT_Rift_Camera_Setup_04.png)(w:1030)](HT_Rift_Camera_Setup_04.png)
+		[/REGION]
+	[/PARAM]
+	[PARAM:snippet]
+			BEGIN OBJECT
+							Begin Object Class=/Script/BlueprintGraph.K2Node_Event Name="K2Node_Event_0"
+			   EventReference=(MemberParent=Class'/Script/Engine.Actor',MemberName="ReceiveBeginPlay")
+			   bOverrideFunction=True
+			   NodePosX=-272
+			   NodePosY=-16
+			   bCommentBubblePinned=True
+			   NodeGuid=0A1ED21A4BC4FFF952C8B1BD85346A9C
+			   CustomProperties Pin (PinId=0862FFB34FDE75E7C2ED91B15CD1119E,PinName="OutputDelegate",Direction="EGPD_Output",PinType.PinCategory="delegate",PinType.PinSubCategory="",PinType.PinSubCategoryObject=None,PinType.PinSubCategoryMemberReference=(MemberParent=Class'/Script/Engine.Actor',MemberName="ReceiveBeginPlay"),PinType.PinValueType=(),PinType.ContainerType=None,PinType.bIsMap=False,PinType.bIsSet=False,PinType.bIsArray=False,PinType.bIsReference=False,PinType.bIsConst=False,PinType.bIsWeakPointer=False,PersistentGuid=00000000000000000000000000000000,bHidden=False,bNotConnectable=False,bDefaultValueIsReadOnly=False,bDefaultValueIsIgnored=False,bAdvancedView=False,bOrphanedPin=False,)
+			   CustomProperties Pin (PinId=CE164667454B3EB8FFA464B7925FD171,PinName="then",Direction="EGPD_Output",PinType.PinCategory="exec",PinType.PinSubCategory="",PinType.PinSubCategoryObject=None,PinType.PinSubCategoryMemberReference=(),PinType.PinValueType=(),PinType.ContainerType=None,PinType.bIsMap=False,PinType.bIsSet=False,PinType.bIsArray=False,PinType.bIsReference=False,PinType.bIsConst=False,PinType.bIsWeakPointer=False,LinkedTo=(K2Node_CallFunction_0 9B118DD74D13E1A2489FEFBC87651604,),PersistentGuid=00000000000000000000000000000000,bHidden=False,bNotConnectable=False,bDefaultValueIsReadOnly=False,bDefaultValueIsIgnored=False,bAdvancedView=False,bOrphanedPin=False,)
+			End Object
+			Begin Object Class=/Script/BlueprintGraph.K2Node_CallFunction Name="K2Node_CallFunction_0"
+			   FunctionReference=(MemberParent=Class'/Script/HeadMountedDisplay.HeadMountedDisplayFunctionLibrary',MemberName="SetTrackingOrigin")
+			   NodePosX=64
+			   NodePosY=-16
+			   NodeGuid=48B45C91438E818692990A8ACF803845
+			   CustomProperties Pin (PinId=9B118DD74D13E1A2489FEFBC87651604,PinName="execute",PinToolTip="\nExec",PinType.PinCategory="exec",PinType.PinSubCategory="",PinType.PinSubCategoryObject=None,PinType.PinSubCategoryMemberReference=(),PinType.PinValueType=(),PinType.ContainerType=None,PinType.bIsMap=False,PinType.bIsSet=False,PinType.bIsArray=False,PinType.bIsReference=False,PinType.bIsConst=False,PinType.bIsWeakPointer=False,LinkedTo=(K2Node_Event_0 CE164667454B3EB8FFA464B7925FD171,),PersistentGuid=00000000000000000000000000000000,bHidden=False,bNotConnectable=False,bDefaultValueIsReadOnly=False,bDefaultValueIsIgnored=False,bAdvancedView=False,bOrphanedPin=False,)
+			   CustomProperties Pin (PinId=DB86848D4CD2762032F965A34161FE78,PinName="then",PinToolTip="\nExec",Direction="EGPD_Output",PinType.PinCategory="exec",PinType.PinSubCategory="",PinType.PinSubCategoryObject=None,PinType.PinSubCategoryMemberReference=(),PinType.PinValueType=(),PinType.ContainerType=None,PinType.bIsMap=False,PinType.bIsSet=False,PinType.bIsArray=False,PinType.bIsReference=False,PinType.bIsConst=False,PinType.bIsWeakPointer=False,LinkedTo=(K2Node_CallFunction_9 DE9FB716426D0575C24DBD90B0B7F9B0,),PersistentGuid=00000000000000000000000000000000,bHidden=False,bNotConnectable=False,bDefaultValueIsReadOnly=False,bDefaultValueIsIgnored=False,bAdvancedView=False,bOrphanedPin=False,)
+			   CustomProperties Pin (PinId=CF7045694A0AC748019DB7B16FD84A1F,PinName="self",PinFriendlyName=NSLOCTEXT("K2Node", "Target", "Target"),PinToolTip="Target\nHead Mounted Display Function Library Object Reference",PinType.PinCategory="object",PinType.PinSubCategory="",PinType.PinSubCategoryObject=Class'/Script/HeadMountedDisplay.HeadMountedDisplayFunctionLibrary',PinType.PinSubCategoryMemberReference=(),PinType.PinValueType=(),PinType.ContainerType=None,PinType.bIsMap=False,PinType.bIsSet=False,PinType.bIsArray=False,PinType.bIsReference=False,PinType.bIsConst=False,PinType.bIsWeakPointer=False,DefaultObject="/Script/HeadMountedDisplay.Default__HeadMountedDisplayFunctionLibrary",PersistentGuid=00000000000000000000000000000000,bHidden=True,bNotConnectable=False,bDefaultValueIsReadOnly=False,bDefaultValueIsIgnored=False,bAdvancedView=False,bOrphanedPin=False,)
+			   CustomProperties Pin (PinId=6E59CEE64B1A4E275A73499235912077,PinName="Origin",PinToolTip="Origin\nEHMDTrackingOrigin Enum",PinType.PinCategory="byte",PinType.PinSubCategory="",PinType.PinSubCategoryObject=Enum'/Script/HeadMountedDisplay.EHMDTrackingOrigin',PinType.PinSubCategoryMemberReference=(),PinType.PinValueType=(),PinType.ContainerType=None,PinType.bIsMap=False,PinType.bIsSet=False,PinType.bIsArray=False,PinType.bIsReference=False,PinType.bIsConst=False,PinType.bIsWeakPointer=False,DefaultValue="Eye",PersistentGuid=00000000000000000000000000000000,bHidden=False,bNotConnectable=False,bDefaultValueIsReadOnly=False,bDefaultValueIsIgnored=False,bAdvancedView=False,bOrphanedPin=False,)
+			End Object
+			Begin Object Class=/Script/UnrealEd.EdGraphNode_Comment Name="EdGraphNode_Comment_1"
+			   NodePosX=-54
+			   NodePosY=-81
+			   NodeWidth=544
+			   NodeHeight=240
+			   NodeComment="Tracking Orgin should be set to Eye Level"
+			   NodeGuid=9B24193F4ED6C3FBF64C68B01AAFA782
+			End Object
+			Begin Object Class=/Script/BlueprintGraph.K2Node_VariableGet Name="K2Node_VariableGet_0"
+			   VariableReference=(MemberName="VRCameraRoot",bSelfContext=True)
+			   NodePosX=624
+			   NodePosY=35
+			   NodeGuid=70636A4A41C97ED0AB31D6B2E8E7AE84
+			   CustomProperties Pin (PinId=D3E17C4549E0A126BF8A2D97653C64D5,PinName="VRCameraRoot",Direction="EGPD_Output",PinType.PinCategory="object",PinType.PinSubCategory="",PinType.PinSubCategoryObject=Class'/Script/Engine.SceneComponent',PinType.PinSubCategoryMemberReference=(),PinType.PinValueType=(),PinType.ContainerType=None,PinType.bIsMap=False,PinType.bIsSet=False,PinType.bIsArray=False,PinType.bIsReference=False,PinType.bIsConst=False,PinType.bIsWeakPointer=False,LinkedTo=(K2Node_CallFunction_9 560BA9474CC7FF8D01863CB899932E82,),PersistentGuid=00000000000000000000000000000000,bHidden=False,bNotConnectable=False,bDefaultValueIsReadOnly=False,bDefaultValueIsIgnored=False,bAdvancedView=False,bOrphanedPin=False,)
+			   CustomProperties Pin (PinId=2C82232F4B776E78BDFA17A5BE5C88D7,PinName="self",PinFriendlyName=NSLOCTEXT("K2Node", "Target", "Target"),PinType.PinCategory="object",PinType.PinSubCategory="",PinType.PinSubCategoryObject=BlueprintGeneratedClass'/Game/Rift_Pawn_Sitting/Rift_Pawn_Sitting.Rift_Pawn_Sitting_C',PinType.PinSubCategoryMemberReference=(),PinType.PinValueType=(),PinType.ContainerType=None,PinType.bIsMap=False,PinType.bIsSet=False,PinType.bIsArray=False,PinType.bIsReference=False,PinType.bIsConst=False,PinType.bIsWeakPointer=False,PersistentGuid=00000000000000000000000000000000,bHidden=True,bNotConnectable=False,bDefaultValueIsReadOnly=False,bDefaultValueIsIgnored=False,bAdvancedView=False,bOrphanedPin=False,)
+			End Object
+			Begin Object Class=/Script/BlueprintGraph.K2Node_CallFunction Name="K2Node_CallFunction_9"
+			   FunctionReference=(MemberParent=Class'/Script/Engine.SceneComponent',MemberName="K2_SetRelativeLocation")
+			   NodePosX=800
+			   NodePosY=-31
+			   NodeGuid=4EF05AB24E0197A031C39A8A8B2B8B3A
+			   CustomProperties Pin (PinId=DE9FB716426D0575C24DBD90B0B7F9B0,PinName="execute",PinToolTip="\nExec",PinType.PinCategory="exec",PinType.PinSubCategory="",PinType.PinSubCategoryObject=None,PinType.PinSubCategoryMemberReference=(),PinType.PinValueType=(),PinType.ContainerType=None,PinType.bIsMap=False,PinType.bIsSet=False,PinType.bIsArray=False,PinType.bIsReference=False,PinType.bIsConst=False,PinType.bIsWeakPointer=False,LinkedTo=(K2Node_CallFunction_0 DB86848D4CD2762032F965A34161FE78,),PersistentGuid=00000000000000000000000000000000,bHidden=False,bNotConnectable=False,bDefaultValueIsReadOnly=False,bDefaultValueIsIgnored=False,bAdvancedView=False,bOrphanedPin=False,)
+			   CustomProperties Pin (PinId=E5FB25D74A81ACC9C7E06593A4BC6E98,PinName="then",PinToolTip="\nExec",Direction="EGPD_Output",PinType.PinCategory="exec",PinType.PinSubCategory="",PinType.PinSubCategoryObject=None,PinType.PinSubCategoryMemberReference=(),PinType.PinValueType=(),PinType.ContainerType=None,PinType.bIsMap=False,PinType.bIsSet=False,PinType.bIsArray=False,PinType.bIsReference=False,PinType.bIsConst=False,PinType.bIsWeakPointer=False,PersistentGuid=00000000000000000000000000000000,bHidden=False,bNotConnectable=False,bDefaultValueIsReadOnly=False,bDefaultValueIsIgnored=False,bAdvancedView=False,bOrphanedPin=False,)
+			   CustomProperties Pin (PinId=560BA9474CC7FF8D01863CB899932E82,PinName="self",PinFriendlyName=NSLOCTEXT("K2Node", "Target", "Target"),PinToolTip="Target\nScene Component Object Reference",PinType.PinCategory="object",PinType.PinSubCategory="",PinType.PinSubCategoryObject=Class'/Script/Engine.SceneComponent',PinType.PinSubCategoryMemberReference=(),PinType.PinValueType=(),PinType.ContainerType=None,PinType.bIsMap=False,PinType.bIsSet=False,PinType.bIsArray=False,PinType.bIsReference=False,PinType.bIsConst=False,PinType.bIsWeakPointer=False,LinkedTo=(K2Node_VariableGet_0 D3E17C4549E0A126BF8A2D97653C64D5,),PersistentGuid=00000000000000000000000000000000,bHidden=False,bNotConnectable=False,bDefaultValueIsReadOnly=False,bDefaultValueIsIgnored=False,bAdvancedView=False,bOrphanedPin=False,)
+			   CustomProperties Pin (PinId=4751BC714842E8CEC1B7A294186125C2,PinName="NewLocation",PinToolTip="New Location\nVector \n\nNew location of the component relative to its parent.",PinType.PinCategory="struct",PinType.PinSubCategory="",PinType.PinSubCategoryObject=ScriptStruct'/Script/CoreUObject.Vector',PinType.PinSubCategoryMemberReference=(),PinType.PinValueType=(),PinType.ContainerType=None,PinType.bIsMap=False,PinType.bIsSet=False,PinType.bIsArray=False,PinType.bIsReference=False,PinType.bIsConst=False,PinType.bIsWeakPointer=False,DefaultValue="0, 0, 0",AutogeneratedDefaultValue="0, 0, 0",LinkedTo=(K2Node_VariableGet_47 78E971AE41B9BC4AA335A6BCC9E5AE69,),PersistentGuid=00000000000000000000000000000000,bHidden=False,bNotConnectable=False,bDefaultValueIsReadOnly=False,bDefaultValueIsIgnored=False,bAdvancedView=False,bOrphanedPin=False,)
+			   CustomProperties Pin (PinId=42E17DA8407F1D6AE61E5AB8CE56D509,PinName="bSweep",PinToolTip="Sweep\nBoolean\n\nWhether we sweep to the destination location, triggering overlaps along the way and stopping short of the target if blocked by something. Only the root component is swept and checked for blocking collision, child components move without sweeping. If collision is off, this has no effect.",PinType.PinCategory="bool",PinType.PinSubCategory="",PinType.PinSubCategoryObject=None,PinType.PinSubCategoryMemberReference=(),PinType.PinValueType=(),PinType.ContainerType=None,PinType.bIsMap=False,PinType.bIsSet=False,PinType.bIsArray=False,PinType.bIsReference=False,PinType.bIsConst=False,PinType.bIsWeakPointer=False,DefaultValue="false",AutogeneratedDefaultValue="false",PersistentGuid=00000000000000000000000000000000,bHidden=False,bNotConnectable=False,bDefaultValueIsReadOnly=False,bDefaultValueIsIgnored=False,bAdvancedView=False,bOrphanedPin=False,)
+			   CustomProperties Pin (PinId=C758BE724ED2C35AF55937B09A2B53F4,PinName="SweepHitResult",PinToolTip="Sweep Hit Result\nHit Result Structure\n\nHit result from any impact if sweep is true.",Direction="EGPD_Output",PinType.PinCategory="struct",PinType.PinSubCategory="",PinType.PinSubCategoryObject=ScriptStruct'/Script/Engine.HitResult',PinType.PinSubCategoryMemberReference=(),PinType.PinValueType=(),PinType.ContainerType=None,PinType.bIsMap=False,PinType.bIsSet=False,PinType.bIsArray=False,PinType.bIsReference=False,PinType.bIsConst=False,PinType.bIsWeakPointer=False,PersistentGuid=00000000000000000000000000000000,bHidden=False,bNotConnectable=False,bDefaultValueIsReadOnly=False,bDefaultValueIsIgnored=False,bAdvancedView=False,bOrphanedPin=False,)
+			   CustomProperties Pin (PinId=C8D4DFF74546FCDE47077A8F575C2D36,PinName="bTeleport",PinToolTip="Teleport\nBoolean\n\nWhether we teleport the physics state (if physics collision is enabled for this object). If true, physics velocity for this object is unchanged (so ragdoll parts are not affected by change in location). If false, physics velocity is updated based on the change in position (affecting ragdoll parts). If CCD is on and not teleporting, this will affect objects along the entire sweep volume.",PinType.PinCategory="bool",PinType.PinSubCategory="",PinType.PinSubCategoryObject=None,PinType.PinSubCategoryMemberReference=(),PinType.PinValueType=(),PinType.ContainerType=None,PinType.bIsMap=False,PinType.bIsSet=False,PinType.bIsArray=False,PinType.bIsReference=False,PinType.bIsConst=False,PinType.bIsWeakPointer=False,DefaultValue="false",AutogeneratedDefaultValue="false",PersistentGuid=00000000000000000000000000000000,bHidden=False,bNotConnectable=False,bDefaultValueIsReadOnly=False,bDefaultValueIsIgnored=False,bAdvancedView=False,bOrphanedPin=False,)
+			End Object
+			Begin Object Class=/Script/BlueprintGraph.K2Node_VariableGet Name="K2Node_VariableGet_47"
+			   VariableReference=(MemberName="RiftCameraHeight",MemberGuid=DB0F658D4BD2CC99041C3ABB8058E5EF,bSelfContext=True)
+			   NodePosX=576
+			   NodePosY=272
+			   NodeGuid=3B3F090A44BA690E035446B6EEB8BE25
+			   CustomProperties Pin (PinId=78E971AE41B9BC4AA335A6BCC9E5AE69,PinName="RiftCameraHeight",Direction="EGPD_Output",PinType.PinCategory="struct",PinType.PinSubCategory="",PinType.PinSubCategoryObject=ScriptStruct'/Script/CoreUObject.Vector',PinType.PinSubCategoryMemberReference=(),PinType.PinValueType=(),PinType.ContainerType=None,PinType.bIsMap=False,PinType.bIsSet=False,PinType.bIsArray=False,PinType.bIsReference=False,PinType.bIsConst=False,PinType.bIsWeakPointer=False,DefaultValue="0, 0, 0",AutogeneratedDefaultValue="0, 0, 0",LinkedTo=(K2Node_CallFunction_9 4751BC714842E8CEC1B7A294186125C2,),PersistentGuid=00000000000000000000000000000000,bHidden=False,bNotConnectable=False,bDefaultValueIsReadOnly=False,bDefaultValueIsIgnored=False,bAdvancedView=False,bOrphanedPin=False,)
+			   CustomProperties Pin (PinId=4C94FD9248478C9DD5EA5393075A2807,PinName="self",PinFriendlyName=NSLOCTEXT("K2Node", "Target", "Target"),PinType.PinCategory="object",PinType.PinSubCategory="",PinType.PinSubCategoryObject=BlueprintGeneratedClass'/Game/Rift_Pawn_Sitting/Rift_Pawn_Sitting.Rift_Pawn_Sitting_C',PinType.PinSubCategoryMemberReference=(),PinType.PinValueType=(),PinType.ContainerType=None,PinType.bIsMap=False,PinType.bIsSet=False,PinType.bIsArray=False,PinType.bIsReference=False,PinType.bIsConst=False,PinType.bIsWeakPointer=False,PersistentGuid=00000000000000000000000000000000,bHidden=True,bNotConnectable=False,bDefaultValueIsReadOnly=False,bDefaultValueIsIgnored=False,bAdvancedView=False,bOrphanedPin=False,)
+			End Object
+			Begin Object Class=/Script/UnrealEd.EdGraphNode_Comment Name="EdGraphNode_Comment_2"
+			   CommentDepth=-2
+			   NodePosX=512
+			   NodePosY=-80
+			   NodeWidth=592
+			   NodeHeight=448
+			   NodeComment="Offset the VRCameraRoot to move the Rift Camera"
+			   NodeGuid=E3607E9749F360D6955336BC8E0E7F3A
+			End Object
+			Begin Object Class=/Script/UnrealEd.EdGraphNode_Comment Name="EdGraphNode_Comment_3"
+			   NodePosX=528
+			   NodePosY=192
+			   NodeWidth=512
+			   NodeHeight=160
+			   NodeComment="Set this to the standing height of the user."
+			   NodeGuid=FB8FE99D4F6E5E1DDE16B1BC685412F1
+			End Object
+			END OBJECT
+	[/PARAM]
+[/OBJECT]
+[/EXCERPT:OculusSittingVRCamera]
+-->
\ No newline at end of file
diff --git a/Engine/Documentation/Source/Platforms/Oculus/HowTo/SplashScreens/HTOculusSplashScreens.INT.udn b/Engine/Documentation/Source/Platforms/Oculus/HowTo/SplashScreens/HTOculusSplashScreens.INT.udn
new file mode 100644
--- /dev/null
+++ b/Engine/Documentation/Source/Platforms/Oculus/HowTo/SplashScreens/HTOculusSplashScreens.INT.udn
@@ -0,0 +1,124 @@
+Availability: Public
+Crumbs: %ROOT%, Platforms
+Title: Using Touch Controllers for the Oculus Rift
+Description: Use the Rift Touch controllers to interact with the world.
+Type: how-to
+Skilllevel: Intermediate
+Version: 4.16
+Parent: Platforms/Oculus
+Order: 2
+tags: Oculus
+topic-image:Platforms\Oculus\HowTo\TouchControllers\HT_Rift_Touch_Setup_Topic_Image.png
+
+Whenever you are changing levels in your UE4 VR project, the user could experience some framerate issues due to the massive amount of data that is being unloaded and loaded. To help avoid this issue and mask any frame rate issues that might arise when loading a new level you can display a Splash screen or movie to the user. In the following How - To we will go over how to setup and call a splash screen in your UE4 projects.
+
+[REGION:note]
+The following How - To is going to make use of the Standing VR Pawn with added Touch Controller and Guardian support. If you are not sure how to set this up the following two How - To's will help you out:
+
+* [Set Up Standing Cameras for the Oculus Rift](Platforms\Oculus\HowTo\StandingCamera)
+* [Setup Touch Controllers for the Oculus Rift](Platforms\Oculus\HowTo\TouchControllers)
+* [Set Up Guardian System](Platforms\Oculus\HowTo\GuardianSystem)
+[/REGION]
+
+## Setting Up Oculus Loading Splash Screens
+
+[region:note]
+ * Please note that the following is showing the simplest way to set VR splash screens up. While each project has unique needs, the overall setup for VR splash screens will be very similar to what is described.
+ * For this How - To you will need to download, unzip and import the two files contained in the following zip file, **[Oculus Splash Screen Source Files](OclusSplashScreenSourceContent.zip)**
+[/region]
+
+
+1. Open up your VRPawn and go to the **Event Graph**. Right-click inside the Event Graph and search for and add the following Blueprint nodes:
+
+    * Event Begin Play
+    * Set Tracking Origin
+    * Enable Auto Loading Splash Screen
+    * Add Loading Splash Screen
+    * Hide Splash Screen
+
+    [REGION:lightbox]
+    [![](HTSplashScreen_00.png)(w:446)](HTSplashScreen_00.png)
+    [/REGION]
+
+    [REGION:caption]
+    Click for full image.
+    [/REGION]
+
+1. Since we want the Splash screen to be called each time a level is loaded, we will want first to make sure that we enable the autoloading of Splash screen and then set what we want to use as the Splash screen. Finally, we will want to hide the Splash Screen so that we can call it when we need it later.  Now set up the nodes in your VRPawn Event Graph to match the following image: 
+
+    [REGION:lightbox]
+    [![](HTSplashScreen_01.png)(w:772)](HTSplashScreen_01.png)
+    [/REGION]
+
+    [REGION:caption]
+    Click for full image.
+    [/REGION]
+
+    [REGION:note]
+    Make sure that you check the  **Auto Show Enabled** property on the **Enable Auto Loading Splash Screen** so that the Splash screen will be automatically called each time a level is loaded.
+    [/REGION]
+
+1. In the **Add Loading Splash Screen** node, there is a **Texture** input which controls what Texture or Movie will be displayed when this is node is called. Set the Texture used to be the **T_UE4_Logo_00** or any other Texture of your choosing.
+    
+    [REGION:lightbox]
+    [![](HTSplashScreen_02.png)(w:353)](HTSplashScreen_02.png)
+    [/REGION]
+
+    [REGION:caption]
+    Click for full image.
+    [/REGION]
+
+    [REGION:note]
+    When selecting a Texture to be used as a Splash image, it is best to set the Texture compression setting to **UserInterface2D** and also to enable the **Never Stream** option to make sure you will see the highest quality version of your Splash screen.
+    [/REGION]
+
+1. Now connect the following three nodes to the VRPawn Event graph so that we can show and hide the Splash screen. When completed your setup should match the following image:
+
+    * Motion Controller (L)Trigger
+    * Show Loading Splash Screen
+    * Hide Loading Splash Screen
+
+    [REGION:lightbox]
+    [![](HTSplashScreen_03.png)(w:467)](HTSplashScreen_03.png)
+    [/REGION]
+
+    [REGION:caption]
+    Click for full image.
+    [/REGION]
+
+    [REGION:note]
+    While you can display Splash screen in this manner, you will want to add this functionality a level Blueprint so that when the level is loaded, the Splash screen will be displayed until the level is done with loading. 
+    [/REGION]
+
+1. After that has been completed and you VRPawn Blueprint looks like the following image, put on your HMD grab your Touch controllers and stand in the middle of your VR interaction area.
+
+    [REGION:lightbox]
+    [![](HTSplashScreen_04.png)(w:467)](HTSplashScreen_04.png)
+    [/REGION]
+
+    [REGION:caption]
+    Click for full image.
+    [/REGION]
+
+1. Now when you squeeze the left trigger on the Touch, your scene should go black, and the Texture of the UE4 logo or the image you selected will be displayed. Releasing the left trigger on the Touch will bring you back the level back into view like in the following video.
+
+    [OBJECT:EmbeddedVideo]
+        [PARAMLITERAL:width]
+        640
+        [/PARAMLITERAL]
+        [PARAMLITERAL:height]
+        360
+        [/PARAMLITERAL]
+        [PARAMLITERAL:videoid]
+        H_APZ1h8H5o
+        [/PARAMLITERAL]
+    [/OBJECT]
+
+## UE4 Project Downloads
+Below you will find a link to where you can download the UE4 project that was used to create this example.
+
+* [**Oculus Rift Splash Screens Example Project**](https://epicgames.box.com/s/1rirqbohl7inchgnssznuam7ylalrzk4)
+
+
+
+
diff --git a/Engine/Documentation/Source/Platforms/Oculus/HowTo/StandingCamera/HTOculusStandingCamera.INT.udn b/Engine/Documentation/Source/Platforms/Oculus/HowTo/StandingCamera/HTOculusStandingCamera.INT.udn
new file mode 100644
--- /dev/null
+++ b/Engine/Documentation/Source/Platforms/Oculus/HowTo/StandingCamera/HTOculusStandingCamera.INT.udn
@@ -0,0 +1,109 @@
+Availability: Public
+Crumbs: %ROOT%, Platforms
+Title:Set Up a Standing Camera for the Oculus Rift
+Description:Setting up the UE4 camera to work with a standing Oculus Rift experience.
+Type: how-to
+Skilllevel:Beginner
+Version:4.17
+Parent:Platforms/Oculus/HowTo
+Order: 2
+tags:Basics
+topic-image:HT_Rift_Camera_Setup_Topic_Image.png
+
+One of the first things you will need to think about when starting to develop any UE4 powered VR experience to be used with the Oculus Rift is deciding whether the experience will be a seated one or a standing one. In the following How - To we will take a look at how to setup your UE4 projects VR camera for a standing Oculus Rift experience.
+
+## Pawn Setup for a Standing Rift Experience 
+
+Below you will find instructions on how to setup up any Pawn to work with a standing Oculus Rift experience.
+
+1. First, open up or create a new Pawn Blueprint and then go to the **Component** section of the **Viewport** tab. From there add the following two components with the following names, making sure that the VRCamera is a child of VRCameraRoot:
+    
+    * **Scene** - VRCameraRoot
+    * **Camera** - VRCamera
+
+    [REGION:lightbox]
+    [![](HT_Rift_Camera_Setup_03.png)(w:224)](HT_Rift_Camera_Setup_03.png)
+    [/REGION]
+
+    [REGION:caption]
+    Click for full image.
+    [/REGION]
+
+    When 
+
+    [region:note]
+    No matter what VR HMD you are using, this is the Epic recommend way to set up a VR camera as it allows you to offset the position of the camera without having to move the actual camera.
+    [/region]
+
+With that completed, the initial setup for the VR Camera is now completed, in the next sections, we will go through setting the needed Blueprint logic for a standing Oculus Rift experience.
+
+
+## Camera Setup for Standing Rift Experience 
+
+Below you will find instructions on how to setup your camera for a standing Rift experience. 
+
+1. First, open up the Pawn Blueprint if not already open and in the **Event Graph** drag off the **Event Begin Play** node to show the Executable Actions list. In the list search for the **Set Tracking Origin** node and click on it to add it to the Event Graph.
+
+    [REGION:lightbox]
+    [![](HT_Rift_Camera_Setup_09.png)(w:555)](HT_Rift_Camera_Setup_09.png)
+    [/REGION]
+
+    [REGION:caption]
+    Click for full image.
+    [/REGION]
+
+1. The Set Tracking Origin node has two options, **Floor Level** and **Eye Level**. For a standing experience, you will need to keep the **Origin** of the Set Tracking Origin node to the default of **Floor Level**.
+    
+    [REGION:lightbox]
+    [![](HT_Rift_Camera_Setup_10.png)(w:555)](HT_Rift_Camera_Setup_10.png)
+    [/REGION]
+
+    [REGION:caption]
+    Click for full image.
+    [/REGION]
+
+1. Drag the Pawn Blueprint from the Content Browser to a level, making sure that it is placed at 0,0,0 in the level.
+    [REGION:lightbox]
+    [![](HT_Rift_Camera_Setup_06.png)(w:555)](HT_Rift_Camera_Setup_06.png)
+    [/REGION]
+
+    [REGION:caption]
+    Click for full image.
+    [/REGION]
+
+1. Select the Pawn Blueprint that was put in the level and in the **Details** panel under the **Pawn** setting, set the **Auto Possess Player** from **Disabled** to **Player 0**.
+
+    [REGION:lightbox]
+    [![](HT_Rift_Standing_Camera_04.png)(w:353)](HT_Rift_Standing_Camera_04.png)
+    [/REGION]
+
+    [REGION:caption]
+    Click for full image.
+    [/REGION]
+
+1. Finally, go to the **Main Toolbar** and change the **Play Mode** to **VR Preview** and then press the **Play** button. When you put your Oculus Rift HMD on and view the level while sitting down, you should see something similar to the following video. 
+
+    [OBJECT:EmbeddedVideo]
+        [PARAMLITERAL:videoid]
+        stJAjmg6b4g
+        [/PARAMLITERAL]
+        [PARAMLITERAL:width]
+        100
+        [/PARAMLITERAL]
+        [PARAMLITERAL:height]
+        56.25
+        [/PARAMLITERAL]
+        [PARAMLITERAL:units]
+        %
+        [/PARAMLITERAL]
+        [PARAMLITERAL:type]
+        simple
+        [/PARAMLITERAL]
+    [/OBJECT]
+
+## UE4 Project Downloads
+Below you will find a link to where you can download the UE4 project that was used to create this example. 
+
+* [**Oculus Rift Standing VR Camera Example Project**](https://epicgames.box.com/s/qgoh6uk0f8ra4rtvp0jdrxygfb2n9ykv)
+
+
diff --git a/Engine/Documentation/Source/Platforms/Oculus/Oculus.INT.udn b/Engine/Documentation/Source/Platforms/Oculus/Oculus.INT.udn
--- a/Engine/Documentation/Source/Platforms/Oculus/Oculus.INT.udn
+++ b/Engine/Documentation/Source/Platforms/Oculus/Oculus.INT.udn
@@ -1,66 +1,41 @@
-Title:Oculus Rift Development
-Crumbs:%ROOT%, Platforms
-Description:
-Availability:Public
-Navigation:platform
-Platform:Oculus
-Version: 4.11
-type:landing
-parent:Platforms/VR
-tags:VR
-tags:Oculus Rift
-tags:Platform
-Topic-image:OculusRiftDevelopment_topic.png
+Title: Oculus Rift Development
+Availability: Public
+Crumbs: %ROOT%, Platforms
+Description: Information on getting UE4 to work with the Oculus Rift HMD.
+Navigation: platform
+Platform: Oculus
+Version: 4.16
+type: landing
+parent: Platforms/Oculus
+order:2
+tags: Oculus Rift
+tags: Platform
+topic-image: Platforms/VR/DesktopVR/Oculus_Topic_Image.png
 
+## Prerequisites 
 
-## Starting Out
+[DIR(end:"1" output:"topic" type:"reference" tags:"Oculus Rift" parent:"Platforms/Oculus")]
 
-[DIR(output:"listbutton" type:"quick start" parent:"Platforms/Oculus")]
+## Quick Start
 
-## Oculus Rift Topics
+[DIR(end:"1" output:"topic" type:"quick start" tags:" VR" parent:"Platforms/Oculus")]
 
-[DIR(output:"listbutton" type:"overview" parent:"Platforms/Oculus")]
+## Best Practices
 
-## Oculus Rift Resources
+[DIR(end:"1" output:"topic" type:"overview" tags:" VR" parent:"Platforms/Oculus")]
 
-[REGION:buttonlist]
-* [Lessons from Integrating the Oculus Rift into Unreal Engine 4](https://de45xmedrsdbp.cloudfront.net/Resources/files/UE4-Integration-and-Demos_OC-100270768.pptx "Oculus Connect 2014 Presentation")
-[/REGION]
+## Oculus How To
 
-## Shipping with Oculus
-If you are planning on releasing a VR game that targets the Oculus Rift, you will need to make sure that you are using **UE4** version **4.11.1** or later. Earlier version of UE4 will not be compiled against the shipping versions of the Oculus Libraries. This means your project will not be able to be submitted to the Oculus store.
+[DIR(end:"1" output:"topic" type:"how-to" tags:"Oculus" tags:"VR" parent:"Platforms/Oculus")]
 
-## Enabling Unknown Sources 
-While you're developing your UE4 VR project, it's best to enable **Unknown Sources** so you can run your application with having to have it be registered or reviewed by Oculus. To enable the **Unknown Sources** option, you will need to do doing the following.
+## Oculus Command Reference
 
-1. Open the Oculus Store then click on the **Gear** icon that is in the upper right-hand corner to expose the **Settings** menu option, then click on the word **Settings**.
+[DIR(end:"1" output:"topic" type:"reference" tags:"VR" parent:"Platforms/Oculus")]
 
-	![](Rift_UKS_01.png)
+## Oculus Profiling & Performance
 
-1. From the left hand side, select the **General** section. Then in the **Unknown Sources** section, click on the box that has an **X** inside of it.
+[DIR(end:"1" output:"topic" type:"overview" tags:"Performance and Profiling" parent:"Platforms/Oculus")]
 
-	![](Rift_UKS_02.png)
 
-1. After the box has been clicked on, a message will be displayed. Click on the **Allow** button to continue.
 
-	![](Rift_UKS_03.png)
-
-1. When completed, you should now see a checkmark in the button instead of a **X**. To revert this change back to the default settings, click on the button again and the checkmark will change from a checkmark back into an X.
-	
-	![](Rift_UKS_04.png)
-
-
-## Oculus Talks
- **Oculus Connect: Integrating the Oculus Rift into Unreal Engine 4**
-[OBJECT:EmbeddedVideo]
-[PARAMLITERAL:width]
-640
-[/PARAMLITERAL]
-[PARAMLITERAL:height]
-360
-[/PARAMLITERAL]
-[PARAMLITERAL:videoid]
-0oM6Xe7fT-8
-[/PARAMLITERAL]
-[/OBJECT]
 
diff --git a/Engine/Documentation/Source/Platforms/Oculus/PreRecs/OculusPrerequisites.INT.udn b/Engine/Documentation/Source/Platforms/Oculus/PreRecs/OculusPrerequisites.INT.udn
new file mode 100644
--- /dev/null
+++ b/Engine/Documentation/Source/Platforms/Oculus/PreRecs/OculusPrerequisites.INT.udn
@@ -0,0 +1,46 @@
+Availability: Public
+Crumbs: %ROOT%
+Title: Oculus Prerequisites
+Description: Information about setting up your Oculus Rift to be used with UE4.
+Type: Reference
+Version: 4.16
+Parent: Platforms/Oculus
+Order: 2
+Tags: Oculus Rift
+Tags: Platform
+topic-image:PreRecProfiling_Topic_Image.png
+
+
+On the following page, you will find information on how to make sure your UE4 project will work with the Oculus Rift HMD. 
+
+## Enabling Unknown Sources 
+
+While you're developing your Unreal Engine 4 (UE4) Virtual Reality (VR) project, it's best to enable allow **Unknown Sources** so you can run your application with having to have it be registered or reviewed by Oculus. To enable the **Unknown Sources** option, you will need to do the following:
+
+1. Open the Oculus Store then click on the **Gear** icon that is in the upper right-hand corner to expose the **Settings** menu option, then click on the word **Settings**.
+
+    ![](Platforms\Oculus\Rift_UKS_01.png)
+
+1. From the left-hand side, select the **General** section. Then in the **Unknown Sources** section, click on the box that has a **X** inside of it.
+
+    ![](Platforms\Oculus\Rift_UKS_02.png)
+
+1. After the box has been clicked on, a message will be displayed. Click on the **Allow** button to continue.
+
+    ![](Platforms\Oculus\Rift_UKS_03.png)
+
+1. When completed, you should now see a checkmark in the button instead of a **X**.
+    
+    ![](Platforms\Oculus\Rift_UKS_04.png)
+    [region:note]
+    To revert this change back to the default settings, click on the button again and the checkmark will change from a checkmark back into an X.
+    [/region]
+
+## SteamVR Troubleshooting
+
+If you are having issues getting your Oculus Rift to run, the Oculus Support Center pages can help you fix most issues that you will run into.
+
+[**Oculus Support Center**](https://support.oculus.com/)
+
+
+
diff --git a/Engine/Documentation/Source/Platforms/Oculus/ProfilingPerf/OculusProfilingPerformance.INT.udn b/Engine/Documentation/Source/Platforms/Oculus/ProfilingPerf/OculusProfilingPerformance.INT.udn
new file mode 100644
--- /dev/null
+++ b/Engine/Documentation/Source/Platforms/Oculus/ProfilingPerf/OculusProfilingPerformance.INT.udn
@@ -0,0 +1,19 @@
+Title:Oculus Profiling & Performance
+Availability:Public
+Crumbs:%ROOT%, Platforms
+Description:Information profiling the performance of your UE4 Rift projects.
+Navigation:platform
+Platform:Oculus
+Version: 4.16
+type:landing
+parent:Platforms/Oculus
+order:2
+tags:Performance and Profiling
+type:overview
+topic-image:Oculus_ProPref_Topic_Image.png
+
+## Oculus PerfHud
+
+![](Oculus_Perf_HUD.png)
+
+[Include:Platforms\VR\Profiling\Tools\#OculusPerfHud]
diff --git a/Engine/Documentation/Source/Platforms/Oculus/QuickStart/1/OculusQuickStart_1.INT.udn b/Engine/Documentation/Source/Platforms/Oculus/QuickStart/1/OculusQuickStart_1.INT.udn
--- a/Engine/Documentation/Source/Platforms/Oculus/QuickStart/1/OculusQuickStart_1.INT.udn
+++ b/Engine/Documentation/Source/Platforms/Oculus/QuickStart/1/OculusQuickStart_1.INT.udn
@@ -1,30 +1,39 @@
-Title:1. Download Oculus Runtime
-Crumbs:%ROOT%, Platforms, Platforms/Oculus, Platforms/Oculus/QuickStart
-Description:Where to download  the Oculus Rift runtime
-Availability:Public
-Version: 4.11
-SkillLevel: Intermediate
+Availability: Public
+Crumbs: %ROOT%
+Title: 1. Download Oculus Runtime
+Description: Where to download  the Oculus Rift runtime
+Type:quick start
+SkillLevel:Beginner
+Version:4.17
+Parent:Platforms/Oculus
+Order: 1
+Checkpoint:RIFTQS
+Tags:Oculus Rift
+Tags:VR
 
 [VAR:Steps]
 [OBJECT:Navigation]
-	[PARAM:previous]
-	[/PARAM]
-	[PARAM:current]
-	[/PARAM]
-	[PARAM:home]
-		[Quick Start Home](Platforms/Oculus/QuickStart)
-	[/PARAM]
-	[PARAM:next]
-		[Next Step](Platforms/Oculus/QuickStart/2 "%Platforms/Oculus/QuickStart/2:title%")
-	[/PARAM]
+    [PARAM:previous]
+        [Previous Step](Platforms/Oculus/QuickStart/1 "%Platforms/Oculus/QuickStart/1:title%")
+    [/PARAM]
+    [PARAM:current]
+    [/PARAM]
+    [PARAM:home]
+        [Quick Start Home](Platforms/Oculus/QuickStart)
+    [/PARAM]
+    [PARAM:next]
+        [Next Step](Platforms/Oculus/QuickStart/3 "%Platforms/Oculus/QuickStart/3:title%")
+    [/PARAM]
 [/OBJECT]
 [/VAR]
 
-In this section we will go over where to download the latest Oculus runtimes that are compatible with Unreal Engine 4.0 as well as anything you should be aware of before you install the Oculus runtimes.
+In this section we will go over where to download the latest Oculus runtimes that are compatible with Unreal Engine 4.0.
 
 %Steps%
 
 1. Visit the [Oculus Setup page](https://www.oculus.com/en-us/setup/) and download the Oculus runtimes by clicking on the **Start  Download** button in the center of the page.
 
 	![](OR_Download_Runtimes_00.png)
+
 %Steps%
+
