Copying //UE4/Dev-Documentation to //UE4/Samples-Main (Source: //UE4/Dev-Documentation @ 3623665)

diff --git a/Engine/Documentation/Source/Platforms/Oculus/QuickStart/2/OculusQuickStart_2.INT.udn b/Engine/Documentation/Source/Platforms/Oculus/QuickStart/2/OculusQuickStart_2.INT.udn
--- a/Engine/Documentation/Source/Platforms/Oculus/QuickStart/2/OculusQuickStart_2.INT.udn
+++ b/Engine/Documentation/Source/Platforms/Oculus/QuickStart/2/OculusQuickStart_2.INT.udn
@@ -1,27 +1,31 @@
-Title:2. Install the Oculus Runtime
-Crumbs:%ROOT%, Platforms, Platforms/Oculus, Platforms/Oculus/QuickStart
-Description:Download the needed software for running games on the Oculus Rift
-Availability:Public
-Version: 4.9
-SkillLevel: Intermediate
+Availability: Public
+Crumbs: %ROOT%
+Title: 2. Install the Oculus Runtime
+Description: Download the needed software for running UE4 projects on the Oculus Rift.
+Type: quick start
+SkillLevel: Beginner
+Version: 4.17
+Parent: Platforms/Oculus
+Order: 2
+Checkpoint: RIFTQS
 
 [VAR:Steps]
 [OBJECT:Navigation]
-	[PARAM:previous]
-		[Previous Step](Platforms/Oculus/QuickStart/1 "%Platforms/Oculus/QuickStart/1:title%")
-	[/PARAM]
-	[PARAM:current]
-	[/PARAM]
-	[PARAM:home]
-		[Quick Start Home](Platforms/Oculus/QuickStart)
-	[/PARAM]
-	[PARAM:next]
-		[Next Step](Platforms/Oculus/QuickStart/3 "%Platforms/Oculus/QuickStart/3:title%")
-	[/PARAM]
+    [PARAM:previous]
+        [Previous Step](Platforms/Oculus/QuickStart/1 "%Platforms/Oculus/QuickStart/1:title%")
+    [/PARAM]
+    [PARAM:current]
+    [/PARAM]
+    [PARAM:home]
+        [Quick Start Home](Platforms/Oculus/QuickStart)
+    [/PARAM]
+    [PARAM:next]
+        [Next Step](Platforms/Oculus/QuickStart/3 "%Platforms/Oculus/QuickStart/3:title%")
+    [/PARAM]
 [/OBJECT]
 [/VAR]
 
-In this section we will go over how to install the Oculus runtimes as well as how to resolve any issues that might come up during the install process.
+In this section, we will go over how to install the Oculus runtimes as well as how to resolve any issues that might come up during the install process.
 
 [region:note]
 You will need to be connected to the internet for 30 to 60 minutes while the Oculus runtimes are installing.
@@ -31,33 +35,14 @@ You will need to be connected to the internet for 30 to 60 minutes while the Ocu
 
 1. Once the Oculus Rift runtimes have finished downloading, locate the **OculusSetup.exe** and double - click on it to begin the install.
 
-	![](RIft_EXE_Download_00.png)
+    ![](RIft_EXE_Download_00.png)
 
-1. Once the Oculus Rift installer has opened, press the **Let's Go** button to begin the install.
+1. Press the **Let's Go** button to begin the install and then follow the provided directions.
 
-	![](Rift_Install_Start_00.png)
+    ![](Rift_Install_Start_00.png)
 
-1. If you agree with the privacy policy and terms of service press the **Agree** button to move on to the next install step.
+1. During the install process, you will be asked to install device software from Oculus VR, LLC. When this is displayed, press the **Install** button to continue the installation. 
 
-	![](Rift_TC_00.png)
-
-1. Read the Health & Safety Warning message and if you agree to it, click on the **I Understand** button to be taken to the next install step.
-
-	![](Rift_Health_Safety_Warning_00.png)
-
-1. Press the **Install Now** button to install the Oculus runtimes on your C: drive. 
-
-	![](Rift_Install_Now.png)
-	[region:note]
-	If you do not have room on your C drive or you want to change the location of the install, press the **Change Drive** button to change which hard drive the Oculus runtimes will be installed to.
-	[/region]
-
-1. The Oculus Rift runtimes will now be downloaded and installed to your PC.
-
-	![](Rift_Installing_Progress_00.png)
-
-1. During the install process you will be asked to install device software from Oculus VR, LLC. When this is displayed, press the **Install** button to continue the installation. 
-
-	![](Rift_Device_Software_Install_00.png)
+    ![](Rift_Device_Software_Install_00.png)
 
 %Steps%
diff --git a/Engine/Documentation/Source/Platforms/Oculus/QuickStart/3/OculusQuickStart_3.INT.udn b/Engine/Documentation/Source/Platforms/Oculus/QuickStart/3/OculusQuickStart_3.INT.udn
--- a/Engine/Documentation/Source/Platforms/Oculus/QuickStart/3/OculusQuickStart_3.INT.udn
+++ b/Engine/Documentation/Source/Platforms/Oculus/QuickStart/3/OculusQuickStart_3.INT.udn
@@ -1,43 +1,42 @@
-Title:3. Creat / Sign in to an Oculus Account
-Crumbs:%ROOT%, Platforms, Platforms/Oculus, Platforms/Oculus/QuickStart
-Description:Install the needed software for running games on the Oculus Rift
-Availability:Public
-Version: 4.9
-SkillLevel: Intermediate
-
-In the section we will go over what to do once the Oculus Rift software has been installed.
+Availability: Public
+Crumbs: %ROOT%
+Title: 3. Creat / Sign into an Oculus Account
+Description: Install the needed software for running UE4 projects on the Oculus Rift.
+Type: quick start
+SkillLevel: Beginner
+Version: 4.17
+Parent: Platforms/Oculus
+Order: 3
+Checkpoint:RIFTQS
+
+In the section, we will go over what to do once the Oculus Rift software has finished installing.
 
 [VAR:Steps]
 [OBJECT:Navigation]
-	[PARAM:previous]
-		[Previous Step](Platforms/Oculus/QuickStart/2 "%Platforms/Oculus/QuickStart/2:title%")
-	[/PARAM]
-	[PARAM:current]
-	[/PARAM]
-	[PARAM:home]
-		[Quick Start Home](Platforms/Oculus/QuickStart)
-	[/PARAM]
-	[PARAM:next]
-		[Next Step](Platforms/Oculus/QuickStart/4 "%Platforms/Oculus/QuickStart/4:title%")
-	[/PARAM]
+    [PARAM:previous]
+        [Previous Step](Platforms/Oculus/QuickStart/2 "%Platforms/Oculus/QuickStart/2:title%")
+    [/PARAM]
+    [PARAM:current]
+    [/PARAM]
+    [PARAM:home]
+        [Quick Start Home](Platforms/Oculus/QuickStart)
+    [/PARAM]
+    [PARAM:next]
+        [Next Step](Platforms/Oculus/QuickStart/4 "%Platforms/Oculus/QuickStart/4:title%")
+    [/PARAM]
 [/OBJECT]
 [/VAR]
 
 %Steps%
 
 
-1. Once the Oculus Rift software has been installed, press the **Next** button to begin the next part of the install process.
-
-	![](Rift_Finshed_Install_00.png)
-
-1. You will now be prompted to either **Sign in** or **Create Account**. If you already have an account click on the Sign in option. If you do not have an Oculus Account, click on the Create Account and follow the provided instructions.
-
-	![](Rift_Create_Account_LogIn_00.png)
-
-
+1. Once the Oculus Rift software has installed, press the **Next** button to begin the hardware setup.
 
+    ![](Rift_Finshed_Install_00.png)
 
+1. You will then be prompted to either **Sign in** or **Create Account** an Oculus account before you can continue.
 
+    ![](Rift_Create_Account_LogIn_00.png)
 
 
 %Steps%
diff --git a/Engine/Documentation/Source/Platforms/Oculus/QuickStart/4/OculusQuickStart_4.INT.udn b/Engine/Documentation/Source/Platforms/Oculus/QuickStart/4/OculusQuickStart_4.INT.udn
--- a/Engine/Documentation/Source/Platforms/Oculus/QuickStart/4/OculusQuickStart_4.INT.udn
+++ b/Engine/Documentation/Source/Platforms/Oculus/QuickStart/4/OculusQuickStart_4.INT.udn
@@ -1,49 +1,54 @@
-Title:4. Set Up Your Oculus Rift Hardware
-Crumbs:%ROOT%, Platforms, Platforms/Oculus, Platforms/Oculus/QuickStart
-Description:Connect and update your Oculus Rift headset.
-Availability:Public
-Version: 4.9
-SkillLevel: Intermediate
+Availability: Public
+Crumbs: %ROOT%
+Title: 4. Set Up Your Oculus Rift Hardware
+Description: Connect and update your Oculus Rift HMD.
+Type: quick start
+SkillLevel: Beginner
+Version: 4.17
+Parent: Platforms/Oculus
+Order: 4
+Checkpoint: RIFTQS
 
 [VAR:Steps]
 [OBJECT:Navigation]
-	[PARAM:previous]
-		[Previous Step](Platforms/Oculus/QuickStart/3 "%Platforms/Oculus/QuickStart/3:title%")
-	[/PARAM]
-	[PARAM:current]
-	[/PARAM]
-	[PARAM:home]
-		[Quick Start Home](Platforms/Oculus/QuickStart)
-	[/PARAM]
-	[PARAM:next]
-		[Next Step](Platforms/Oculus/QuickStart/5 "%Platforms/Oculus/QuickStart/5:title%")
-	[/PARAM]
+    [PARAM:previous]
+        [Previous Step](Platforms/Oculus/QuickStart/3 "%Platforms/Oculus/QuickStart/3:title%")
+    [/PARAM]
+    [PARAM:current]
+    [/PARAM]
+    [PARAM:home]
+        [Quick Start Home](Platforms/Oculus/QuickStart)
+    [/PARAM]
+    [PARAM:next]
+        [Next Step](Platforms/Oculus/QuickStart/5 "%Platforms/Oculus/QuickStart/5:title%")
+    [/PARAM]
 [/OBJECT]
 [/VAR]
 
-In this section we will look at how you complete the setup for your Oculus HMD and make sure that it will work with the latest Oculus drivers.
+In this section, we will look at how you complete the setup for your Oculus Rift HMD and peripherals.
 
 [region:note]
-When you first plug your Oculus Rift into your PC, the Rift software may ask you to update the firmware for your Rift. It is highly recommended that you allow this to update to happen to ensure that your Rift will work properly with UE4.
+When you first plug your Oculus Rift into your PC, the Rift software may ask you to update the firmware for your Rift. DO NOT skip this or your Rift might not work with UE4!
 [/region]
 
 %Steps%
 
-1. After the Oculus software has finished installing its drivers and software, you will be shown the Rift setup screen. 
-	Press the **Next** button to begin the Rift setup process and follow the directions on the screen.
+1. After the Rift drivers have finished installing the following screen will be displayed. 
+    Press the **Next** button to begin the Rift setup process and follow the directions on the screen.
 
-	![](Rift_Setup_Screen_00.png)
+    ![](Rift_Setup_Screen_00.png)
 
 [region:note]
-If you need help during any part of the Rift setup process click on the **Get Help** link that is in the lower lefthand corner to be brought to that help section.
+If you need help during any part of the Rift setup process, click on the **Get Help** link that is in the lower lefthand corner to get help.
 [/region]
 
-For quick reference, the help links that you can find in the lower right hand corner of each setup section have been listed below. The links have been order in the same way that they come up in the setup process.
+For quick reference, you can find the help links listed in the lower right-hand corner of each setup below. The links have are ordered in the same way that they appear in the setup process.
 
 * [**Rift Initial Setup Help**](https://support.oculus.com/674612489309093) 
 * [**Rift Care and Maintenance**](https://support.oculus.com/1659691017629165) 
 * [**Using the Headset**](https://support.oculus.com/674612489309093) 
 * [**Oculus Remote**](https://support.oculus.com/835449819935261) 
 * [**Setup Your Xbox Controller (wireless mode)**](https://support.oculus.com/1560296144292386) 
+* [**Touch Setup**](https://support.oculus.com/857827607684748?locale=en_US&ref=oculus-pc-app)
 
 %Steps%
diff --git a/Engine/Documentation/Source/Platforms/Oculus/QuickStart/5/OculusQuickStart_5.INT.udn b/Engine/Documentation/Source/Platforms/Oculus/QuickStart/5/OculusQuickStart_5.INT.udn
--- a/Engine/Documentation/Source/Platforms/Oculus/QuickStart/5/OculusQuickStart_5.INT.udn
+++ b/Engine/Documentation/Source/Platforms/Oculus/QuickStart/5/OculusQuickStart_5.INT.udn
@@ -1,11 +1,15 @@
-Title:5. Testing out the Rift and UE4
-Crumbs:%ROOT%, Platforms, Platforms/Oculus, Platforms/Oculus/QuickStart
-Description:Load and play content designed for the Oculus Rift using Unreal Engine 4.
-Availability:Public
-Version: 4.11
-SkillLevel: Intermediate
-
-In this section we will cover how you can test to see if your Rift is working inside Unreal Engine 4.
+Availability: Public
+Crumbs: %ROOT%
+Title: 5. Testing out the Rift and UE4
+Description: Load and play content designed for the Oculus Rift using Unreal Engine 4.
+Type: quick start
+SkillLevel: Beginner
+Version: 4.17
+Parent: Platforms/Oculus
+Order: 5
+Checkpoint: RIFTQS
+
+In this section we will cover how you can test to see if your Rift is working with Unreal Engine 4.
 
 [VAR:Steps]
 [OBJECT:Navigation]
@@ -25,19 +29,32 @@ In this section we will cover how you can test to see if your Rift is working in
 
 %Steps%
 
-1. Now that the Oculus drivers have been installed and the Oculus Rift has been set up, open up the Unreal Project Browser and in the **New Project** tab create a new project that has the following settings.
+1. Open up the Unreal Project Browser and in the **New Project** tab create a new project that has the following settings:
+
+	[REGION:lightbox]
+	[![](Rift_New_Project_00.png)(w:555)](Rift_New_Project_00.png)
+	[/REGION]
 
-	![](Rift_New_Project_00.png)
+	[REGION:caption]
+	Click for full image.
+	[/REGION]
 
-1. Once UE4 has launched, go to the Main Menu and change the Play option from the default of **Selected Viewport** to **VR Preview** and put your Rift on.
+1. Once UE4 has launched, go to the **Main Menu** and change the **Play** option from the default of **Selected Viewport** to **VR Preview**, pressing VR PReview to launch the level.
 
 	![](Rift_VR_Preview_00.png)
 
-1. Put the Rift HMD on your head and you should see an image that looks like the following. 
+1. Now, put the Rift on and you should see the following. 
+
+	[REGION:lightbox]
+	[![](Rift_Preview_Window_00.png)(w:648)](Rift_Preview_Window_00.png)
+	[/REGION]
+
+	[REGION:caption]
+	Click for full image.
+	[/REGION]
 
-	![](Rift_Preview_Window_00.png)
 	[region:note]
-	When using the VR preview option and the Rift, the VR preview window will display black and look like it is not updating anytime time you launch or re-launch the VR preview window. This is because when the Rift is not placed on your head it disables the output to help extend the life of the display. Placing the Rift on your head will enable the output to the VR preview window allowing you the VR preview window to show what you are seeing in the Rift.
+	When using the VR preview option and the Rift, the VR preview window will display black and look like it is not updating anytime time you launch or re-launch the VR preview window. This is because when the Rift is not placed on your head it disables the display output to help extend the life of the display. 
 	[/region]
 
 %Steps%
diff --git a/Engine/Documentation/Source/Platforms/Oculus/QuickStart/6/OculusQuickStart_6.INT.udn b/Engine/Documentation/Source/Platforms/Oculus/QuickStart/6/OculusQuickStart_6.INT.udn
--- a/Engine/Documentation/Source/Platforms/Oculus/QuickStart/6/OculusQuickStart_6.INT.udn
+++ b/Engine/Documentation/Source/Platforms/Oculus/QuickStart/6/OculusQuickStart_6.INT.udn
@@ -1,9 +1,13 @@
-Title:6. Additional Resources
-Crumbs:%ROOT%, Platforms, Platforms/Oculus, Platforms/Oculus/QuickStart
-Description:Read more about Oculus development with Unreal Engine.
-Availability:Public
-Version: 4.9
-SkillLevel: Intermediate
+Availability: Public
+Crumbs: %ROOT%
+Title: 6. Additional Resources
+Description: Read more about Oculus development with Unreal Engine.
+Type: quick start
+SkillLevel: Beginner
+Version: 4.17
+Parent: Platforms/Oculus
+Order: 6
+Checkpoint: RIFTQS
 
 Here are some additional resources that provide useful information for developing VR projects inside of Unreal Engine 4.
 
@@ -52,11 +56,11 @@ Here are some additional resources that provide useful information for developin
 	[/PARAM]
 	[PARAM:links]
 		[REGION:imagetable]
-		| ![](store_CouchKnights_screenshot_3.png)(w:420) | ![](VehicleTemplate.png)(w:420) |
+		| ![](store_CouchKnights_screenshot_3.png)(w:420) | ![](store_Showdown_screenshot_5.png)(w:420) |
 		| ------------ | ---------------- |
-		| CouchKnights | Vehicle Template[REGION:caption]Press **Tab** to switch to being inside the car.[/REGION] |
-		| ![](store_RealisticRendering_screenshot_3.png)(w:420) | ![](store_MobileTemple_screenshot_4.png)(w:420) |
-		| Architecture[REGION:caption] Disable the Matinee fly-through to walk around.[/REGION] | Simple Temple Environment |
+		| CouchKnights | Showdown |
+		| ![](store_VREditor_screenshot_5.png)(w:420) | ![](store_RoboRecall_screenshot_5.png)(w:420) |
+		| VR Editor |Robo Recall |
 		[/REGION]
 	[/PARAM]
 [/OBJECT]
diff --git a/Engine/Documentation/Source/Platforms/Oculus/QuickStart/OculusQuickStart.INT.udn b/Engine/Documentation/Source/Platforms/Oculus/QuickStart/OculusQuickStart.INT.udn
--- a/Engine/Documentation/Source/Platforms/Oculus/QuickStart/OculusQuickStart.INT.udn
+++ b/Engine/Documentation/Source/Platforms/Oculus/QuickStart/OculusQuickStart.INT.udn
@@ -1,67 +1,30 @@
-Title:Oculus Rift Quick Start
-Crumbs:%ROOT%, Platforms, Platforms/Oculus
-Description:Instructions on setting up and launching games on the Oculus Rift.
-Availability:Public
-Version: 4.11
-SkillLevel: Intermediate
-parent:Platforms/Oculus
+Availability: Public
+Title: Oculus Rift Quick Start
+Crumbs: %ROOT%
+Description: Instructions on setting up and launching UE4 projects with the Oculus Rift.
+Navigation:platform
+Platform:Oculus
+SkillLevel: Beginner
+Version: 4.17
 type:quick start
-type:landing
-tags:VR
-tags:Oculus Rift
+Type:landing
+Tags: VR
+Tags: Oculus Rift
+Order: 2
+Checkpoint:RIFTQS
+topic-image:Platforms\Oculus\QuickStart\oculus_qs_topic.png
 
-[VAR:Topic]
-[OBJECT:Topic]
-	[PARAM:image]
-		![](oculus_qs_topic.png)
-	[/PARAM]
-	[PARAM:title]
-		%Platforms/Oculus/QuickStart:title%
-	[/PARAM]
-	[PARAM:path]
-		[RELATIVE:Platforms/Oculus/QuickStart]
-	[/PARAM]
-	[PARAM:description]
-	%Platforms/Oculus/QuickStart:description%
-	[/PARAM]
-[/OBJECT]
-[/VAR]
+![](Oculus_Topic_Image.png)
 
-[VAR:TopicCompact]
-[OBJECT:TopicCompact]
-	[PARAM:image]
-		![](oculus_qs_topic.png)
-	[/PARAM]
-	[PARAM:title]
-		%Platforms/Oculus/QuickStart:title%
-	[/PARAM]
-	[PARAM:path]
-		[RELATIVE:Platforms/Oculus/QuickStart]
-	[/PARAM]
-	[PARAM:description]
-	%Platforms/Oculus/QuickStart:description%
-	[/PARAM]
-[/OBJECT]
-[/VAR]
+## Goals
+The Oculus Quick Start will walk you through setting up your computer and your Oculus Rift to be used with Unreal Engine 4.0. 
 
+## Objectives
+* Create a new UE4 project that is targeted towards Steam VR development.
+* Setting the required project settings so your project can be used with Steam VR.
 
-The Oculus Quick Start will walk you through setting up your computer and your Oculus Rift to be used with Unreal Engine 4.0. 
+[DIR(output:"checkpoint" parent:"Platforms/Oculus/QuickStart")]
 
-[OBJECT:TopicButtonList]
-	[PARAM:icon]
-		![](%ROOT%/build_icon.png)(convert:false)
-	[/PARAM]
-	[PARAM:title]
-		Steps
-	[/PARAM]
-	[PARAM:description]
-	[/PARAM]
-	[PARAM:links]
-	* [](Platforms/Oculus/QuickStart/1 "%Platforms/Oculus/QuickStart/1:description%")
-	* [](Platforms/Oculus/QuickStart/2 "%Platforms/Oculus/QuickStart/2:description%")
-	* [](Platforms/Oculus/QuickStart/3 "%Platforms/Oculus/QuickStart/3:description%")
-	* [](Platforms/Oculus/QuickStart/4 "%Platforms/Oculus/QuickStart/4:description%")
-	* [](Platforms/Oculus/QuickStart/5 "%Platforms/Oculus/QuickStart/5:description%")
-	* [](Platforms/Oculus/QuickStart/6 "%Platforms/Oculus/QuickStart/6:description%")
-	[/PARAM]
-[/OBJECT]
+[REGION:call_to_action]
+ [Click to Start](Platforms/Oculus/QuickStart/1 "%Platforms/Oculus/QuickStart/1:title%")
+[/REGION]
diff --git a/Engine/Documentation/Source/Platforms/Oculus/Reference/OculusCommandReference.INT.udn b/Engine/Documentation/Source/Platforms/Oculus/Reference/OculusCommandReference.INT.udn
new file mode 100644
--- /dev/null
+++ b/Engine/Documentation/Source/Platforms/Oculus/Reference/OculusCommandReference.INT.udn
@@ -0,0 +1,32 @@
+Availability: Public
+Crumbs: %ROOT%
+Title: Rift Command Reference
+Description: Information on the various commands you can use with the Rift.
+Type: Reference
+Version: 4.17
+Parent: Platforms/Oculus
+Order: 2
+Tags: VR
+Tags: Platform
+topic-image:oculus_command_reff.png
+
+Below you will find various INI and Console commands that can be used to adjust the way UE4 and the Oculus Rift interact with one another. 
+ 
+
+## Oculus Rift VR Mirror Mode
+
+The VR Mirror is the name for the viewport window that allows non VR users to see what the VR user is seeing. When using the Oculus Rift there are a few ways you can change the way the VR Mirror is displayed. Below you will find the commands to change the way the VR Mirror mode and what the effect that will have on the how the VR Mirror looks.
+
+[REGION:note]
+To input and of the following command press the backtick **`** key to open up the UE4 console and then input **VR.MirrorMode** followed by any of the following numbers.
+[/REGION]
+
+
+|Command|Description|Image|
+|-----------|-------------------|----------------|
+|VR.MirrorMode 0|Disables the VR Mirror|N/A|
+|VR.MirrorMode 1| Shows a single eye without any distortion.|[REGION:lightbox][![](Rift_VRMode_1.png)(w:160)](Rift_VRMode_1.png)[/REGION]|
+|VR.MirrorMode 2| Shows both eyes in stereo without lense distortion.|[REGION:lightbox][![](Rift_VRMode_2.png)(w:160)](Rift_VRMode_1.png)[/REGION]|
+|VR.MirrorMode 3| Shows both eyes in stereo with lense distortion and Guardian Boundaries.|[REGION:lightbox][![](Rift_VRMode_3.png)(w:160)](Rift_VRMode_1.png)[/REGION]|
+|VR.MirrorMode 4| Shows a single eye in stereo.|[REGION:lightbox][![](Rift_VRMode_4.png)(w:160)](Rift_VRMode_1.png)[/REGION]|
+|VR.MirrorMode 5| Shows a single eye in Letterbox .|[REGION:lightbox][![](Rift_VRMode_5.png)(w:160)](Rift_VRMode_1.png)[/REGION]|
\ No newline at end of file
diff --git a/Engine/Documentation/Source/Platforms/SteamVR/BestPractices/SteamVRBestPractice.INT.udn b/Engine/Documentation/Source/Platforms/SteamVR/BestPractices/SteamVRBestPractice.INT.udn
--- a/Engine/Documentation/Source/Platforms/SteamVR/BestPractices/SteamVRBestPractice.INT.udn
+++ b/Engine/Documentation/Source/Platforms/SteamVR/BestPractices/SteamVRBestPractice.INT.udn
@@ -3,6 +3,7 @@ Crumbs:%ROOT%, Platforms
 Description:An introduction to getting started with SteamVR
 Availability:Public
 Navigation:platform
+Version: 4.17
 Platform:SteamVR
 type:overview
 parent:Platforms/SteamVR
@@ -11,62 +12,76 @@ tags:SteamVR
 
 ![](Platforms\SteamVR\T_Vive_HMD.png)
 
+The SteamVR SDK is unlike other Unreal Engine 4 (UE4) Virtual Reality SDK as it is not bound to be used with a particular Head Mounted Display (HMD). Because of this UE4 projects that are developed for SteamVR can be used with any HMD that supports SteamVR. The following guide will help to highlight some information that you need to know about when developing content for SteamVR and UE4.
 
-Virtual Reality or VR for short is quickly becoming the next big thing in digital entertainment. 
-Because of this industry wide push to VR, Unreal Engine 4 (UE4) supports a wide variety of VR headsets and controllers without developers having to do anything else besides just opening the Unreal Editor and making sure their VR device is connected to their development PC. 
-One of the newer VR integrations UE4 now fully supports is SteamVR. 
-Unlike other VR integrations the SteamVR integration not only includes code for the new Vive VR headset, 
-it also includes code that helps to power the new Steam Controllers as well as the Steam Lighthouse Base stations. 
-In the following sections we will cover all you need to know about how to get UE4 working with the SteamVR integration.
+## SteamVR Beta
 
-[region:note]
-Currently the SteamVR headset and controllers are not publicly available for purchase. 
-The only way to receive a development kit is to sign up as a SteamVR developer on the [Valve](http://store.steampowered.com/universe/vr) website. 
-Also keep in mind that because SteamVR is still currently prototype hardware, images or information contained within this document could abruptly change or be slightly different than what you might have received or are currently using.
-As things like this happen this document will be updated so make sure to check back from time to time to familiarize yourself with anything new or any changes that might have happened.
-[/region]
+To make sure that you have the most up to date version of SteamVR installed, make sure that you opted into the SteamVR Beta by right-clicking on the SteamVR Tool and going to **Properties** >  **Betas** and then select the **beta - SteamVR  Beta Update** option.
 
+![](SteamVRBP_01.png)
 
-## Prerequisites
+## SteamVR HMD Target Frame Rate
 
-Before we can start developing content to use with SteamVR there are a few things that need to be setup first. 
-In the following section we will go over what needs to be setup before you can begin to develop content for SteamVR.
+Below you will find the frame rate that your UE4 project must meet when using SteamVR. 
 
+| HMD Device  | Target Frame Rate |
+| --- | --- |
+| HTC Vive      | 90 FPS|
+| Oculus Rift | 90 FPS|
 
-* Make sure to turn off [Windows Aero](http://windows.microsoft.com/en-us/windows-vista/what-is-windows-aero) to ensure that all of your development PC's power is available to use with VR.
-* From the Steam client make sure to download and install the **SteamVR** tools. You can find and install the tools by doing the following.
- 1. From the Steam client go to **Library** > **Tools** then in the list of tools look for **SteamVR**.
- 1. **Double - Click** on the **SteamVR** entry to start the tools install process and then follow the directions provide by Steam.
 
-* Make sure to download and install Unreal Engine version 4.8 or later from the Unreal Engine launcher.
-* Have all the required cords and power supplies for both the Vive Head Mounted Display(HMD) and the Lighthouse Base stations unpacked and plugged in.
-* Make sure you have a 15' x 15' (4.5 M x 4.5 M) area cleared of all objects that you could trip over while you have the Vive HMD on. 
-* While not required, having a wireless mouse and keyboard will make entering commands easier as you will not have to take off the HMD and then move over to where the development PC is located then input the commands.
+## SteamVR Sanity Check
 
+Unreal Engine will automatically use SteamVR if a supported HMD is plugged in and the SteamVR plug-in is enabled. If for some reason SteamVR is not working first check that the SteamVR plug-in is enabled. You can find the SteamVR plugin under the **Virtual Reality** section in the [Plug-in](Programming/Plugins) menu.  
 
-## What is SteamVR
+[REGION:lightbox]
+[![](SteamVRBP_00.png)(w:546)](SteamVRBP_00.png)
+[/REGION]
 
-SteamVR or Steam Virtual Reality is fully featured, 360 degree, room scale VR experience that allows users to experience, interact, and even move around in virtual worlds. 
-SteamVR can be utilized in either an active, moving around the room, or passive, sitting in a chair, setting allowing developers to choose which setting will give the best experience to users. 
-SteamVR also makes use of two wireless single handed controllers whose real world  movement, position, and rotation can be precisely tracked and then translated in to VR allowing users the ability to precisely interact with object in VR like never before. 
-If that was not enough, SteamVR makes use of the Lighthouse tracking system which utilizes room scanning lasers to precisely track the real world position, rotation, and movement of the user and then translates that data so that what the user is doing in the real world will match up to the VR world they are experiencing. 
+[REGION:caption]
+Click for full image.
+[/REGION]
 
-## How is SteamVR different than other VR devices
+## Checking Your Work With SteamVR 
 
-While there are a quite a few similarities between SteamVR and other VR implementations, there are also quite a few differences.
-One of the main differences between SteamVR and other VR systems is that instead of using a camera to track the user's location in the real world and then translate that into the virtual one, Steam VR makes use of laser driven tracking stations to track the users real world movement. 
-These stations, called Lighthouse trackers or Base stations, can more accurately track the position and orientation of the headset and controllers and are optimized for accuracy with low latency. 
-In fact the Lighthouse trackers are so precise you can simultaneously track multiple input devices with sub millimeter precision in just milliseconds.
-SteamVR is also built around the SteamVR Controllers which are specifically designed to make it easier to interact with objects in VR. 
-Finally, and probably the biggest difference, is that with the SteamVR setup you can physically move around in the room you are using SteamVR in and have that movement translated into the VR world allowing the user to explore or view parts of the world that other VR systems would not be able to achieve due to you not being able to physically move around. 
+SteamVR will not work with any of the UE4 Editor viewports or the default Play In Editor(PIE) session. To view your project with SteamVR, you will need to use the **VR Preview** option. To access the VR Preview option, you will need to do the following inside the UE4 Editor. 
 
+1. On the main Toolbar in the **Play** section, click on the small white triangle that is next to the Play button. 
 
+    ![](Platforms\Oculus\BestPractices\Rift_Preview_00.png)
 
+1. From the dropdown menu, select the **VR Preview** option and then put your Rift on to view your project in VR.
 
+    ![](Platforms\Oculus\BestPractices\Rift_Preview_01.png)
 
+    [region:note]
+    Once you have changed the play mode to the VR Preview option your project will always launch in VR, even when using shortcuts like **ALT + P**.  
+    [/region]
 
+## SteamVR Mirror Mode
 
+The SteamVR Headset Mirror allows you to see what the user sees in the HMD. Enabling this mode is particularly useful if you need to record what you are seeing in your UE4 project and the SteamVR compositor. To enable the mirror, you will need to do the following:
 
+1. First, locate the SteamVR Tools and then right - click on them to display the menu and select the **Display Mirror** option.
 
+    ![](SteamVRBP_02.png)
+
+1. The mirror will then be displayed in a new window called **HeadSetMirror** like in the image below.
+
+    [REGION:lightbox]
+    [![](SteamVRBP_03.png)(w:540)](SteamVRBP_03.png)
+    [/REGION]
+
+    [REGION:caption]
+    Click for full image.
+    [/REGION]
+
+
+## SteamVR Developer Links
+
+The following is a collection of links that will provide useful information on items such as hardware or software issue with SteamVR.
+
+* [**SteamVR Developer Support**](https://support.steampowered.com/kb_article.php?ref=1131-WSFG-3320)
+* [**HTC Vive Developer Support**](https://developer.viveport.com/us/develop_portal/)
 
 
diff --git a/Engine/Documentation/Source/Platforms/SteamVR/HowTo/Chaperone/HTSteamVRChaperone.INT.udn b/Engine/Documentation/Source/Platforms/SteamVR/HowTo/Chaperone/HTSteamVRChaperone.INT.udn
new file mode 100644
--- /dev/null
+++ b/Engine/Documentation/Source/Platforms/SteamVR/HowTo/Chaperone/HTSteamVRChaperone.INT.udn
@@ -0,0 +1,64 @@
+Availability: Public
+Crumbs: %ROOT%, Platforms
+Title:Set Up the SteamVR Chaperone 
+Description:Setting up the SteamVR Chaperone system in UE4.
+Type: how-to
+Skilllevel:Beginner
+Version:4.17
+Parent:Platforms/SteamVR/HowTo
+Order: 2
+tags:Basics
+topic-image:Platforms\SteamVR\SteamVRDevelopment_topic.png
+
+The SteamVR Chaperone component is used to display to the user where the soft and hard bounds of the VR interaction area are. In the following How - To we will take a look at how to add the SteamVR Chaperone to any UE4 VR Pawn.
+
+[REGION:note]
+The following How - To is going to make use of the Standing VRPawn with added Motion Controller support. If you are not sure how to set this up the following two How - To's will help you out:
+* [Set Up Standing Cameras for SteamVR ](Platforms\SteamVR\HowTo\StandingCamera)
+* [Set Up Motion Controllers for SteamVR](Platforms\VR\MotionController)
+[/REGION]
+
+[REGION:warning]
+It is **NOT** possible nor advisable to disable the Chaperone system from inside UE4. You can, however, adjust how UE4 will react to the user getting close to the boundaries.  
+[/REGION]
+
+## SteamVR Chaperone Setup
+
+To add the SteamVR Chaperone system to your UE4 player pawn, you will need to do the following:
+
+1. Open up your projects player Pawn Blueprint and make sure the **Components** tab is visible.  
+
+1. Click on the **Add Component** button and from the displayed list search for the **Steam VRChaperone** component and when located, click on it to add it to the Component list.
+
+	![](SVRChaperone_00.png)
+
+1. When completed, your player pawn should look like the following image.
+
+	![](SVRChaperone_01.png)
+
+1. Now, launh your project in VR and put on your Vive HMD and when you aproach the bounds of your VR interaction area you should see something simialr to the following video.
+
+	[OBJECT:EmbeddedVideo]
+		[PARAMLITERAL:videoid]
+		ibnskPOz1Hw
+		[/PARAMLITERAL]
+		[PARAMLITERAL:width]
+		100
+		[/PARAMLITERAL]
+		[PARAMLITERAL:height]
+		56.25
+		[/PARAMLITERAL]
+		[PARAMLITERAL:units]
+		%
+		[/PARAMLITERAL]
+		[PARAMLITERAL:type]
+		simple
+		[/PARAMLITERAL]
+	[/OBJECT]
+
+## UE4 Project Downloads
+
+Below you will find a link to where you can download the UE4 project that was used to create this example. 
+
+* [**SteamVR Chaperone Set Up Project**](https://epicgames.box.com/s/f9aw7n5wpjc0nnzjtiom78hticmdl3w1)
+
diff --git a/Engine/Documentation/Source/Platforms/SteamVR/HowTo/ChaperoneActivation/SteamVRChaperoneActivation.INT.udn b/Engine/Documentation/Source/Platforms/SteamVR/HowTo/ChaperoneActivation/SteamVRChaperoneActivation.INT.udn
new file mode 100644
--- /dev/null
+++ b/Engine/Documentation/Source/Platforms/SteamVR/HowTo/ChaperoneActivation/SteamVRChaperoneActivation.INT.udn
@@ -0,0 +1,297 @@
+Availability: Public
+Crumbs: %ROOT%, Platforms
+Title: Detect SteamVR Chaperone System Activation
+Description: Detecting if the user has activated the SteamVR Chaperone system.
+Type: how-to
+Skilllevel: Intermidate
+Version: 4.17
+Parent: Platforms/SteamVR
+Order: 2
+tags: VR
+topic-image:Platforms\SteamVR\SteamVRDevelopment_topic.png
+
+The SteamVR Chaperone system is used to display the boundaries of the VR interaction area. As a tracked device approaches the boundaries, the SteamVR Runtime automatically provides visual cues notifying the user of this. In the following How - To we will take a look at how to display the interaction as well as other visual cues to notify the users that one of their devices is outside or are close to going outside the interaction area.
+
+[REGION:note]
+The following How - To is going to make use of the Standing VR Pawn with added Motion Controller and Chaperone support. If you are not sure how to set this up, the following two How - To's will help you out.
+* [Set Up Standing Cameras for SteamVR](Platforms\SteamVR\HowTo\StandingCamera)
+* [Set Up Motion Controllers](Platforms\VR\MotionController)
+* [Set Up the Chaperone System](Platforms\SteamVR\HowTo\Chaperone)
+[/REGION]
+
+[REGION:note]
+For the Chaperone system to work, you will need to make sure you have set up Room Scale VR using the SteamVR tools. For more information on how to do this, please see the official [HTC Vive](https://www.vive.com/us/support/category_howto/setting-up-room-scale-play-area.html) Setup pages.
+[/REGION]
+
+[REGION:warning]
+It is **NOT** possible nor advisable to disable the Chaperone system from inside UE4. You can, however, adjust how UE4 will react to the user getting close to the boundaries.  
+[/REGION]
+
+## Interacting with the SteamVR Chaperone System
+
+1.  To display the SteamVR Chaperone Boundary, we first have to know when one of the user's devices is close to it. To figure this out inside UE4, we can use the **Bind Event to OnRetunToBounds** and **Bind Event to OnLeaveBounds** nodes.  Inside of your player Pawn Blueprint add the following nodes to the **Event Graph**.
+
+    [REGION:lightbox]
+    [![](SteamVRChapActivation_00.png)(w:283)](SteamVRChapActivation_00.png)
+    [/REGION]
+
+    [REGION:caption]
+    Click for full image.
+    [/REGION]
+
+    * **Bind Event to OnRetunToBounds**
+    * **Bind Event to OnLeaveBounds**
+    * **SteamVRChaperone**
+
+1. Over in the **Variables** panel, create a new **Boolean** variable called **inBounds** so that we can track when the user leaves and enters the Chaperone bounds. 
+
+    ![](SteamVRChapActivation_01.png)
+
+1. Next, we want to make sure that as soon as the level starts, we are checking to see if the user is close to the bounds or not. To do this, we need to connect the Bind Event to OnRetunToBounds and Bind Event to OnLeaveBounds nodes together and then connect the **SteamVRChaperone** Object Reference to the **Target** input on both Bind Event nodes. Finally connecting the input to the Bind Event to OnLeaveBounds node to the output of the Set Tracking Origin node. When completed, your Pawn Blueprint should look like the following image.
+
+    [REGION:lightbox]
+    [![](SteamVRChapActivation_02.png)(w:833)](SteamVRChapActivation_02.png)
+    [/REGION]
+
+    [REGION:caption]
+    Click for full image.
+    [/REGION]
+
+1. Now we need to create some Blueprint logic that will help us track when the user has left or entered the bounds so that we display the appropriate message. To do this, we first need to create two new **Custom Event** nodes named **ReturnToBounds** and **LeaveBounds**. Next, we need to set the  **In Bounds** variable to **True** when we Return To Bounds and **False** when we Leave Bounds. When completed, your Pawn Blueprint should look like the following image.
+
+    [REGION:lightbox]
+    [![](SteamVRChapActivation_03.png)(w:434)](SteamVRChapActivation_03.png)
+    [/REGION]
+
+    [REGION:caption]
+    Click for full image.
+    [/REGION]
+
+1. Next, we need to connect the ReturnToBounds and LeaveBounds Event output to the respective input on the Bind Event nodes so that we can track when the users have returned or left the Chaperone bounds. When completed, your Pawn Blueprint should look like the following image.
+
+    [REGION:lightbox]
+    [![](SteamVRChapActivation_04.png)(w:833)](SteamVRChapActivation_04.png)
+    [/REGION]
+
+    [REGION:caption]
+    Click for full image.
+    [/REGION]
+
+1. We now need a way to notify the user that they have left or returned to the tracking area. To do this, we can use the **Print String** node and print a message to the screen every Tick that lets the user know if they are in or outside the tracking area. To set this up in your Pawn Blueprint, you will need to add the following nodes to the Event graph.
+
+    * **Event Tick**
+    * **Branch**
+    * **Print String** x 2 
+    * **InBounds**
+
+1. Connect the Event Tick output to the input on the Branch node and connect the **InBounds** to the **Condition** input and then connect the Print String nodes to the **True** and **False** outputs on the Branch node. Make sure to input some text in the Print String nodes that describes what is going. For this example the phrase **You have returned to the tracking bounds**  and **You have left the tracking bounds** where used but feel free to use whatever you like. When completed your Pawn Blueprint should look like the following image.
+
+    [REGION:lightbox]
+    [Include:Platforms\SteamVR\HowTo\ChaperoneActivation\#SVRFinalBP]
+    [/REGION]
+
+    [REGION:caption]
+    Click on image to copy Blueprint code.
+    [/REGION]
+
+1. Now, put on your HTC Vive HMD on, pick up your Motion controllers and launch the project using VR Preview. Once the project launches, slowly move one of the Motion controllers to towards the Chaperone boundary, and as soon as you get close the message that is printed to the screen will change like in the following video.
+
+    [OBJECT:EmbeddedVideo]
+        [PARAMLITERAL:videoid]
+            _QyYbsdbQT4
+        [/PARAMLITERAL]
+        [PARAMLITERAL:width]
+        100
+        [/PARAMLITERAL]
+        [PARAMLITERAL:height]
+        56.25
+        [/PARAMLITERAL]
+        [PARAMLITERAL:units]
+        %
+        [/PARAMLITERAL]
+        [PARAMLITERAL:type]
+        simple
+        [/PARAMLITERAL]
+    [/OBJECT]
+
+## UE4 Project Downloads
+Below you will find a link to where you can download the UE4 project that was used to create this example.
+
+* [**SteamVR Chaperon Activation Example Project**](https://epicgames.box.com/s/zpvbgdudx756l8n6z50fk58pap4ou7u0)
+
+<!--
+[EXCERPT:SVRFinalBP]	
+[OBJECT:Snippet]
+	[PARAMLITERAL:id]
+		SVRFinalBP
+	[/PARAMLITERAL]
+	[PARAMLITERAL:type]
+		blueprint
+	[/PARAMLITERAL]
+	[PARAM:image]
+		[REGION:lightbox]
+		[![](SteamVRChapActivation_05.png)(w:1000)](SteamVRChapActivation_05.png)
+		[/REGION]
+	[/PARAM]
+	[PARAM:snippet]
+			BEGIN OBJECT
+				Begin Object Class=/Script/BlueprintGraph.K2Node_Event Name="K2Node_Event_3"
+			   EventReference=(MemberParent=Class'/Script/Engine.Actor',MemberName="ReceiveBeginPlay")
+			   bOverrideFunction=True
+			   NodePosX=-192
+			   NodePosY=560
+			   bCommentBubblePinned=True
+			   NodeGuid=71F8712B42D30302CBA9CE84C38862B7
+			   CustomProperties Pin (PinId=63E599D74E38B8776D6E88A4AAB2C2B3,PinName="OutputDelegate",Direction="EGPD_Output",PinType.PinCategory="delegate",PinType.PinSubCategory="",PinType.PinSubCategoryObject=None,PinType.PinSubCategoryMemberReference=(MemberParent=Class'/Script/Engine.Actor',MemberName="ReceiveBeginPlay"),PinType.PinValueType=(),PinType.ContainerType=None,PinType.bIsMap=False,PinType.bIsSet=False,PinType.bIsArray=False,PinType.bIsReference=False,PinType.bIsConst=False,PinType.bIsWeakPointer=False,PersistentGuid=00000000000000000000000000000000,bHidden=False,bNotConnectable=False,bDefaultValueIsReadOnly=False,bDefaultValueIsIgnored=False,bAdvancedView=False,bOrphanedPin=False,)
+			   CustomProperties Pin (PinId=96D13FDB4718AC06F0F132AEEBA98395,PinName="then",Direction="EGPD_Output",PinType.PinCategory="exec",PinType.PinSubCategory="",PinType.PinSubCategoryObject=None,PinType.PinSubCategoryMemberReference=(),PinType.PinValueType=(),PinType.ContainerType=None,PinType.bIsMap=False,PinType.bIsSet=False,PinType.bIsArray=False,PinType.bIsReference=False,PinType.bIsConst=False,PinType.bIsWeakPointer=False,LinkedTo=(K2Node_CallFunction_7 0AF4F2C243DAA259B82A0AB55A91B640,),PersistentGuid=00000000000000000000000000000000,bHidden=False,bNotConnectable=False,bDefaultValueIsReadOnly=False,bDefaultValueIsIgnored=False,bAdvancedView=False,bOrphanedPin=False,)
+			End Object
+			Begin Object Class=/Script/BlueprintGraph.K2Node_CallFunction Name="K2Node_CallFunction_7"
+			   FunctionReference=(MemberParent=Class'/Script/HeadMountedDisplay.HeadMountedDisplayFunctionLibrary',MemberName="SetTrackingOrigin")
+			   NodePosY=560
+			   NodeGuid=401DF7CA48C88D2786ECB8B5C2AC8F46
+			   CustomProperties Pin (PinId=0AF4F2C243DAA259B82A0AB55A91B640,PinName="execute",PinToolTip="\nExec",PinType.PinCategory="exec",PinType.PinSubCategory="",PinType.PinSubCategoryObject=None,PinType.PinSubCategoryMemberReference=(),PinType.PinValueType=(),PinType.ContainerType=None,PinType.bIsMap=False,PinType.bIsSet=False,PinType.bIsArray=False,PinType.bIsReference=False,PinType.bIsConst=False,PinType.bIsWeakPointer=False,LinkedTo=(K2Node_Event_3 96D13FDB4718AC06F0F132AEEBA98395,),PersistentGuid=00000000000000000000000000000000,bHidden=False,bNotConnectable=False,bDefaultValueIsReadOnly=False,bDefaultValueIsIgnored=False,bAdvancedView=False,bOrphanedPin=False,)
+			   CustomProperties Pin (PinId=E9F9736F435EF94B7A4B16A0C9347FFB,PinName="then",PinToolTip="\nExec",Direction="EGPD_Output",PinType.PinCategory="exec",PinType.PinSubCategory="",PinType.PinSubCategoryObject=None,PinType.PinSubCategoryMemberReference=(),PinType.PinValueType=(),PinType.ContainerType=None,PinType.bIsMap=False,PinType.bIsSet=False,PinType.bIsArray=False,PinType.bIsReference=False,PinType.bIsConst=False,PinType.bIsWeakPointer=False,LinkedTo=(K2Node_AddDelegate_8 DC0FD5A54E4CB179DA797C9730EDB755,),PersistentGuid=00000000000000000000000000000000,bHidden=False,bNotConnectable=False,bDefaultValueIsReadOnly=False,bDefaultValueIsIgnored=False,bAdvancedView=False,bOrphanedPin=False,)
+			   CustomProperties Pin (PinId=E21CFDB440B456D7A54A2B991677B221,PinName="self",PinFriendlyName=NSLOCTEXT("K2Node", "Target", "Target"),PinToolTip="Target\nHead Mounted Display Function Library Object Reference",PinType.PinCategory="object",PinType.PinSubCategory="",PinType.PinSubCategoryObject=Class'/Script/HeadMountedDisplay.HeadMountedDisplayFunctionLibrary',PinType.PinSubCategoryMemberReference=(),PinType.PinValueType=(),PinType.ContainerType=None,PinType.bIsMap=False,PinType.bIsSet=False,PinType.bIsArray=False,PinType.bIsReference=False,PinType.bIsConst=False,PinType.bIsWeakPointer=False,DefaultObject="/Script/HeadMountedDisplay.Default__HeadMountedDisplayFunctionLibrary",PersistentGuid=00000000000000000000000000000000,bHidden=True,bNotConnectable=False,bDefaultValueIsReadOnly=False,bDefaultValueIsIgnored=False,bAdvancedView=False,bOrphanedPin=False,)
+			   CustomProperties Pin (PinId=D3A2C5724D17455C5C1DCBAA7E724282,PinName="Origin",PinToolTip="Origin\nEHMDTrackingOrigin Enum",PinType.PinCategory="byte",PinType.PinSubCategory="",PinType.PinSubCategoryObject=Enum'/Script/HeadMountedDisplay.EHMDTrackingOrigin',PinType.PinSubCategoryMemberReference=(),PinType.PinValueType=(),PinType.ContainerType=None,PinType.bIsMap=False,PinType.bIsSet=False,PinType.bIsArray=False,PinType.bIsReference=False,PinType.bIsConst=False,PinType.bIsWeakPointer=False,DefaultValue="Floor",PersistentGuid=00000000000000000000000000000000,bHidden=False,bNotConnectable=False,bDefaultValueIsReadOnly=False,bDefaultValueIsIgnored=False,bAdvancedView=False,bOrphanedPin=False,)
+			End Object
+			Begin Object Class=/Script/UnrealEd.EdGraphNode_Comment Name="EdGraphNode_Comment_1"
+			   NodePosX=-22
+			   NodePosY=493
+			   NodeWidth=304
+			   NodeHeight=224
+			   NodeComment="Set Tracking Orgin"
+			   NodeGuid=EC11D40A4E737119184DB7AA39CA25A9
+			End Object
+			Begin Object Class=/Script/BlueprintGraph.K2Node_AddDelegate Name="K2Node_AddDelegate_0"
+			   DelegateReference=(MemberParent=Class'/Script/SteamVR.SteamVRChaperoneComponent',MemberName="OnReturnToBounds")
+			   NodePosX=736
+			   NodePosY=560
+			   NodeGuid=AAAA740B4DDFC88348F8F39CE67A15FB
+			   CustomProperties Pin (PinId=ACA57A1F41254B5105DBA5B8A3E9B465,PinName="execute",PinType.PinCategory="exec",PinType.PinSubCategory="",PinType.PinSubCategoryObject=None,PinType.PinSubCategoryMemberReference=(),PinType.PinValueType=(),PinType.ContainerType=None,PinType.bIsMap=False,PinType.bIsSet=False,PinType.bIsArray=False,PinType.bIsReference=False,PinType.bIsConst=False,PinType.bIsWeakPointer=False,LinkedTo=(K2Node_AddDelegate_8 0E486F4841FA0210BCAC5DBB3CDBCAC6,),PersistentGuid=00000000000000000000000000000000,bHidden=False,bNotConnectable=False,bDefaultValueIsReadOnly=False,bDefaultValueIsIgnored=False,bAdvancedView=False,bOrphanedPin=False,)
+			   CustomProperties Pin (PinId=C7A1C99E4EFE7F38B38FE39F566802E0,PinName="then",Direction="EGPD_Output",PinType.PinCategory="exec",PinType.PinSubCategory="",PinType.PinSubCategoryObject=None,PinType.PinSubCategoryMemberReference=(),PinType.PinValueType=(),PinType.ContainerType=None,PinType.bIsMap=False,PinType.bIsSet=False,PinType.bIsArray=False,PinType.bIsReference=False,PinType.bIsConst=False,PinType.bIsWeakPointer=False,PersistentGuid=00000000000000000000000000000000,bHidden=False,bNotConnectable=False,bDefaultValueIsReadOnly=False,bDefaultValueIsIgnored=False,bAdvancedView=False,bOrphanedPin=False,)
+			   CustomProperties Pin (PinId=F9303983417A9E8E17973BBA54D96191,PinName="self",PinFriendlyName=NSLOCTEXT("K2Node", "BaseMCDelegateSelfPinName", "Target"),PinType.PinCategory="object",PinType.PinSubCategory="",PinType.PinSubCategoryObject=Class'/Script/SteamVR.SteamVRChaperoneComponent',PinType.PinSubCategoryMemberReference=(),PinType.PinValueType=(),PinType.ContainerType=None,PinType.bIsMap=False,PinType.bIsSet=False,PinType.bIsArray=False,PinType.bIsReference=False,PinType.bIsConst=False,PinType.bIsWeakPointer=False,LinkedTo=(K2Node_VariableGet_1633 D95560FA45E6D9502D40D796C287ED62,),PersistentGuid=00000000000000000000000000000000,bHidden=False,bNotConnectable=False,bDefaultValueIsReadOnly=False,bDefaultValueIsIgnored=False,bAdvancedView=False,bOrphanedPin=False,)
+			   CustomProperties Pin (PinId=177D13EF4A04AEBC62607D90670C7B1A,PinName="Delegate",PinFriendlyName=NSLOCTEXT("K2Node", "PinFriendlyDelegatetName", "Event"),PinType.PinCategory="delegate",PinType.PinSubCategory="",PinType.PinSubCategoryObject=None,PinType.PinSubCategoryMemberReference=(MemberParent=Class'/Script/SteamVR.SteamVRChaperoneComponent',MemberName="SteamVRChaperoneEvent__DelegateSignature"),PinType.PinValueType=(),PinType.ContainerType=None,PinType.bIsMap=False,PinType.bIsSet=False,PinType.bIsArray=False,PinType.bIsReference=True,PinType.bIsConst=True,PinType.bIsWeakPointer=False,LinkedTo=(K2Node_CustomEvent_0 F90945A14C76ADFA94D6299D67B2E402,),PersistentGuid=00000000000000000000000000000000,bHidden=False,bNotConnectable=False,bDefaultValueIsReadOnly=False,bDefaultValueIsIgnored=False,bAdvancedView=False,bOrphanedPin=False,)
+			End Object
+			Begin Object Class=/Script/BlueprintGraph.K2Node_VariableGet Name="K2Node_VariableGet_1633"
+			   VariableReference=(MemberName="SteamVRChaperone",bSelfContext=True)
+			   NodePosX=96
+			   NodePosY=752
+			   NodeGuid=A8155B024551A4C3C4B809AFC8804A08
+			   CustomProperties Pin (PinId=D95560FA45E6D9502D40D796C287ED62,PinName="SteamVRChaperone",Direction="EGPD_Output",PinType.PinCategory="object",PinType.PinSubCategory="",PinType.PinSubCategoryObject=Class'/Script/SteamVR.SteamVRChaperoneComponent',PinType.PinSubCategoryMemberReference=(),PinType.PinValueType=(),PinType.ContainerType=None,PinType.bIsMap=False,PinType.bIsSet=False,PinType.bIsArray=False,PinType.bIsReference=False,PinType.bIsConst=False,PinType.bIsWeakPointer=False,LinkedTo=(K2Node_AddDelegate_8 2AA013A6411401005E187FA5F35A422B,K2Node_AddDelegate_0 F9303983417A9E8E17973BBA54D96191,),PersistentGuid=00000000000000000000000000000000,bHidden=False,bNotConnectable=False,bDefaultValueIsReadOnly=False,bDefaultValueIsIgnored=False,bAdvancedView=False,bOrphanedPin=False,)
+			   CustomProperties Pin (PinId=EE51CFDD46E30EC35BAF96A0C5273DB3,PinName="self",PinFriendlyName=NSLOCTEXT("K2Node", "Target", "Target"),PinType.PinCategory="object",PinType.PinSubCategory="",PinType.PinSubCategoryObject=BlueprintGeneratedClass'/Game/Blueprints/SteamVR_Pawn_Chaperone.SteamVR_Pawn_Chaperone_C',PinType.PinSubCategoryMemberReference=(),PinType.PinValueType=(),PinType.ContainerType=None,PinType.bIsMap=False,PinType.bIsSet=False,PinType.bIsArray=False,PinType.bIsReference=False,PinType.bIsConst=False,PinType.bIsWeakPointer=False,PersistentGuid=00000000000000000000000000000000,bHidden=True,bNotConnectable=False,bDefaultValueIsReadOnly=False,bDefaultValueIsIgnored=False,bAdvancedView=False,bOrphanedPin=False,)
+			End Object
+			Begin Object Class=/Script/BlueprintGraph.K2Node_CustomEvent Name="K2Node_CustomEvent_0"
+			   CustomFunctionName="ReturnToBounds"
+			   NodePosX=496
+			   NodePosY=368
+			   NodeGuid=3E51186C45950AD11A6CF7A1A0B47653
+			   CustomProperties Pin (PinId=F90945A14C76ADFA94D6299D67B2E402,PinName="OutputDelegate",Direction="EGPD_Output",PinType.PinCategory="delegate",PinType.PinSubCategory="",PinType.PinSubCategoryObject=None,PinType.PinSubCategoryMemberReference=(MemberParent=BlueprintGeneratedClass'/Game/Blueprints/SteamVR_Pawn_Chaperone.SteamVR_Pawn_Chaperone_C',MemberName="ReturnToBounds",MemberGuid=3E51186C45950AD11A6CF7A1A0B47653),PinType.PinValueType=(),PinType.ContainerType=None,PinType.bIsMap=False,PinType.bIsSet=False,PinType.bIsArray=False,PinType.bIsReference=False,PinType.bIsConst=False,PinType.bIsWeakPointer=False,LinkedTo=(K2Node_AddDelegate_0 177D13EF4A04AEBC62607D90670C7B1A,),PersistentGuid=00000000000000000000000000000000,bHidden=False,bNotConnectable=False,bDefaultValueIsReadOnly=False,bDefaultValueIsIgnored=False,bAdvancedView=False,bOrphanedPin=False,)
+			   CustomProperties Pin (PinId=9105FF0749B49F5B8699AF9EF6E2CB90,PinName="then",Direction="EGPD_Output",PinType.PinCategory="exec",PinType.PinSubCategory="",PinType.PinSubCategoryObject=None,PinType.PinSubCategoryMemberReference=(),PinType.PinValueType=(),PinType.ContainerType=None,PinType.bIsMap=False,PinType.bIsSet=False,PinType.bIsArray=False,PinType.bIsReference=False,PinType.bIsConst=False,PinType.bIsWeakPointer=False,LinkedTo=(K2Node_VariableSet_1 950E57D34CABE389C8C52384D0937924,),PersistentGuid=00000000000000000000000000000000,bHidden=False,bNotConnectable=False,bDefaultValueIsReadOnly=False,bDefaultValueIsIgnored=False,bAdvancedView=False,bOrphanedPin=False,)
+			End Object
+			Begin Object Class=/Script/BlueprintGraph.K2Node_VariableSet Name="K2Node_VariableSet_1"
+			   VariableReference=(MemberName="inBounds",MemberGuid=352D7A464AB7AFE78D60A0A0AEB70382,bSelfContext=True)
+			   NodePosX=720
+			   NodePosY=400
+			   NodeGuid=155C3F4949846E8F725478A1288A01B9
+			   CustomProperties Pin (PinId=950E57D34CABE389C8C52384D0937924,PinName="execute",PinType.PinCategory="exec",PinType.PinSubCategory="",PinType.PinSubCategoryObject=None,PinType.PinSubCategoryMemberReference=(),PinType.PinValueType=(),PinType.ContainerType=None,PinType.bIsMap=False,PinType.bIsSet=False,PinType.bIsArray=False,PinType.bIsReference=False,PinType.bIsConst=False,PinType.bIsWeakPointer=False,LinkedTo=(K2Node_CustomEvent_0 9105FF0749B49F5B8699AF9EF6E2CB90,),PersistentGuid=00000000000000000000000000000000,bHidden=False,bNotConnectable=False,bDefaultValueIsReadOnly=False,bDefaultValueIsIgnored=False,bAdvancedView=False,bOrphanedPin=False,)
+			   CustomProperties Pin (PinId=6F940A414136B1733C973B90F2571FFE,PinName="then",Direction="EGPD_Output",PinType.PinCategory="exec",PinType.PinSubCategory="",PinType.PinSubCategoryObject=None,PinType.PinSubCategoryMemberReference=(),PinType.PinValueType=(),PinType.ContainerType=None,PinType.bIsMap=False,PinType.bIsSet=False,PinType.bIsArray=False,PinType.bIsReference=False,PinType.bIsConst=False,PinType.bIsWeakPointer=False,PersistentGuid=00000000000000000000000000000000,bHidden=False,bNotConnectable=False,bDefaultValueIsReadOnly=False,bDefaultValueIsIgnored=False,bAdvancedView=False,bOrphanedPin=False,)
+			   CustomProperties Pin (PinId=66439B5848F1408F17784D98B32D1CD4,PinName="inBounds",PinType.PinCategory="bool",PinType.PinSubCategory="",PinType.PinSubCategoryObject=None,PinType.PinSubCategoryMemberReference=(),PinType.PinValueType=(),PinType.ContainerType=None,PinType.bIsMap=False,PinType.bIsSet=False,PinType.bIsArray=False,PinType.bIsReference=False,PinType.bIsConst=False,PinType.bIsWeakPointer=False,DefaultValue="true",AutogeneratedDefaultValue="false",PersistentGuid=00000000000000000000000000000000,bHidden=False,bNotConnectable=False,bDefaultValueIsReadOnly=False,bDefaultValueIsIgnored=False,bAdvancedView=False,bOrphanedPin=False,)
+			   CustomProperties Pin (PinId=6033C36647C705F00096D590D02C9C0C,PinName="self",PinFriendlyName=NSLOCTEXT("K2Node", "Target", "Target"),PinType.PinCategory="object",PinType.PinSubCategory="",PinType.PinSubCategoryObject=BlueprintGeneratedClass'/Game/Blueprints/SteamVR_Pawn_Chaperone.SteamVR_Pawn_Chaperone_C',PinType.PinSubCategoryMemberReference=(),PinType.PinValueType=(),PinType.ContainerType=None,PinType.bIsMap=False,PinType.bIsSet=False,PinType.bIsArray=False,PinType.bIsReference=False,PinType.bIsConst=False,PinType.bIsWeakPointer=False,PersistentGuid=00000000000000000000000000000000,bHidden=True,bNotConnectable=False,bDefaultValueIsReadOnly=False,bDefaultValueIsIgnored=False,bAdvancedView=False,bOrphanedPin=False,)
+			   CustomProperties Pin (PinId=87CF008F4AB038060A33A7B6897E6AAF,PinName="Output_Get",PinToolTip="Retrieves the value of the variable, can use instead of a separate Get node",Direction="EGPD_Output",PinType.PinCategory="bool",PinType.PinSubCategory="",PinType.PinSubCategoryObject=None,PinType.PinSubCategoryMemberReference=(),PinType.PinValueType=(),PinType.ContainerType=None,PinType.bIsMap=False,PinType.bIsSet=False,PinType.bIsArray=False,PinType.bIsReference=False,PinType.bIsConst=False,PinType.bIsWeakPointer=False,DefaultValue="false",AutogeneratedDefaultValue="false",PersistentGuid=00000000000000000000000000000000,bHidden=False,bNotConnectable=False,bDefaultValueIsReadOnly=False,bDefaultValueIsIgnored=False,bAdvancedView=False,bOrphanedPin=False,)
+			End Object
+			Begin Object Class=/Script/BlueprintGraph.K2Node_AddDelegate Name="K2Node_AddDelegate_8"
+			   DelegateReference=(MemberParent=Class'/Script/SteamVR.SteamVRChaperoneComponent',MemberName="OnLeaveBounds")
+			   NodePosX=368
+			   NodePosY=560
+			   NodeGuid=F871EDF540F09305DFCCE3A88B5B1FA3
+			   CustomProperties Pin (PinId=DC0FD5A54E4CB179DA797C9730EDB755,PinName="execute",PinType.PinCategory="exec",PinType.PinSubCategory="",PinType.PinSubCategoryObject=None,PinType.PinSubCategoryMemberReference=(),PinType.PinValueType=(),PinType.ContainerType=None,PinType.bIsMap=False,PinType.bIsSet=False,PinType.bIsArray=False,PinType.bIsReference=False,PinType.bIsConst=False,PinType.bIsWeakPointer=False,LinkedTo=(K2Node_CallFunction_7 E9F9736F435EF94B7A4B16A0C9347FFB,),PersistentGuid=00000000000000000000000000000000,bHidden=False,bNotConnectable=False,bDefaultValueIsReadOnly=False,bDefaultValueIsIgnored=False,bAdvancedView=False,bOrphanedPin=False,)
+			   CustomProperties Pin (PinId=0E486F4841FA0210BCAC5DBB3CDBCAC6,PinName="then",Direction="EGPD_Output",PinType.PinCategory="exec",PinType.PinSubCategory="",PinType.PinSubCategoryObject=None,PinType.PinSubCategoryMemberReference=(),PinType.PinValueType=(),PinType.ContainerType=None,PinType.bIsMap=False,PinType.bIsSet=False,PinType.bIsArray=False,PinType.bIsReference=False,PinType.bIsConst=False,PinType.bIsWeakPointer=False,LinkedTo=(K2Node_AddDelegate_0 ACA57A1F41254B5105DBA5B8A3E9B465,),PersistentGuid=00000000000000000000000000000000,bHidden=False,bNotConnectable=False,bDefaultValueIsReadOnly=False,bDefaultValueIsIgnored=False,bAdvancedView=False,bOrphanedPin=False,)
+			   CustomProperties Pin (PinId=2AA013A6411401005E187FA5F35A422B,PinName="self",PinFriendlyName=NSLOCTEXT("K2Node", "BaseMCDelegateSelfPinName", "Target"),PinType.PinCategory="object",PinType.PinSubCategory="",PinType.PinSubCategoryObject=Class'/Script/SteamVR.SteamVRChaperoneComponent',PinType.PinSubCategoryMemberReference=(),PinType.PinValueType=(),PinType.ContainerType=None,PinType.bIsMap=False,PinType.bIsSet=False,PinType.bIsArray=False,PinType.bIsReference=False,PinType.bIsConst=False,PinType.bIsWeakPointer=False,LinkedTo=(K2Node_VariableGet_1633 D95560FA45E6D9502D40D796C287ED62,),PersistentGuid=00000000000000000000000000000000,bHidden=False,bNotConnectable=False,bDefaultValueIsReadOnly=False,bDefaultValueIsIgnored=False,bAdvancedView=False,bOrphanedPin=False,)
+			   CustomProperties Pin (PinId=3925FF0F494A9C7DE55C72810451F48D,PinName="Delegate",PinFriendlyName=NSLOCTEXT("K2Node", "PinFriendlyDelegatetName", "Event"),PinType.PinCategory="delegate",PinType.PinSubCategory="",PinType.PinSubCategoryObject=None,PinType.PinSubCategoryMemberReference=(MemberParent=Class'/Script/SteamVR.SteamVRChaperoneComponent',MemberName="SteamVRChaperoneEvent__DelegateSignature"),PinType.PinValueType=(),PinType.ContainerType=None,PinType.bIsMap=False,PinType.bIsSet=False,PinType.bIsArray=False,PinType.bIsReference=True,PinType.bIsConst=True,PinType.bIsWeakPointer=False,LinkedTo=(K2Node_CustomEvent_2 275B80D047C68B56ADC40EB36D74CBA7,),PersistentGuid=00000000000000000000000000000000,bHidden=False,bNotConnectable=False,bDefaultValueIsReadOnly=False,bDefaultValueIsIgnored=False,bAdvancedView=False,bOrphanedPin=False,)
+			End Object
+			Begin Object Class=/Script/BlueprintGraph.K2Node_CustomEvent Name="K2Node_CustomEvent_2"
+			   CustomFunctionName="LeaveBounds"
+			   NodePosX=80
+			   NodePosY=368
+			   NodeGuid=F91E8A8C481CA8D644EE05A15D8E5A94
+			   CustomProperties Pin (PinId=275B80D047C68B56ADC40EB36D74CBA7,PinName="OutputDelegate",Direction="EGPD_Output",PinType.PinCategory="delegate",PinType.PinSubCategory="",PinType.PinSubCategoryObject=None,PinType.PinSubCategoryMemberReference=(MemberParent=BlueprintGeneratedClass'/Game/Blueprints/SteamVR_Pawn_Chaperone.SteamVR_Pawn_Chaperone_C',MemberName="LeaveBounds",MemberGuid=F91E8A8C481CA8D644EE05A15D8E5A94),PinType.PinValueType=(),PinType.ContainerType=None,PinType.bIsMap=False,PinType.bIsSet=False,PinType.bIsArray=False,PinType.bIsReference=False,PinType.bIsConst=False,PinType.bIsWeakPointer=False,LinkedTo=(K2Node_AddDelegate_8 3925FF0F494A9C7DE55C72810451F48D,),PersistentGuid=00000000000000000000000000000000,bHidden=False,bNotConnectable=False,bDefaultValueIsReadOnly=False,bDefaultValueIsIgnored=False,bAdvancedView=False,bOrphanedPin=False,)
+			   CustomProperties Pin (PinId=3E8E611F4E488133D139B985F11E4D82,PinName="then",Direction="EGPD_Output",PinType.PinCategory="exec",PinType.PinSubCategory="",PinType.PinSubCategoryObject=None,PinType.PinSubCategoryMemberReference=(),PinType.PinValueType=(),PinType.ContainerType=None,PinType.bIsMap=False,PinType.bIsSet=False,PinType.bIsArray=False,PinType.bIsReference=False,PinType.bIsConst=False,PinType.bIsWeakPointer=False,LinkedTo=(K2Node_VariableSet_730 950E57D34CABE389C8C52384D0937924,),PersistentGuid=00000000000000000000000000000000,bHidden=False,bNotConnectable=False,bDefaultValueIsReadOnly=False,bDefaultValueIsIgnored=False,bAdvancedView=False,bOrphanedPin=False,)
+			End Object
+			Begin Object Class=/Script/BlueprintGraph.K2Node_VariableSet Name="K2Node_VariableSet_730"
+			   VariableReference=(MemberName="inBounds",MemberGuid=352D7A464AB7AFE78D60A0A0AEB70382,bSelfContext=True)
+			   NodePosX=304
+			   NodePosY=400
+			   NodeGuid=86400F864EA1FD3418F8F89090DE648B
+			   CustomProperties Pin (PinId=950E57D34CABE389C8C52384D0937924,PinName="execute",PinType.PinCategory="exec",PinType.PinSubCategory="",PinType.PinSubCategoryObject=None,PinType.PinSubCategoryMemberReference=(),PinType.PinValueType=(),PinType.ContainerType=None,PinType.bIsMap=False,PinType.bIsSet=False,PinType.bIsArray=False,PinType.bIsReference=False,PinType.bIsConst=False,PinType.bIsWeakPointer=False,LinkedTo=(K2Node_CustomEvent_2 3E8E611F4E488133D139B985F11E4D82,),PersistentGuid=00000000000000000000000000000000,bHidden=False,bNotConnectable=False,bDefaultValueIsReadOnly=False,bDefaultValueIsIgnored=False,bAdvancedView=False,bOrphanedPin=False,)
+			   CustomProperties Pin (PinId=6F940A414136B1733C973B90F2571FFE,PinName="then",Direction="EGPD_Output",PinType.PinCategory="exec",PinType.PinSubCategory="",PinType.PinSubCategoryObject=None,PinType.PinSubCategoryMemberReference=(),PinType.PinValueType=(),PinType.ContainerType=None,PinType.bIsMap=False,PinType.bIsSet=False,PinType.bIsArray=False,PinType.bIsReference=False,PinType.bIsConst=False,PinType.bIsWeakPointer=False,PersistentGuid=00000000000000000000000000000000,bHidden=False,bNotConnectable=False,bDefaultValueIsReadOnly=False,bDefaultValueIsIgnored=False,bAdvancedView=False,bOrphanedPin=False,)
+			   CustomProperties Pin (PinId=66439B5848F1408F17784D98B32D1CD4,PinName="inBounds",PinType.PinCategory="bool",PinType.PinSubCategory="",PinType.PinSubCategoryObject=None,PinType.PinSubCategoryMemberReference=(),PinType.PinValueType=(),PinType.ContainerType=None,PinType.bIsMap=False,PinType.bIsSet=False,PinType.bIsArray=False,PinType.bIsReference=False,PinType.bIsConst=False,PinType.bIsWeakPointer=False,DefaultValue="false",AutogeneratedDefaultValue="false",PersistentGuid=00000000000000000000000000000000,bHidden=False,bNotConnectable=False,bDefaultValueIsReadOnly=False,bDefaultValueIsIgnored=False,bAdvancedView=False,bOrphanedPin=False,)
+			   CustomProperties Pin (PinId=87CF008F4AB038060A33A7B6897E6AAF,PinName="Output_Get",PinToolTip="Retrieves the value of the variable, can use instead of a separate Get node",Direction="EGPD_Output",PinType.PinCategory="bool",PinType.PinSubCategory="",PinType.PinSubCategoryObject=None,PinType.PinSubCategoryMemberReference=(),PinType.PinValueType=(),PinType.ContainerType=None,PinType.bIsMap=False,PinType.bIsSet=False,PinType.bIsArray=False,PinType.bIsReference=False,PinType.bIsConst=False,PinType.bIsWeakPointer=False,DefaultValue="false",AutogeneratedDefaultValue="false",PersistentGuid=00000000000000000000000000000000,bHidden=False,bNotConnectable=False,bDefaultValueIsReadOnly=False,bDefaultValueIsIgnored=False,bAdvancedView=False,bOrphanedPin=False,)
+			   CustomProperties Pin (PinId=6033C36647C705F00096D590D02C9C0C,PinName="self",PinFriendlyName=NSLOCTEXT("K2Node", "Target", "Target"),PinType.PinCategory="object",PinType.PinSubCategory="",PinType.PinSubCategoryObject=BlueprintGeneratedClass'/Game/Blueprints/SteamVR_Pawn_Chaperone.SteamVR_Pawn_Chaperone_C',PinType.PinSubCategoryMemberReference=(),PinType.PinValueType=(),PinType.ContainerType=None,PinType.bIsMap=False,PinType.bIsSet=False,PinType.bIsArray=False,PinType.bIsReference=False,PinType.bIsConst=False,PinType.bIsWeakPointer=False,PersistentGuid=00000000000000000000000000000000,bHidden=True,bNotConnectable=False,bDefaultValueIsReadOnly=False,bDefaultValueIsIgnored=False,bAdvancedView=False,bOrphanedPin=False,)
+			End Object
+			Begin Object Class=/Script/BlueprintGraph.K2Node_CallFunction Name="K2Node_CallFunction_121"
+			   FunctionReference=(MemberParent=Class'/Script/Engine.KismetSystemLibrary',MemberName="PrintString")
+			   NodePosX=400
+			   NodePosY=928
+			   AdvancedPinDisplay=Hidden
+			   EnabledState=DevelopmentOnly
+			   NodeGuid=1BCED2B1460B2AA8CDDD61BA55DA53E3
+			   CustomProperties Pin (PinId=6863525643E71242071D6481F714C829,PinName="execute",PinToolTip="\nExec",PinType.PinCategory="exec",PinType.PinSubCategory="",PinType.PinSubCategoryObject=None,PinType.PinSubCategoryMemberReference=(),PinType.PinValueType=(),PinType.ContainerType=None,PinType.bIsMap=False,PinType.bIsSet=False,PinType.bIsArray=False,PinType.bIsReference=False,PinType.bIsConst=False,PinType.bIsWeakPointer=False,LinkedTo=(K2Node_IfThenElse_120 38C905BB46846179190960A620BCC78C,),PersistentGuid=00000000000000000000000000000000,bHidden=False,bNotConnectable=False,bDefaultValueIsReadOnly=False,bDefaultValueIsIgnored=False,bAdvancedView=False,bOrphanedPin=False,)
+			   CustomProperties Pin (PinId=E13897364CE1B5A8A6D9D89E7CAA376D,PinName="then",PinToolTip="\nExec",Direction="EGPD_Output",PinType.PinCategory="exec",PinType.PinSubCategory="",PinType.PinSubCategoryObject=None,PinType.PinSubCategoryMemberReference=(),PinType.PinValueType=(),PinType.ContainerType=None,PinType.bIsMap=False,PinType.bIsSet=False,PinType.bIsArray=False,PinType.bIsReference=False,PinType.bIsConst=False,PinType.bIsWeakPointer=False,PersistentGuid=00000000000000000000000000000000,bHidden=False,bNotConnectable=False,bDefaultValueIsReadOnly=False,bDefaultValueIsIgnored=False,bAdvancedView=False,bOrphanedPin=False,)
+			   CustomProperties Pin (PinId=5E4EFBED43607C5035952DA4524D24FA,PinName="self",PinFriendlyName=NSLOCTEXT("K2Node", "Target", "Target"),PinToolTip="Target\nKismet System Library Object Reference",PinType.PinCategory="object",PinType.PinSubCategory="",PinType.PinSubCategoryObject=Class'/Script/Engine.KismetSystemLibrary',PinType.PinSubCategoryMemberReference=(),PinType.PinValueType=(),PinType.ContainerType=None,PinType.bIsMap=False,PinType.bIsSet=False,PinType.bIsArray=False,PinType.bIsReference=False,PinType.bIsConst=False,PinType.bIsWeakPointer=False,DefaultObject="/Script/Engine.Default__KismetSystemLibrary",PersistentGuid=00000000000000000000000000000000,bHidden=True,bNotConnectable=False,bDefaultValueIsReadOnly=False,bDefaultValueIsIgnored=False,bAdvancedView=False,bOrphanedPin=False,)
+			   CustomProperties Pin (PinId=AC46C00F409FB31290CB0ABD8B6865D4,PinName="WorldContextObject",PinToolTip="World Context Object\nObject Object Reference",PinType.PinCategory="object",PinType.PinSubCategory="",PinType.PinSubCategoryObject=Class'/Script/CoreUObject.Object',PinType.PinSubCategoryMemberReference=(),PinType.PinValueType=(),PinType.ContainerType=None,PinType.bIsMap=False,PinType.bIsSet=False,PinType.bIsArray=False,PinType.bIsReference=False,PinType.bIsConst=False,PinType.bIsWeakPointer=False,PersistentGuid=00000000000000000000000000000000,bHidden=True,bNotConnectable=False,bDefaultValueIsReadOnly=False,bDefaultValueIsIgnored=False,bAdvancedView=False,bOrphanedPin=False,)
+			   CustomProperties Pin (PinId=8E4D51D74047E6E7A7B8C2B2FEBBC962,PinName="InString",PinToolTip="In String\nString\n\nThe string to log out",PinType.PinCategory="string",PinType.PinSubCategory="",PinType.PinSubCategoryObject=None,PinType.PinSubCategoryMemberReference=(),PinType.PinValueType=(),PinType.ContainerType=None,PinType.bIsMap=False,PinType.bIsSet=False,PinType.bIsArray=False,PinType.bIsReference=False,PinType.bIsConst=False,PinType.bIsWeakPointer=False,DefaultValue="You have left the tracking bounds",AutogeneratedDefaultValue="Hello",PersistentGuid=00000000000000000000000000000000,bHidden=False,bNotConnectable=False,bDefaultValueIsReadOnly=False,bDefaultValueIsIgnored=False,bAdvancedView=False,bOrphanedPin=False,)
+			   CustomProperties Pin (PinId=3E2695224F6B062FDED392854CAC6475,PinName="bPrintToScreen",PinToolTip="Print to Screen\nBoolean\n\nWhether or not to print the output to the screen",PinType.PinCategory="bool",PinType.PinSubCategory="",PinType.PinSubCategoryObject=None,PinType.PinSubCategoryMemberReference=(),PinType.PinValueType=(),PinType.ContainerType=None,PinType.bIsMap=False,PinType.bIsSet=False,PinType.bIsArray=False,PinType.bIsReference=False,PinType.bIsConst=False,PinType.bIsWeakPointer=False,DefaultValue="true",AutogeneratedDefaultValue="true",PersistentGuid=00000000000000000000000000000000,bHidden=False,bNotConnectable=False,bDefaultValueIsReadOnly=False,bDefaultValueIsIgnored=False,bAdvancedView=True,bOrphanedPin=False,)
+			   CustomProperties Pin (PinId=4290930B4C818859B03219AA63F9C5D8,PinName="bPrintToLog",PinToolTip="Print to Log\nBoolean\n\nWhether or not to print the output to the log",PinType.PinCategory="bool",PinType.PinSubCategory="",PinType.PinSubCategoryObject=None,PinType.PinSubCategoryMemberReference=(),PinType.PinValueType=(),PinType.ContainerType=None,PinType.bIsMap=False,PinType.bIsSet=False,PinType.bIsArray=False,PinType.bIsReference=False,PinType.bIsConst=False,PinType.bIsWeakPointer=False,DefaultValue="true",AutogeneratedDefaultValue="true",PersistentGuid=00000000000000000000000000000000,bHidden=False,bNotConnectable=False,bDefaultValueIsReadOnly=False,bDefaultValueIsIgnored=False,bAdvancedView=True,bOrphanedPin=False,)
+			   CustomProperties Pin (PinId=92CF3043471A01113A6B4DBD87B135BF,PinName="TextColor",PinToolTip="Text Color\nLinear Color Structure\n\nWhether or not to print the output to the console",PinType.PinCategory="struct",PinType.PinSubCategory="",PinType.PinSubCategoryObject=ScriptStruct'/Script/CoreUObject.LinearColor',PinType.PinSubCategoryMemberReference=(),PinType.PinValueType=(),PinType.ContainerType=None,PinType.bIsMap=False,PinType.bIsSet=False,PinType.bIsArray=False,PinType.bIsReference=False,PinType.bIsConst=False,PinType.bIsWeakPointer=False,DefaultValue="(R=0.000000,G=0.660000,B=1.000000,A=1.000000)",AutogeneratedDefaultValue="(R=0.000000,G=0.660000,B=1.000000,A=1.000000)",PersistentGuid=00000000000000000000000000000000,bHidden=False,bNotConnectable=False,bDefaultValueIsReadOnly=False,bDefaultValueIsIgnored=False,bAdvancedView=True,bOrphanedPin=False,)
+			   CustomProperties Pin (PinId=69983EF24EFFFC63F228069F5FEA7060,PinName="Duration",PinToolTip="Duration\nFloat\n\nThe display duration (if Print to Screen is True). Using negative number will result in loading the duration time from the config.",PinType.PinCategory="float",PinType.PinSubCategory="",PinType.PinSubCategoryObject=None,PinType.PinSubCategoryMemberReference=(),PinType.PinValueType=(),PinType.ContainerType=None,PinType.bIsMap=False,PinType.bIsSet=False,PinType.bIsArray=False,PinType.bIsReference=False,PinType.bIsConst=False,PinType.bIsWeakPointer=False,DefaultValue="2.000000",AutogeneratedDefaultValue="2.000000",PersistentGuid=00000000000000000000000000000000,bHidden=False,bNotConnectable=False,bDefaultValueIsReadOnly=False,bDefaultValueIsIgnored=False,bAdvancedView=True,bOrphanedPin=False,)
+			End Object
+			Begin Object Class=/Script/BlueprintGraph.K2Node_CallFunction Name="K2Node_CallFunction_42"
+			   FunctionReference=(MemberParent=Class'/Script/Engine.KismetSystemLibrary',MemberName="PrintString")
+			   NodePosX=400
+			   NodePosY=784
+			   AdvancedPinDisplay=Hidden
+			   EnabledState=DevelopmentOnly
+			   NodeGuid=E6A1898A4079E3A73463C9AD3C94A669
+			   CustomProperties Pin (PinId=6863525643E71242071D6481F714C829,PinName="execute",PinToolTip="\nExec",PinType.PinCategory="exec",PinType.PinSubCategory="",PinType.PinSubCategoryObject=None,PinType.PinSubCategoryMemberReference=(),PinType.PinValueType=(),PinType.ContainerType=None,PinType.bIsMap=False,PinType.bIsSet=False,PinType.bIsArray=False,PinType.bIsReference=False,PinType.bIsConst=False,PinType.bIsWeakPointer=False,LinkedTo=(K2Node_IfThenElse_120 F1F227324FA2E60EE9CFECAFF857CB77,),PersistentGuid=00000000000000000000000000000000,bHidden=False,bNotConnectable=False,bDefaultValueIsReadOnly=False,bDefaultValueIsIgnored=False,bAdvancedView=False,bOrphanedPin=False,)
+			   CustomProperties Pin (PinId=E13897364CE1B5A8A6D9D89E7CAA376D,PinName="then",PinToolTip="\nExec",Direction="EGPD_Output",PinType.PinCategory="exec",PinType.PinSubCategory="",PinType.PinSubCategoryObject=None,PinType.PinSubCategoryMemberReference=(),PinType.PinValueType=(),PinType.ContainerType=None,PinType.bIsMap=False,PinType.bIsSet=False,PinType.bIsArray=False,PinType.bIsReference=False,PinType.bIsConst=False,PinType.bIsWeakPointer=False,PersistentGuid=00000000000000000000000000000000,bHidden=False,bNotConnectable=False,bDefaultValueIsReadOnly=False,bDefaultValueIsIgnored=False,bAdvancedView=False,bOrphanedPin=False,)
+			   CustomProperties Pin (PinId=5E4EFBED43607C5035952DA4524D24FA,PinName="self",PinFriendlyName=NSLOCTEXT("K2Node", "Target", "Target"),PinToolTip="Target\nKismet System Library Object Reference",PinType.PinCategory="object",PinType.PinSubCategory="",PinType.PinSubCategoryObject=Class'/Script/Engine.KismetSystemLibrary',PinType.PinSubCategoryMemberReference=(),PinType.PinValueType=(),PinType.ContainerType=None,PinType.bIsMap=False,PinType.bIsSet=False,PinType.bIsArray=False,PinType.bIsReference=False,PinType.bIsConst=False,PinType.bIsWeakPointer=False,DefaultObject="/Script/Engine.Default__KismetSystemLibrary",PersistentGuid=00000000000000000000000000000000,bHidden=True,bNotConnectable=False,bDefaultValueIsReadOnly=False,bDefaultValueIsIgnored=False,bAdvancedView=False,bOrphanedPin=False,)
+			   CustomProperties Pin (PinId=AC46C00F409FB31290CB0ABD8B6865D4,PinName="WorldContextObject",PinToolTip="World Context Object\nObject Object Reference",PinType.PinCategory="object",PinType.PinSubCategory="",PinType.PinSubCategoryObject=Class'/Script/CoreUObject.Object',PinType.PinSubCategoryMemberReference=(),PinType.PinValueType=(),PinType.ContainerType=None,PinType.bIsMap=False,PinType.bIsSet=False,PinType.bIsArray=False,PinType.bIsReference=False,PinType.bIsConst=False,PinType.bIsWeakPointer=False,PersistentGuid=00000000000000000000000000000000,bHidden=True,bNotConnectable=False,bDefaultValueIsReadOnly=False,bDefaultValueIsIgnored=False,bAdvancedView=False,bOrphanedPin=False,)
+			   CustomProperties Pin (PinId=8E4D51D74047E6E7A7B8C2B2FEBBC962,PinName="InString",PinToolTip="In String\nString\n\nThe string to log out",PinType.PinCategory="string",PinType.PinSubCategory="",PinType.PinSubCategoryObject=None,PinType.PinSubCategoryMemberReference=(),PinType.PinValueType=(),PinType.ContainerType=None,PinType.bIsMap=False,PinType.bIsSet=False,PinType.bIsArray=False,PinType.bIsReference=False,PinType.bIsConst=False,PinType.bIsWeakPointer=False,DefaultValue="You have returned to the tracking bounds",AutogeneratedDefaultValue="Hello",PersistentGuid=00000000000000000000000000000000,bHidden=False,bNotConnectable=False,bDefaultValueIsReadOnly=False,bDefaultValueIsIgnored=False,bAdvancedView=False,bOrphanedPin=False,)
+			   CustomProperties Pin (PinId=3E2695224F6B062FDED392854CAC6475,PinName="bPrintToScreen",PinToolTip="Print to Screen\nBoolean\n\nWhether or not to print the output to the screen",PinType.PinCategory="bool",PinType.PinSubCategory="",PinType.PinSubCategoryObject=None,PinType.PinSubCategoryMemberReference=(),PinType.PinValueType=(),PinType.ContainerType=None,PinType.bIsMap=False,PinType.bIsSet=False,PinType.bIsArray=False,PinType.bIsReference=False,PinType.bIsConst=False,PinType.bIsWeakPointer=False,DefaultValue="true",AutogeneratedDefaultValue="true",PersistentGuid=00000000000000000000000000000000,bHidden=False,bNotConnectable=False,bDefaultValueIsReadOnly=False,bDefaultValueIsIgnored=False,bAdvancedView=True,bOrphanedPin=False,)
+			   CustomProperties Pin (PinId=4290930B4C818859B03219AA63F9C5D8,PinName="bPrintToLog",PinToolTip="Print to Log\nBoolean\n\nWhether or not to print the output to the log",PinType.PinCategory="bool",PinType.PinSubCategory="",PinType.PinSubCategoryObject=None,PinType.PinSubCategoryMemberReference=(),PinType.PinValueType=(),PinType.ContainerType=None,PinType.bIsMap=False,PinType.bIsSet=False,PinType.bIsArray=False,PinType.bIsReference=False,PinType.bIsConst=False,PinType.bIsWeakPointer=False,DefaultValue="true",AutogeneratedDefaultValue="true",PersistentGuid=00000000000000000000000000000000,bHidden=False,bNotConnectable=False,bDefaultValueIsReadOnly=False,bDefaultValueIsIgnored=False,bAdvancedView=True,bOrphanedPin=False,)
+			   CustomProperties Pin (PinId=92CF3043471A01113A6B4DBD87B135BF,PinName="TextColor",PinToolTip="Text Color\nLinear Color Structure\n\nWhether or not to print the output to the console",PinType.PinCategory="struct",PinType.PinSubCategory="",PinType.PinSubCategoryObject=ScriptStruct'/Script/CoreUObject.LinearColor',PinType.PinSubCategoryMemberReference=(),PinType.PinValueType=(),PinType.ContainerType=None,PinType.bIsMap=False,PinType.bIsSet=False,PinType.bIsArray=False,PinType.bIsReference=False,PinType.bIsConst=False,PinType.bIsWeakPointer=False,DefaultValue="(R=0.000000,G=0.660000,B=1.000000,A=1.000000)",AutogeneratedDefaultValue="(R=0.000000,G=0.660000,B=1.000000,A=1.000000)",PersistentGuid=00000000000000000000000000000000,bHidden=False,bNotConnectable=False,bDefaultValueIsReadOnly=False,bDefaultValueIsIgnored=False,bAdvancedView=True,bOrphanedPin=False,)
+			   CustomProperties Pin (PinId=69983EF24EFFFC63F228069F5FEA7060,PinName="Duration",PinToolTip="Duration\nFloat\n\nThe display duration (if Print to Screen is True). Using negative number will result in loading the duration time from the config.",PinType.PinCategory="float",PinType.PinSubCategory="",PinType.PinSubCategoryObject=None,PinType.PinSubCategoryMemberReference=(),PinType.PinValueType=(),PinType.ContainerType=None,PinType.bIsMap=False,PinType.bIsSet=False,PinType.bIsArray=False,PinType.bIsReference=False,PinType.bIsConst=False,PinType.bIsWeakPointer=False,DefaultValue="2.000000",AutogeneratedDefaultValue="2.000000",PersistentGuid=00000000000000000000000000000000,bHidden=False,bNotConnectable=False,bDefaultValueIsReadOnly=False,bDefaultValueIsIgnored=False,bAdvancedView=True,bOrphanedPin=False,)
+			End Object
+			Begin Object Class=/Script/BlueprintGraph.K2Node_IfThenElse Name="K2Node_IfThenElse_120"
+			   NodePosX=112
+			   NodePosY=832
+			   NodeGuid=6832684C460BBD21410810A679AF1B38
+			   CustomProperties Pin (PinId=2C79AC0046047D0AE4C82EAB908855F1,PinName="execute",PinType.PinCategory="exec",PinType.PinSubCategory="",PinType.PinSubCategoryObject=None,PinType.PinSubCategoryMemberReference=(),PinType.PinValueType=(),PinType.ContainerType=None,PinType.bIsMap=False,PinType.bIsSet=False,PinType.bIsArray=False,PinType.bIsReference=False,PinType.bIsConst=False,PinType.bIsWeakPointer=False,LinkedTo=(K2Node_Event_10 21BE25694FD0A23AEA72C29D0C028719,),PersistentGuid=00000000000000000000000000000000,bHidden=False,bNotConnectable=False,bDefaultValueIsReadOnly=False,bDefaultValueIsIgnored=False,bAdvancedView=False,bOrphanedPin=False,)
+			   CustomProperties Pin (PinId=6C5BC410452F3D5F71607C8AACDD263F,PinName="Condition",PinType.PinCategory="bool",PinType.PinSubCategory="",PinType.PinSubCategoryObject=None,PinType.PinSubCategoryMemberReference=(),PinType.PinValueType=(),PinType.ContainerType=None,PinType.bIsMap=False,PinType.bIsSet=False,PinType.bIsArray=False,PinType.bIsReference=False,PinType.bIsConst=False,PinType.bIsWeakPointer=False,DefaultValue="true",AutogeneratedDefaultValue="true",LinkedTo=(K2Node_VariableGet_307 24E6D8DB483055B1170E9BA0D0084D15,),PersistentGuid=00000000000000000000000000000000,bHidden=False,bNotConnectable=False,bDefaultValueIsReadOnly=False,bDefaultValueIsIgnored=False,bAdvancedView=False,bOrphanedPin=False,)
+			   CustomProperties Pin (PinId=F1F227324FA2E60EE9CFECAFF857CB77,PinName="then",PinFriendlyName=NSLOCTEXT("K2Node", "true", "true"),Direction="EGPD_Output",PinType.PinCategory="exec",PinType.PinSubCategory="",PinType.PinSubCategoryObject=None,PinType.PinSubCategoryMemberReference=(),PinType.PinValueType=(),PinType.ContainerType=None,PinType.bIsMap=False,PinType.bIsSet=False,PinType.bIsArray=False,PinType.bIsReference=False,PinType.bIsConst=False,PinType.bIsWeakPointer=False,LinkedTo=(K2Node_CallFunction_42 6863525643E71242071D6481F714C829,),PersistentGuid=00000000000000000000000000000000,bHidden=False,bNotConnectable=False,bDefaultValueIsReadOnly=False,bDefaultValueIsIgnored=False,bAdvancedView=False,bOrphanedPin=False,)
+			   CustomProperties Pin (PinId=38C905BB46846179190960A620BCC78C,PinName="else",PinFriendlyName=NSLOCTEXT("K2Node", "false", "false"),Direction="EGPD_Output",PinType.PinCategory="exec",PinType.PinSubCategory="",PinType.PinSubCategoryObject=None,PinType.PinSubCategoryMemberReference=(),PinType.PinValueType=(),PinType.ContainerType=None,PinType.bIsMap=False,PinType.bIsSet=False,PinType.bIsArray=False,PinType.bIsReference=False,PinType.bIsConst=False,PinType.bIsWeakPointer=False,LinkedTo=(K2Node_CallFunction_121 6863525643E71242071D6481F714C829,),PersistentGuid=00000000000000000000000000000000,bHidden=False,bNotConnectable=False,bDefaultValueIsReadOnly=False,bDefaultValueIsIgnored=False,bAdvancedView=False,bOrphanedPin=False,)
+			End Object
+			Begin Object Class=/Script/BlueprintGraph.K2Node_VariableGet Name="K2Node_VariableGet_307"
+			   VariableReference=(MemberName="inBounds",MemberGuid=352D7A464AB7AFE78D60A0A0AEB70382,bSelfContext=True)
+			   NodePosX=112
+			   NodePosY=944
+			   NodeGuid=E4719ADD4F41B653FD9E5F99E00E95D2
+			   CustomProperties Pin (PinId=24E6D8DB483055B1170E9BA0D0084D15,PinName="inBounds",Direction="EGPD_Output",PinType.PinCategory="bool",PinType.PinSubCategory="",PinType.PinSubCategoryObject=None,PinType.PinSubCategoryMemberReference=(),PinType.PinValueType=(),PinType.ContainerType=None,PinType.bIsMap=False,PinType.bIsSet=False,PinType.bIsArray=False,PinType.bIsReference=False,PinType.bIsConst=False,PinType.bIsWeakPointer=False,DefaultValue="false",AutogeneratedDefaultValue="false",LinkedTo=(K2Node_IfThenElse_120 6C5BC410452F3D5F71607C8AACDD263F,),PersistentGuid=00000000000000000000000000000000,bHidden=False,bNotConnectable=False,bDefaultValueIsReadOnly=False,bDefaultValueIsIgnored=False,bAdvancedView=False,bOrphanedPin=False,)
+			   CustomProperties Pin (PinId=1635F5694B3DB7FD01A6A78EB0615579,PinName="self",PinFriendlyName=NSLOCTEXT("K2Node", "Target", "Target"),PinType.PinCategory="object",PinType.PinSubCategory="",PinType.PinSubCategoryObject=BlueprintGeneratedClass'/Game/Blueprints/SteamVR_Pawn_Chaperone.SteamVR_Pawn_Chaperone_C',PinType.PinSubCategoryMemberReference=(),PinType.PinValueType=(),PinType.ContainerType=None,PinType.bIsMap=False,PinType.bIsSet=False,PinType.bIsArray=False,PinType.bIsReference=False,PinType.bIsConst=False,PinType.bIsWeakPointer=False,PersistentGuid=00000000000000000000000000000000,bHidden=True,bNotConnectable=False,bDefaultValueIsReadOnly=False,bDefaultValueIsIgnored=False,bAdvancedView=False,bOrphanedPin=False,)
+			End Object
+			Begin Object Class=/Script/BlueprintGraph.K2Node_Event Name="K2Node_Event_10"
+			   EventReference=(MemberParent=Class'/Script/Engine.Actor',MemberName="ReceiveTick")
+			   bOverrideFunction=True
+			   NodePosX=-64
+			   NodePosY=832
+			   NodeGuid=54627A2543CF14C2F2D86DBFA6302CAA
+			   CustomProperties Pin (PinId=F7D16DEE4A9C99A2CF9F1BB769AD2051,PinName="OutputDelegate",Direction="EGPD_Output",PinType.PinCategory="delegate",PinType.PinSubCategory="",PinType.PinSubCategoryObject=None,PinType.PinSubCategoryMemberReference=(MemberParent=Class'/Script/Engine.Actor',MemberName="ReceiveTick"),PinType.PinValueType=(),PinType.ContainerType=None,PinType.bIsMap=False,PinType.bIsSet=False,PinType.bIsArray=False,PinType.bIsReference=False,PinType.bIsConst=False,PinType.bIsWeakPointer=False,PersistentGuid=00000000000000000000000000000000,bHidden=False,bNotConnectable=False,bDefaultValueIsReadOnly=False,bDefaultValueIsIgnored=False,bAdvancedView=False,bOrphanedPin=False,)
+			   CustomProperties Pin (PinId=21BE25694FD0A23AEA72C29D0C028719,PinName="then",Direction="EGPD_Output",PinType.PinCategory="exec",PinType.PinSubCategory="",PinType.PinSubCategoryObject=None,PinType.PinSubCategoryMemberReference=(),PinType.PinValueType=(),PinType.ContainerType=None,PinType.bIsMap=False,PinType.bIsSet=False,PinType.bIsArray=False,PinType.bIsReference=False,PinType.bIsConst=False,PinType.bIsWeakPointer=False,LinkedTo=(K2Node_IfThenElse_120 2C79AC0046047D0AE4C82EAB908855F1,),PersistentGuid=00000000000000000000000000000000,bHidden=False,bNotConnectable=False,bDefaultValueIsReadOnly=False,bDefaultValueIsIgnored=False,bAdvancedView=False,bOrphanedPin=False,)
+			   CustomProperties Pin (PinId=AE33B7B54D11C228B5C903BA4883A601,PinName="DeltaSeconds",PinToolTip="Delta Seconds\nFloat",Direction="EGPD_Output",PinType.PinCategory="float",PinType.PinSubCategory="",PinType.PinSubCategoryObject=None,PinType.PinSubCategoryMemberReference=(),PinType.PinValueType=(),PinType.ContainerType=None,PinType.bIsMap=False,PinType.bIsSet=False,PinType.bIsArray=False,PinType.bIsReference=False,PinType.bIsConst=False,PinType.bIsWeakPointer=False,DefaultValue="0.0",AutogeneratedDefaultValue="0.0",PersistentGuid=00000000000000000000000000000000,bHidden=False,bNotConnectable=False,bDefaultValueIsReadOnly=False,bDefaultValueIsIgnored=False,bAdvancedView=False,bOrphanedPin=False,)
+			End Object
+
+			END OBJECT
+	[/PARAM]
+[/OBJECT]
+[/EXCERPT:SVRFinalBP]
+-->
\ No newline at end of file
