Copying //UE4/Dev-Documentation to //UE4/Samples-Main (Source: //UE4/Dev-Documentation @ 3623665)

diff --git a/Engine/Documentation/Source/Platforms/SteamVR/HowTo/ChaperoneInteraction/SteamVRChaperoneInteraction.INT.udn b/Engine/Documentation/Source/Platforms/SteamVR/HowTo/ChaperoneInteraction/SteamVRChaperoneInteraction.INT.udn
new file mode 100644
--- /dev/null
+++ b/Engine/Documentation/Source/Platforms/SteamVR/HowTo/ChaperoneInteraction/SteamVRChaperoneInteraction.INT.udn
@@ -0,0 +1,550 @@
+Availability: Public
+Crumbs: %ROOT%, Platforms
+Title: Interact with the SteamVR Chaperone System 
+Description: Interact with the various components of the SteamVR Chaperone system.
+Type: how-to
+Skilllevel: Intermidate
+Version: 4.17
+Parent: Platforms/SteamVR
+Order: 2
+tags: VR
+topic-image:Platforms\SteamVR\SteamVRDevelopment_topic.png
+
+The SteamVR Chaperone system is used to display the boundaries of the Virtual Reality (VR) interaction area. As a tracked device approaches the boundaries, the SteamVR Runtime automatically provides visual cues notifying the user of this. In the following How - To we will take a look at two different ways to display the VR interaction area to the user. 
+
+[REGION:note]
+The following How - To is going to make use of the Standing VR Pawn with added Motion Controller and Chaperone support. If you are not sure how to set this up, the following two How - To's will help you out.
+* [Set Up Standing Cameras for SteamVR](Platforms\SteamVR\HowTo\StandingCamera)
+* [Set Up Motion Controllers](Platforms\VR\MotionController)
+* [Set Up the Chaperone System](Platforms\SteamVR\HowTo\Chaperone)
+[/REGION]
+
+[REGION:tip]
+For the Chaperone system to work, you will need to make sure you have set up Room Scale VR using the SteamVR tools. For more information on how to do this, please see the official [HTC Vive](https://www.vive.com/us/support/category_howto/setting-up-room-scale-play-area.html) Setup pages.
+[/REGION]
+
+[REGION:warning]
+It is **NOT** possible nor advisable to disable the Chaperone system from inside UE4. You can, however, adjust how UE4 will react to the user getting close to the boundaries.  
+[/REGION]
+
+## Display the Corners of the SteamVR Interaction Area
+
+In this part of the How - To we will take a look at spawning static meshes at the four corners of the VR interaction area.
+
+[region:note]
+You can download the Lighthouse Basestation Static Mesh and Texture used in this How - To from the following link. [Lighthouse Basestation Static Mesh & Texture](https://epicgames.box.com/s/3n6yilg2is2f7nq8ju27z0wbxhkvyms1)
+[/region]
+
+1. First, open up your Player Pawn Blueprint and add a new **Static Mesh Component** variable called **baseStationArray** . Since this variable will be used to store the Static Mesh we want to spawn at the four corners of the VR Interaction area when creating the baseStationArray variable, make sure that the **Variable Type** is set to **Array**.
+
+    ![](SVRChapInteraction_00.png)
+
+1. Next, add the following nodes with the following values to the **Event Graph**:
+
+    [REGION:lightbox]
+    [![](SVRChapInteraction_01.png)(w:748)](SVRChapInteraction_01.png)
+    [/REGION]
+
+    [REGION:caption]
+    Click for full image.
+    [/REGION]
+
+    |Node Name| Value |
+    |--------|------------|
+    |**Motion Controller (R) Trigger**|N/A|
+    |**SteamVRChaperone**|N/A|
+    |**baseStationArray**|N/A|
+    |**Get Bounds**|N/A|
+    |**ForEachLoop**|N/A|
+    |**Add Static Mesh Component**|Static Mesh: lh_basestation_vive |
+    |**Destroy Component**|N/A|
+    |**Array Add **|N/A|
+    |**IsValid**|N/A|
+
+1. Next, we need to setup a way to get information about the SteamVR interaction area that the user setup. To accomplish inside your Pawn Blueprint connect the **SteamVRChaperone** variable output to the input on the **Get Bounds** node as this will allow us to interact with points that make up the SteamVR Chaperone bounds.
+
+    [REGION:lightbox]
+    [![](SVRChapInteraction_02.png)(w:500)](SVRChapInteraction_02.png)
+    [/REGION]
+
+    [REGION:caption]
+    Click for full image.
+    [/REGION]
+
+1. Now that we can get points that make up the VR interaction area, next we area are going to setup up a ForEachLoop that will go through each corner point in the SteamVRChaperone and add a Static Mesh when the **Motion Controller(R) Trigger** is pressed. To accomplish this in your Pawn Blueprint make sure your Event Graph setup matches the following image.
+
+    [INCLUDE:Platforms\SteamVR\HowTo\ChaperoneInteraction\#SteamVRC1]
+
+1. With the ability to display a Static Mesh at the four corners of the VR interaction area, next we will setup some Blueprint logic that will check and see if we have a valid Static Mesh to remove and if we do, remove it. To accomplish this in your Pawn Blueprint make sure your Event Graph setup matches the following image. 
+
+    [INCLUDE:Platforms\SteamVR\HowTo\ChaperoneInteraction\#SteamVRC2]
+
+With the ability to display the corners of the VR interaction area using a Static Mesh, in the next part of this How - To we will take a look at how to use the **Minimum Area Rectangle** node to draw a series of lines that are the exact size and shape of the VR interaction.
+
+## Display the Entire SteamVR Interaction Area
+
+In this part of the How - To we will take a look at drawing a box that is the exact size and shape of the VR interaction area.
+
+1. First over in the **Variables** section create a new **Vector** variable called, **SampleSurface**, that will be used to store the surface normals of the VR interaction area. 
+
+    ![](SVRChapInteraction_05.png)
+
+1. Next, add the following nodes with the following values to the **Event Graph**:
+
+    [REGION:lightbox]
+    [![](SVRChapInteraction_06.png)(w:748)](SVRChapInteraction_06.png)
+    [/REGION]
+
+    [REGION:caption]
+    Click for full image.
+    [/REGION]
+
+    |Node Name| Value |
+    |--------|------------|
+    |**Motion Controller (R) Trigger**|N/A|
+    |**SteamVRChaperone**|N/A|
+    |**GetBounds**|N/A|
+    |**Float / Float x 2**|0,2|
+    |**SampleSurface**|N/A|
+    |**MinimumAreaRectangle**|N/A|
+    |**DrawDebugBox**|N/A|
+
+1. With the needed nodes now added to the Event Graph, next, we need to connect them so that when the Motion Controller(R) Trigger is pressed we can sample the position of the vertices that make up the VR interaction area and then draw a box around them. Setting up the nodes to match the following image will allow you to achieve this in your Blueprint.
+
+    [INCLUDE:Platforms\SteamVR\HowTo\ChaperoneInteraction\#SteamVRC3]
+
+## End Results
+
+With everything now hooked up launch the level using the VR Preview button and put on your HTC Vive HMD and pick up your Motion Controller. Now when you squeeze the trigger on the right-hand Motion Controller the bounds of the SteamVR interaction area will be displayed like in the following video. 
+
+[OBJECT:EmbeddedVideo]
+    [PARAMLITERAL:videoid]
+    xz8bAEQEmf0
+    [/PARAMLITERAL]
+    [PARAMLITERAL:width]
+    100
+    [/PARAMLITERAL]
+    [PARAMLITERAL:height]
+    56.25
+    [/PARAMLITERAL]
+    [PARAMLITERAL:units]
+    %
+    [/PARAMLITERAL]
+    [PARAMLITERAL:type]
+    simple
+    [/PARAMLITERAL]
+[/OBJECT]
+
+## UE4 Project Downloads
+
+Below you will find a link to where you can download the UE4 project that was used to create this example.
+
+* [**SteamVR Chaperon Interaction Example Project**](https://epicgames.box.com/s/q0j0i9agz94z0fh8jjz3tspipdbrs2ud)
+
+
+<!--
+[EXCERPT:SteamVRC1]	
+[OBJECT:Snippet]
+	[PARAMLITERAL:id]
+		SteamVRC1
+	[/PARAMLITERAL]
+	[PARAMLITERAL:type]
+		blueprint
+	[/PARAMLITERAL]
+	[PARAM:image]
+		[REGION:lightbox]
+		[![](SVRChapInteraction_03.png)(w:1000)](SVRChapInteraction_03.png)
+		[/REGION]
+	[/PARAM]
+	[PARAM:snippet]
+			BEGIN OBJECT
+							Begin Object Class=/Script/BlueprintGraph.K2Node_CallFunction Name="K2Node_CallFunction_1084"
+			   bIsPureFunc=True
+			   bIsConstFunc=True
+			   FunctionReference=(MemberParent=Class'/Script/SteamVR.SteamVRChaperoneComponent',MemberName="GetBounds")
+			   NodePosX=1232
+			   NodePosY=-2624
+			   NodeGuid=9D8CE6074BBE8BCF7AA85D8BB3B15185
+			   CustomProperties Pin (PinId=8B9136B54BDA0D37E448BDAC16DFFFF7,PinName="self",PinFriendlyName=NSLOCTEXT("K2Node", "Target", "Target"),PinToolTip="Target\nSteam VRChaperone Component Object Reference",PinType.PinCategory="object",PinType.PinSubCategory="",PinType.PinSubCategoryObject=Class'/Script/SteamVR.SteamVRChaperoneComponent',PinType.PinSubCategoryMemberReference=(),PinType.PinValueType=(),PinType.ContainerType=None,PinType.bIsMap=False,PinType.bIsSet=False,PinType.bIsArray=False,PinType.bIsReference=False,PinType.bIsConst=False,PinType.bIsWeakPointer=False,LinkedTo=(K2Node_VariableGet_197 4857BDE84AFE450BEB4D218869646F9B,),PersistentGuid=00000000000000000000000000000000,bHidden=False,bNotConnectable=False,bDefaultValueIsReadOnly=False,bDefaultValueIsIgnored=False,bAdvancedView=False,bOrphanedPin=False,)
+			   CustomProperties Pin (PinId=B35FC4BA4D68C6738438C18702E70E4C,PinName="ReturnValue",PinToolTip="Return Value\nArray of Vector s",Direction="EGPD_Output",PinType.PinCategory="struct",PinType.PinSubCategory="",PinType.PinSubCategoryObject=ScriptStruct'/Script/CoreUObject.Vector',PinType.PinSubCategoryMemberReference=(),PinType.PinValueType=(),PinType.ContainerType=Array,PinType.bIsMap=False,PinType.bIsSet=False,PinType.bIsArray=False,PinType.bIsReference=False,PinType.bIsConst=False,PinType.bIsWeakPointer=False,LinkedTo=(K2Node_MacroInstance_118 B6F62194447A3A47B7677494C9E6AF42,),PersistentGuid=00000000000000000000000000000000,bHidden=False,bNotConnectable=False,bDefaultValueIsReadOnly=False,bDefaultValueIsIgnored=False,bAdvancedView=False,bOrphanedPin=False,)
+			End Object
+			Begin Object Class=/Script/BlueprintGraph.K2Node_VariableGet Name="K2Node_VariableGet_197"
+			   VariableReference=(MemberName="SteamVRChaperone",bSelfContext=True)
+			   NodePosX=1248
+			   NodePosY=-2544
+			   NodeGuid=7E57A3734E176EA15785F397B1DC76FB
+			   CustomProperties Pin (PinId=4857BDE84AFE450BEB4D218869646F9B,PinName="SteamVRChaperone",Direction="EGPD_Output",PinType.PinCategory="object",PinType.PinSubCategory="",PinType.PinSubCategoryObject=Class'/Script/SteamVR.SteamVRChaperoneComponent',PinType.PinSubCategoryMemberReference=(),PinType.PinValueType=(),PinType.ContainerType=None,PinType.bIsMap=False,PinType.bIsSet=False,PinType.bIsArray=False,PinType.bIsReference=False,PinType.bIsConst=False,PinType.bIsWeakPointer=False,LinkedTo=(K2Node_CallFunction_1084 8B9136B54BDA0D37E448BDAC16DFFFF7,),PersistentGuid=00000000000000000000000000000000,bHidden=False,bNotConnectable=False,bDefaultValueIsReadOnly=False,bDefaultValueIsIgnored=False,bAdvancedView=False,bOrphanedPin=False,)
+			   CustomProperties Pin (PinId=1E59DFAA490E313860D944A45C6A61A1,PinName="self",PinFriendlyName=NSLOCTEXT("K2Node", "Target", "Target"),PinType.PinCategory="object",PinType.PinSubCategory="",PinType.PinSubCategoryObject=BlueprintGeneratedClass'/Game/Blueprints/SteamVR_Pawn_Standing.SteamVR_Pawn_Standing_C',PinType.PinSubCategoryMemberReference=(),PinType.PinValueType=(),PinType.ContainerType=None,PinType.bIsMap=False,PinType.bIsSet=False,PinType.bIsArray=False,PinType.bIsReference=False,PinType.bIsConst=False,PinType.bIsWeakPointer=False,PersistentGuid=00000000000000000000000000000000,bHidden=True,bNotConnectable=False,bDefaultValueIsReadOnly=False,bDefaultValueIsIgnored=False,bAdvancedView=False,bOrphanedPin=False,)
+			End Object
+			Begin Object Class=/Script/BlueprintGraph.K2Node_MacroInstance Name="K2Node_MacroInstance_118"
+			   MacroGraphReference=(MacroGraph=EdGraph'/Engine/EditorBlueprintResources/StandardMacros.StandardMacros:ForEachLoop',GraphBlueprint=Blueprint'/Engine/EditorBlueprintResources/StandardMacros.StandardMacros',GraphGuid=99DBFD5540A796041F72A5A9DA655026)
+			   ResolvedWildcardType=(PinCategory="struct",PinSubCategoryObject=ScriptStruct'/Script/CoreUObject.Vector',ContainerType=Array)
+			   NodePosX=1584
+			   NodePosY=-2768
+			   NodeGuid=37FA01EA488A9C7B9FFEA3AFCA4B5D7B
+			   CustomProperties Pin (PinId=2BDC94F34208A8F557E05F94CE1CE4AE,PinName="Exec",PinType.PinCategory="exec",PinType.PinSubCategory="",PinType.PinSubCategoryObject=None,PinType.PinSubCategoryMemberReference=(),PinType.PinValueType=(),PinType.ContainerType=None,PinType.bIsMap=False,PinType.bIsSet=False,PinType.bIsArray=False,PinType.bIsReference=False,PinType.bIsConst=False,PinType.bIsWeakPointer=False,LinkedTo=(K2Node_InputKey_74 451F705A45E7D71FD4CDD191C52193E1,),PersistentGuid=00000000000000000000000000000000,bHidden=False,bNotConnectable=False,bDefaultValueIsReadOnly=False,bDefaultValueIsIgnored=False,bAdvancedView=False,bOrphanedPin=False,)
+			   CustomProperties Pin (PinId=B6F62194447A3A47B7677494C9E6AF42,PinName="Array",PinType.PinCategory="struct",PinType.PinSubCategory="",PinType.PinSubCategoryObject=ScriptStruct'/Script/CoreUObject.Vector',PinType.PinSubCategoryMemberReference=(),PinType.PinValueType=(),PinType.ContainerType=Array,PinType.bIsMap=False,PinType.bIsSet=False,PinType.bIsArray=False,PinType.bIsReference=False,PinType.bIsConst=False,PinType.bIsWeakPointer=False,LinkedTo=(K2Node_CallFunction_1084 B35FC4BA4D68C6738438C18702E70E4C,),PersistentGuid=00000000000000000000000000000000,bHidden=False,bNotConnectable=False,bDefaultValueIsReadOnly=False,bDefaultValueIsIgnored=False,bAdvancedView=False,bOrphanedPin=False,)
+			   CustomProperties Pin (PinId=46A798794EB954364926C08AABB07BFF,PinName="LoopBody",Direction="EGPD_Output",PinType.PinCategory="exec",PinType.PinSubCategory="",PinType.PinSubCategoryObject=None,PinType.PinSubCategoryMemberReference=(),PinType.PinValueType=(),PinType.ContainerType=None,PinType.bIsMap=False,PinType.bIsSet=False,PinType.bIsArray=False,PinType.bIsReference=False,PinType.bIsConst=False,PinType.bIsWeakPointer=False,LinkedTo=(K2Node_AddComponent_71 34310B40490B2BAD8DED8ABB2127E66A,),PersistentGuid=00000000000000000000000000000000,bHidden=False,bNotConnectable=False,bDefaultValueIsReadOnly=False,bDefaultValueIsIgnored=False,bAdvancedView=False,bOrphanedPin=False,)
+			   CustomProperties Pin (PinId=E0271977466AEE7454D6A0AD5AABB816,PinName="Array Element",Direction="EGPD_Output",PinType.PinCategory="struct",PinType.PinSubCategory="",PinType.PinSubCategoryObject=ScriptStruct'/Script/CoreUObject.Vector',PinType.PinSubCategoryMemberReference=(),PinType.PinValueType=(),PinType.ContainerType=None,PinType.bIsMap=False,PinType.bIsSet=False,PinType.bIsArray=False,PinType.bIsReference=False,PinType.bIsConst=False,PinType.bIsWeakPointer=False,LinkedTo=(K2Node_CallFunction_3281 72212A6A4BCA13036E0B57BFD90B270E,),PersistentGuid=00000000000000000000000000000000,bHidden=False,bNotConnectable=False,bDefaultValueIsReadOnly=False,bDefaultValueIsIgnored=False,bAdvancedView=False,bOrphanedPin=False,)
+			   CustomProperties Pin (PinId=48EF763F42B4DA5034F2F58F9BBEB00D,PinName="Array Index",Direction="EGPD_Output",PinType.PinCategory="int",PinType.PinSubCategory="",PinType.PinSubCategoryObject=None,PinType.PinSubCategoryMemberReference=(),PinType.PinValueType=(),PinType.ContainerType=None,PinType.bIsMap=False,PinType.bIsSet=False,PinType.bIsArray=False,PinType.bIsReference=False,PinType.bIsConst=False,PinType.bIsWeakPointer=False,DefaultValue="0",AutogeneratedDefaultValue="0",PersistentGuid=00000000000000000000000000000000,bHidden=False,bNotConnectable=False,bDefaultValueIsReadOnly=False,bDefaultValueIsIgnored=False,bAdvancedView=False,bOrphanedPin=False,)
+			   CustomProperties Pin (PinId=7DA82531415B398F5EB95691B0938909,PinName="Completed",Direction="EGPD_Output",PinType.PinCategory="exec",PinType.PinSubCategory="",PinType.PinSubCategoryObject=None,PinType.PinSubCategoryMemberReference=(),PinType.PinValueType=(),PinType.ContainerType=None,PinType.bIsMap=False,PinType.bIsSet=False,PinType.bIsArray=False,PinType.bIsReference=False,PinType.bIsConst=False,PinType.bIsWeakPointer=False,PersistentGuid=00000000000000000000000000000000,bHidden=False,bNotConnectable=False,bDefaultValueIsReadOnly=False,bDefaultValueIsIgnored=False,bAdvancedView=False,bOrphanedPin=False,)
+			End Object
+			Begin Object Class=/Script/BlueprintGraph.K2Node_AddComponent Name="K2Node_AddComponent_71"
+			   TemplateBlueprint="/Game/Blueprints/SteamVR_Pawn_Standing.SteamVR_Pawn_Standing"
+			   TemplateType=Class'/Script/Engine.StaticMeshComponent'
+			   FunctionReference=(MemberName="AddComponent",bSelfContext=True)
+			   NodePosX=1920
+			   NodePosY=-2784
+			   NodeGuid=D899A1934C7F4236190FBB83E80E36E4
+			   CustomProperties Pin (PinId=34310B40490B2BAD8DED8ABB2127E66A,PinName="execute",PinToolTip="\nExec",PinType.PinCategory="exec",PinType.PinSubCategory="",PinType.PinSubCategoryObject=None,PinType.PinSubCategoryMemberReference=(),PinType.PinValueType=(),PinType.ContainerType=None,PinType.bIsMap=False,PinType.bIsSet=False,PinType.bIsArray=False,PinType.bIsReference=False,PinType.bIsConst=False,PinType.bIsWeakPointer=False,LinkedTo=(K2Node_MacroInstance_118 46A798794EB954364926C08AABB07BFF,),PersistentGuid=00000000000000000000000000000000,bHidden=False,bNotConnectable=False,bDefaultValueIsReadOnly=False,bDefaultValueIsIgnored=False,bAdvancedView=False,bOrphanedPin=False,)
+			   CustomProperties Pin (PinId=497F5C1D48B8B3DB2C2BA3AFBDDDA07D,PinName="then",PinToolTip="\nExec",Direction="EGPD_Output",PinType.PinCategory="exec",PinType.PinSubCategory="",PinType.PinSubCategoryObject=None,PinType.PinSubCategoryMemberReference=(),PinType.PinValueType=(),PinType.ContainerType=None,PinType.bIsMap=False,PinType.bIsSet=False,PinType.bIsArray=False,PinType.bIsReference=False,PinType.bIsConst=False,PinType.bIsWeakPointer=False,LinkedTo=(K2Node_CallArrayFunction_227 AB7C29724A4370E374E219B3136B3B14,),PersistentGuid=00000000000000000000000000000000,bHidden=False,bNotConnectable=False,bDefaultValueIsReadOnly=False,bDefaultValueIsIgnored=False,bAdvancedView=False,bOrphanedPin=False,)
+			   CustomProperties Pin (PinId=0942B0BA46BC8263AD4B8DB55CDECCEF,PinName="self",PinFriendlyName=NSLOCTEXT("K2Node", "Target", "Target"),PinToolTip="Target\nActor Object Reference",PinType.PinCategory="object",PinType.PinSubCategory="",PinType.PinSubCategoryObject=Class'/Script/Engine.Actor',PinType.PinSubCategoryMemberReference=(),PinType.PinValueType=(),PinType.ContainerType=None,PinType.bIsMap=False,PinType.bIsSet=False,PinType.bIsArray=False,PinType.bIsReference=False,PinType.bIsConst=False,PinType.bIsWeakPointer=False,PersistentGuid=00000000000000000000000000000000,bHidden=False,bNotConnectable=False,bDefaultValueIsReadOnly=False,bDefaultValueIsIgnored=False,bAdvancedView=False,bOrphanedPin=False,)
+			   CustomProperties Pin (PinId=21E468274D3DBA5F04E63E9919D1B6EB,PinName="TemplateName",PinToolTip="Template Name\nName\n\nThe name of the Component Template to use.",PinType.PinCategory="name",PinType.PinSubCategory="",PinType.PinSubCategoryObject=None,PinType.PinSubCategoryMemberReference=(),PinType.PinValueType=(),PinType.ContainerType=None,PinType.bIsMap=False,PinType.bIsSet=False,PinType.bIsArray=False,PinType.bIsReference=False,PinType.bIsConst=False,PinType.bIsWeakPointer=False,DefaultValue="NODE_AddStaticMeshComponent-5",AutogeneratedDefaultValue="None",PersistentGuid=00000000000000000000000000000000,bHidden=True,bNotConnectable=True,bDefaultValueIsReadOnly=True,bDefaultValueIsIgnored=False,bAdvancedView=False,bOrphanedPin=False,)
+			   CustomProperties Pin (PinId=948A308C4E60D0B8C5EC1D92028E558A,PinName="bManualAttachment",PinToolTip="Manual Attachment\nBoolean\n\nWhether manual or automatic attachment is to be used",PinType.PinCategory="bool",PinType.PinSubCategory="",PinType.PinSubCategoryObject=None,PinType.PinSubCategoryMemberReference=(),PinType.PinValueType=(),PinType.ContainerType=None,PinType.bIsMap=False,PinType.bIsSet=False,PinType.bIsArray=False,PinType.bIsReference=False,PinType.bIsConst=False,PinType.bIsWeakPointer=False,DefaultValue="false",AutogeneratedDefaultValue="false",PersistentGuid=00000000000000000000000000000000,bHidden=False,bNotConnectable=False,bDefaultValueIsReadOnly=False,bDefaultValueIsIgnored=False,bAdvancedView=False,bOrphanedPin=False,)
+			   CustomProperties Pin (PinId=5C74A8634DDD697D1234B38F487C2D6D,PinName="RelativeTransform",PinToolTip="Relative Transform\nTransform \n\nThe relative transform between the new component and its attach parent (automatic only)",PinType.PinCategory="struct",PinType.PinSubCategory="",PinType.PinSubCategoryObject=ScriptStruct'/Script/CoreUObject.Transform',PinType.PinSubCategoryMemberReference=(),PinType.PinValueType=(),PinType.ContainerType=None,PinType.bIsMap=False,PinType.bIsSet=False,PinType.bIsArray=False,PinType.bIsReference=False,PinType.bIsConst=True,PinType.bIsWeakPointer=False,LinkedTo=(K2Node_CallFunction_3281 301B61884B3507108C1051B08AF72EC0,),PersistentGuid=00000000000000000000000000000000,bHidden=False,bNotConnectable=False,bDefaultValueIsReadOnly=False,bDefaultValueIsIgnored=True,bAdvancedView=False,bOrphanedPin=False,)
+			   CustomProperties Pin (PinId=A79F33AE4CF14547CB740F9D0B5ECB67,PinName="ComponentTemplateContext",PinToolTip="Component Template Context\nObject Object Reference\n\nOptional UBlueprintGeneratedClass reference to use to find the template in. If null (or not a BPGC), component is sought in this Actor\'s class",PinType.PinCategory="object",PinType.PinSubCategory="",PinType.PinSubCategoryObject=Class'/Script/CoreUObject.Object',PinType.PinSubCategoryMemberReference=(),PinType.PinValueType=(),PinType.ContainerType=None,PinType.bIsMap=False,PinType.bIsSet=False,PinType.bIsArray=False,PinType.bIsReference=False,PinType.bIsConst=True,PinType.bIsWeakPointer=False,PersistentGuid=00000000000000000000000000000000,bHidden=True,bNotConnectable=True,bDefaultValueIsReadOnly=False,bDefaultValueIsIgnored=False,bAdvancedView=False,bOrphanedPin=False,)
+			   CustomProperties Pin (PinId=6F754DCB4AA8982E55C33D8BEA044333,PinName="ReturnValue",PinToolTip="Return Value\nStatic Mesh Component Object Reference",Direction="EGPD_Output",PinType.PinCategory="object",PinType.PinSubCategory="",PinType.PinSubCategoryObject=Class'/Script/Engine.StaticMeshComponent',PinType.PinSubCategoryMemberReference=(),PinType.PinValueType=(),PinType.ContainerType=None,PinType.bIsMap=False,PinType.bIsSet=False,PinType.bIsArray=False,PinType.bIsReference=False,PinType.bIsConst=False,PinType.bIsWeakPointer=False,LinkedTo=(K2Node_CallArrayFunction_227 EE5E261346A9D948350D08A538206500,),PersistentGuid=00000000000000000000000000000000,bHidden=False,bNotConnectable=False,bDefaultValueIsReadOnly=False,bDefaultValueIsIgnored=False,bAdvancedView=False,bOrphanedPin=False,)
+			End Object
+			Begin Object Class=/Script/BlueprintGraph.K2Node_InputKey Name="K2Node_InputKey_74"
+			   InputKey=MotionController_Right_Trigger
+			   NodePosX=1232
+			   NodePosY=-2768
+			   NodeGuid=115D70DD4D72911C84CDD396206876CB
+			   CustomProperties Pin (PinId=451F705A45E7D71FD4CDD191C52193E1,PinName="Pressed",Direction="EGPD_Output",PinType.PinCategory="exec",PinType.PinSubCategory="",PinType.PinSubCategoryObject=None,PinType.PinSubCategoryMemberReference=(),PinType.PinValueType=(),PinType.ContainerType=None,PinType.bIsMap=False,PinType.bIsSet=False,PinType.bIsArray=False,PinType.bIsReference=False,PinType.bIsConst=False,PinType.bIsWeakPointer=False,LinkedTo=(K2Node_MacroInstance_118 2BDC94F34208A8F557E05F94CE1CE4AE,),PersistentGuid=00000000000000000000000000000000,bHidden=False,bNotConnectable=False,bDefaultValueIsReadOnly=False,bDefaultValueIsIgnored=False,bAdvancedView=False,bOrphanedPin=False,)
+			   CustomProperties Pin (PinId=6F1385DB4EF7010345F655BD3744E632,PinName="Released",Direction="EGPD_Output",PinType.PinCategory="exec",PinType.PinSubCategory="",PinType.PinSubCategoryObject=None,PinType.PinSubCategoryMemberReference=(),PinType.PinValueType=(),PinType.ContainerType=None,PinType.bIsMap=False,PinType.bIsSet=False,PinType.bIsArray=False,PinType.bIsReference=False,PinType.bIsConst=False,PinType.bIsWeakPointer=False,PersistentGuid=00000000000000000000000000000000,bHidden=False,bNotConnectable=False,bDefaultValueIsReadOnly=False,bDefaultValueIsIgnored=False,bAdvancedView=False,bOrphanedPin=False,)
+			   CustomProperties Pin (PinId=022C7579454B610F2F2ED0B282488461,PinName="Key",Direction="EGPD_Output",PinType.PinCategory="struct",PinType.PinSubCategory="",PinType.PinSubCategoryObject=ScriptStruct'/Script/InputCore.Key',PinType.PinSubCategoryMemberReference=(),PinType.PinValueType=(),PinType.ContainerType=None,PinType.bIsMap=False,PinType.bIsSet=False,PinType.bIsArray=False,PinType.bIsReference=False,PinType.bIsConst=False,PinType.bIsWeakPointer=False,PersistentGuid=00000000000000000000000000000000,bHidden=False,bNotConnectable=False,bDefaultValueIsReadOnly=False,bDefaultValueIsIgnored=False,bAdvancedView=False,bOrphanedPin=False,)
+			End Object
+			Begin Object Class=/Script/BlueprintGraph.K2Node_CallArrayFunction Name="K2Node_CallArrayFunction_227"
+			   FunctionReference=(MemberParent=Class'/Script/Engine.KismetArrayLibrary',MemberName="Array_Add")
+			   NodePosX=2256
+			   NodePosY=-2744
+			   NodeGuid=43700AEC407CB6EBCFD0BBBDB6E7566D
+			   CustomProperties Pin (PinId=AB7C29724A4370E374E219B3136B3B14,PinName="execute",PinToolTip="\nExec",PinType.PinCategory="exec",PinType.PinSubCategory="",PinType.PinSubCategoryObject=None,PinType.PinSubCategoryMemberReference=(),PinType.PinValueType=(),PinType.ContainerType=None,PinType.bIsMap=False,PinType.bIsSet=False,PinType.bIsArray=False,PinType.bIsReference=False,PinType.bIsConst=False,PinType.bIsWeakPointer=False,LinkedTo=(K2Node_AddComponent_71 497F5C1D48B8B3DB2C2BA3AFBDDDA07D,),PersistentGuid=00000000000000000000000000000000,bHidden=False,bNotConnectable=False,bDefaultValueIsReadOnly=False,bDefaultValueIsIgnored=False,bAdvancedView=False,bOrphanedPin=False,)
+			   CustomProperties Pin (PinId=9C49631A4DCB38432F500787BD709DCB,PinName="then",PinToolTip="\nExec",Direction="EGPD_Output",PinType.PinCategory="exec",PinType.PinSubCategory="",PinType.PinSubCategoryObject=None,PinType.PinSubCategoryMemberReference=(),PinType.PinValueType=(),PinType.ContainerType=None,PinType.bIsMap=False,PinType.bIsSet=False,PinType.bIsArray=False,PinType.bIsReference=False,PinType.bIsConst=False,PinType.bIsWeakPointer=False,PersistentGuid=00000000000000000000000000000000,bHidden=False,bNotConnectable=False,bDefaultValueIsReadOnly=False,bDefaultValueIsIgnored=False,bAdvancedView=False,bOrphanedPin=False,)
+			   CustomProperties Pin (PinId=00FB150349E4D66D6F7FE7983CE94DBA,PinName="self",PinFriendlyName=NSLOCTEXT("K2Node", "Target", "Target"),PinToolTip="Target\nKismet Array Library Object Reference",PinType.PinCategory="object",PinType.PinSubCategory="",PinType.PinSubCategoryObject=Class'/Script/Engine.KismetArrayLibrary',PinType.PinSubCategoryMemberReference=(),PinType.PinValueType=(),PinType.ContainerType=None,PinType.bIsMap=False,PinType.bIsSet=False,PinType.bIsArray=False,PinType.bIsReference=False,PinType.bIsConst=False,PinType.bIsWeakPointer=False,DefaultObject="/Script/Engine.Default__KismetArrayLibrary",PersistentGuid=00000000000000000000000000000000,bHidden=True,bNotConnectable=False,bDefaultValueIsReadOnly=False,bDefaultValueIsIgnored=False,bAdvancedView=False,bOrphanedPin=False,)
+			   CustomProperties Pin (PinId=22674C5B4C7D8F21FD49C8A2A4EFF177,PinName="TargetArray",PinToolTip="Target Array\nArray of Static Mesh Component Object References\n\nThe array to add item to",PinType.PinCategory="object",PinType.PinSubCategory="",PinType.PinSubCategoryObject=Class'/Script/Engine.StaticMeshComponent',PinType.PinSubCategoryMemberReference=(),PinType.PinValueType=(),PinType.ContainerType=Array,PinType.bIsMap=False,PinType.bIsSet=False,PinType.bIsArray=False,PinType.bIsReference=True,PinType.bIsConst=True,PinType.bIsWeakPointer=False,LinkedTo=(K2Node_VariableGet_199 92B8D86246B7B64BCB79D093554F79E8,),PersistentGuid=00000000000000000000000000000000,bHidden=False,bNotConnectable=False,bDefaultValueIsReadOnly=False,bDefaultValueIsIgnored=True,bAdvancedView=False,bOrphanedPin=False,)
+			   CustomProperties Pin (PinId=EE5E261346A9D948350D08A538206500,PinName="NewItem",PinToolTip="New Item\nStatic Mesh Component Object Reference (by ref)\n\nThe item to add to the array",PinType.PinCategory="object",PinType.PinSubCategory="",PinType.PinSubCategoryObject=Class'/Script/Engine.StaticMeshComponent',PinType.PinSubCategoryMemberReference=(),PinType.PinValueType=(),PinType.ContainerType=None,PinType.bIsMap=False,PinType.bIsSet=False,PinType.bIsArray=False,PinType.bIsReference=True,PinType.bIsConst=True,PinType.bIsWeakPointer=False,LinkedTo=(K2Node_AddComponent_71 6F754DCB4AA8982E55C33D8BEA044333,),PersistentGuid=00000000000000000000000000000000,bHidden=False,bNotConnectable=False,bDefaultValueIsReadOnly=False,bDefaultValueIsIgnored=False,bAdvancedView=False,bOrphanedPin=False,)
+			   CustomProperties Pin (PinId=BE1A5C404E2C69FB80B60D8497412508,PinName="ReturnValue",PinToolTip="Return Value\nInteger\n\nThe index of the newly added item",Direction="EGPD_Output",PinType.PinCategory="int",PinType.PinSubCategory="",PinType.PinSubCategoryObject=None,PinType.PinSubCategoryMemberReference=(),PinType.PinValueType=(),PinType.ContainerType=None,PinType.bIsMap=False,PinType.bIsSet=False,PinType.bIsArray=False,PinType.bIsReference=False,PinType.bIsConst=False,PinType.bIsWeakPointer=False,DefaultValue="0",AutogeneratedDefaultValue="0",PersistentGuid=00000000000000000000000000000000,bHidden=False,bNotConnectable=False,bDefaultValueIsReadOnly=False,bDefaultValueIsIgnored=False,bAdvancedView=False,bOrphanedPin=False,)
+			End Object
+			Begin Object Class=/Script/BlueprintGraph.K2Node_VariableGet Name="K2Node_VariableGet_199"
+			   VariableReference=(MemberName="baseStationArray",MemberGuid=BD0B12DA483B113DA20838A8B7EF86C6,bSelfContext=True)
+			   NodePosX=2256
+			   NodePosY=-2656
+			   NodeGuid=8441735341894AC0424507AE38C594E1
+			   CustomProperties Pin (PinId=92B8D86246B7B64BCB79D093554F79E8,PinName="baseStationArray",Direction="EGPD_Output",PinType.PinCategory="object",PinType.PinSubCategory="",PinType.PinSubCategoryObject=Class'/Script/Engine.StaticMeshComponent',PinType.PinSubCategoryMemberReference=(),PinType.PinValueType=(),PinType.ContainerType=Array,PinType.bIsMap=False,PinType.bIsSet=False,PinType.bIsArray=False,PinType.bIsReference=False,PinType.bIsConst=False,PinType.bIsWeakPointer=False,LinkedTo=(K2Node_CallArrayFunction_227 22674C5B4C7D8F21FD49C8A2A4EFF177,),PersistentGuid=00000000000000000000000000000000,bHidden=False,bNotConnectable=False,bDefaultValueIsReadOnly=False,bDefaultValueIsIgnored=False,bAdvancedView=False,bOrphanedPin=False,)
+			   CustomProperties Pin (PinId=4653AFB045E0B8B895E1C2B2CCC2283D,PinName="self",PinFriendlyName=NSLOCTEXT("K2Node", "Target", "Target"),PinType.PinCategory="object",PinType.PinSubCategory="",PinType.PinSubCategoryObject=BlueprintGeneratedClass'/Game/Blueprints/SteamVR_Pawn_Standing.SteamVR_Pawn_Standing_C',PinType.PinSubCategoryMemberReference=(),PinType.PinValueType=(),PinType.ContainerType=None,PinType.bIsMap=False,PinType.bIsSet=False,PinType.bIsArray=False,PinType.bIsReference=False,PinType.bIsConst=False,PinType.bIsWeakPointer=False,PersistentGuid=00000000000000000000000000000000,bHidden=True,bNotConnectable=False,bDefaultValueIsReadOnly=False,bDefaultValueIsIgnored=False,bAdvancedView=False,bOrphanedPin=False,)
+			End Object
+			Begin Object Class=/Script/BlueprintGraph.K2Node_CallFunction Name="K2Node_CallFunction_3281"
+			   bIsPureFunc=True
+			   FunctionReference=(MemberParent=Class'/Script/Engine.KismetMathLibrary',MemberName="Conv_VectorToTransform")
+			   NodePosX=1808
+			   NodePosY=-2640
+			   NodeGuid=1962718745EDAB1CF846C88774415A68
+			   CustomProperties Pin (PinId=96B155934815308F019857896A1E44D4,PinName="self",PinFriendlyName=NSLOCTEXT("K2Node", "Target", "Target"),PinToolTip="Target\nKismet Math Library Object Reference",PinType.PinCategory="object",PinType.PinSubCategory="",PinType.PinSubCategoryObject=Class'/Script/Engine.KismetMathLibrary',PinType.PinSubCategoryMemberReference=(),PinType.PinValueType=(),PinType.ContainerType=None,PinType.bIsMap=False,PinType.bIsSet=False,PinType.bIsArray=False,PinType.bIsReference=False,PinType.bIsConst=False,PinType.bIsWeakPointer=False,DefaultObject="/Script/Engine.Default__KismetMathLibrary",PersistentGuid=00000000000000000000000000000000,bHidden=True,bNotConnectable=False,bDefaultValueIsReadOnly=False,bDefaultValueIsIgnored=False,bAdvancedView=False,bOrphanedPin=False,)
+			   CustomProperties Pin (PinId=72212A6A4BCA13036E0B57BFD90B270E,PinName="InLocation",PinToolTip="In Location\nVector ",PinType.PinCategory="struct",PinType.PinSubCategory="",PinType.PinSubCategoryObject=ScriptStruct'/Script/CoreUObject.Vector',PinType.PinSubCategoryMemberReference=(),PinType.PinValueType=(),PinType.ContainerType=None,PinType.bIsMap=False,PinType.bIsSet=False,PinType.bIsArray=False,PinType.bIsReference=False,PinType.bIsConst=False,PinType.bIsWeakPointer=False,DefaultValue="0, 0, 0",AutogeneratedDefaultValue="0, 0, 0",LinkedTo=(K2Node_MacroInstance_118 E0271977466AEE7454D6A0AD5AABB816,),PersistentGuid=00000000000000000000000000000000,bHidden=False,bNotConnectable=False,bDefaultValueIsReadOnly=False,bDefaultValueIsIgnored=False,bAdvancedView=False,bOrphanedPin=False,)
+			   CustomProperties Pin (PinId=301B61884B3507108C1051B08AF72EC0,PinName="ReturnValue",PinToolTip="Return Value\nTransform ",Direction="EGPD_Output",PinType.PinCategory="struct",PinType.PinSubCategory="",PinType.PinSubCategoryObject=ScriptStruct'/Script/CoreUObject.Transform',PinType.PinSubCategoryMemberReference=(),PinType.PinValueType=(),PinType.ContainerType=None,PinType.bIsMap=False,PinType.bIsSet=False,PinType.bIsArray=False,PinType.bIsReference=False,PinType.bIsConst=False,PinType.bIsWeakPointer=False,LinkedTo=(K2Node_AddComponent_71 5C74A8634DDD697D1234B38F487C2D6D,),PersistentGuid=00000000000000000000000000000000,bHidden=False,bNotConnectable=False,bDefaultValueIsReadOnly=False,bDefaultValueIsIgnored=False,bAdvancedView=False,bOrphanedPin=False,)
+			End Object
+
+			END OBJECT
+	[/PARAM]
+[/OBJECT]
+[/EXCERPT:SteamVRC1]
+-->
+
+<!--
+[EXCERPT:SteamVRC2]	
+[OBJECT:Snippet]
+	[PARAMLITERAL:id]
+		SteamVRC2
+	[/PARAMLITERAL]
+	[PARAMLITERAL:type]
+		blueprint
+	[/PARAMLITERAL]
+	[PARAM:image]
+		[REGION:lightbox]
+		[![](SVRChapInteraction_04.png)(w:1000)](SVRChapInteraction_04.png)
+		[/REGION]
+	[/PARAM]
+	[PARAM:snippet]
+			BEGIN OBJECT
+							Begin Object Class=/Script/BlueprintGraph.K2Node_CallFunction Name="K2Node_CallFunction_1084"
+			   bIsPureFunc=True
+			   bIsConstFunc=True
+			   FunctionReference=(MemberParent=Class'/Script/SteamVR.SteamVRChaperoneComponent',MemberName="GetBounds")
+			   NodePosX=1232
+			   NodePosY=-2624
+			   NodeGuid=9D8CE6074BBE8BCF7AA85D8BB3B15185
+			   CustomProperties Pin (PinId=8B9136B54BDA0D37E448BDAC16DFFFF7,PinName="self",PinFriendlyName=NSLOCTEXT("K2Node", "Target", "Target"),PinToolTip="Target\nSteam VRChaperone Component Object Reference",PinType.PinCategory="object",PinType.PinSubCategory="",PinType.PinSubCategoryObject=Class'/Script/SteamVR.SteamVRChaperoneComponent',PinType.PinSubCategoryMemberReference=(),PinType.PinValueType=(),PinType.ContainerType=None,PinType.bIsMap=False,PinType.bIsSet=False,PinType.bIsArray=False,PinType.bIsReference=False,PinType.bIsConst=False,PinType.bIsWeakPointer=False,LinkedTo=(K2Node_VariableGet_197 4857BDE84AFE450BEB4D218869646F9B,),PersistentGuid=00000000000000000000000000000000,bHidden=False,bNotConnectable=False,bDefaultValueIsReadOnly=False,bDefaultValueIsIgnored=False,bAdvancedView=False,bOrphanedPin=False,)
+			   CustomProperties Pin (PinId=B35FC4BA4D68C6738438C18702E70E4C,PinName="ReturnValue",PinToolTip="Return Value\nArray of Vector s",Direction="EGPD_Output",PinType.PinCategory="struct",PinType.PinSubCategory="",PinType.PinSubCategoryObject=ScriptStruct'/Script/CoreUObject.Vector',PinType.PinSubCategoryMemberReference=(),PinType.PinValueType=(),PinType.ContainerType=Array,PinType.bIsMap=False,PinType.bIsSet=False,PinType.bIsArray=False,PinType.bIsReference=False,PinType.bIsConst=False,PinType.bIsWeakPointer=False,LinkedTo=(K2Node_MacroInstance_118 B6F62194447A3A47B7677494C9E6AF42,),PersistentGuid=00000000000000000000000000000000,bHidden=False,bNotConnectable=False,bDefaultValueIsReadOnly=False,bDefaultValueIsIgnored=False,bAdvancedView=False,bOrphanedPin=False,)
+			End Object
+			Begin Object Class=/Script/BlueprintGraph.K2Node_VariableGet Name="K2Node_VariableGet_197"
+			   VariableReference=(MemberName="SteamVRChaperone",bSelfContext=True)
+			   NodePosX=1248
+			   NodePosY=-2544
+			   NodeGuid=7E57A3734E176EA15785F397B1DC76FB
+			   CustomProperties Pin (PinId=4857BDE84AFE450BEB4D218869646F9B,PinName="SteamVRChaperone",Direction="EGPD_Output",PinType.PinCategory="object",PinType.PinSubCategory="",PinType.PinSubCategoryObject=Class'/Script/SteamVR.SteamVRChaperoneComponent',PinType.PinSubCategoryMemberReference=(),PinType.PinValueType=(),PinType.ContainerType=None,PinType.bIsMap=False,PinType.bIsSet=False,PinType.bIsArray=False,PinType.bIsReference=False,PinType.bIsConst=False,PinType.bIsWeakPointer=False,LinkedTo=(K2Node_CallFunction_1084 8B9136B54BDA0D37E448BDAC16DFFFF7,),PersistentGuid=00000000000000000000000000000000,bHidden=False,bNotConnectable=False,bDefaultValueIsReadOnly=False,bDefaultValueIsIgnored=False,bAdvancedView=False,bOrphanedPin=False,)
+			   CustomProperties Pin (PinId=1E59DFAA490E313860D944A45C6A61A1,PinName="self",PinFriendlyName=NSLOCTEXT("K2Node", "Target", "Target"),PinType.PinCategory="object",PinType.PinSubCategory="",PinType.PinSubCategoryObject=BlueprintGeneratedClass'/Game/Blueprints/SteamVR_Pawn_Standing.SteamVR_Pawn_Standing_C',PinType.PinSubCategoryMemberReference=(),PinType.PinValueType=(),PinType.ContainerType=None,PinType.bIsMap=False,PinType.bIsSet=False,PinType.bIsArray=False,PinType.bIsReference=False,PinType.bIsConst=False,PinType.bIsWeakPointer=False,PersistentGuid=00000000000000000000000000000000,bHidden=True,bNotConnectable=False,bDefaultValueIsReadOnly=False,bDefaultValueIsIgnored=False,bAdvancedView=False,bOrphanedPin=False,)
+			End Object
+			Begin Object Class=/Script/BlueprintGraph.K2Node_MacroInstance Name="K2Node_MacroInstance_118"
+			   MacroGraphReference=(MacroGraph=EdGraph'/Engine/EditorBlueprintResources/StandardMacros.StandardMacros:ForEachLoop',GraphBlueprint=Blueprint'/Engine/EditorBlueprintResources/StandardMacros.StandardMacros',GraphGuid=99DBFD5540A796041F72A5A9DA655026)
+			   ResolvedWildcardType=(PinCategory="struct",PinSubCategoryObject=ScriptStruct'/Script/CoreUObject.Vector',ContainerType=Array)
+			   NodePosX=1584
+			   NodePosY=-2768
+			   NodeGuid=37FA01EA488A9C7B9FFEA3AFCA4B5D7B
+			   CustomProperties Pin (PinId=2BDC94F34208A8F557E05F94CE1CE4AE,PinName="Exec",PinType.PinCategory="exec",PinType.PinSubCategory="",PinType.PinSubCategoryObject=None,PinType.PinSubCategoryMemberReference=(),PinType.PinValueType=(),PinType.ContainerType=None,PinType.bIsMap=False,PinType.bIsSet=False,PinType.bIsArray=False,PinType.bIsReference=False,PinType.bIsConst=False,PinType.bIsWeakPointer=False,LinkedTo=(K2Node_InputKey_74 451F705A45E7D71FD4CDD191C52193E1,),PersistentGuid=00000000000000000000000000000000,bHidden=False,bNotConnectable=False,bDefaultValueIsReadOnly=False,bDefaultValueIsIgnored=False,bAdvancedView=False,bOrphanedPin=False,)
+			   CustomProperties Pin (PinId=B6F62194447A3A47B7677494C9E6AF42,PinName="Array",PinType.PinCategory="struct",PinType.PinSubCategory="",PinType.PinSubCategoryObject=ScriptStruct'/Script/CoreUObject.Vector',PinType.PinSubCategoryMemberReference=(),PinType.PinValueType=(),PinType.ContainerType=Array,PinType.bIsMap=False,PinType.bIsSet=False,PinType.bIsArray=False,PinType.bIsReference=False,PinType.bIsConst=False,PinType.bIsWeakPointer=False,LinkedTo=(K2Node_CallFunction_1084 B35FC4BA4D68C6738438C18702E70E4C,),PersistentGuid=00000000000000000000000000000000,bHidden=False,bNotConnectable=False,bDefaultValueIsReadOnly=False,bDefaultValueIsIgnored=False,bAdvancedView=False,bOrphanedPin=False,)
+			   CustomProperties Pin (PinId=46A798794EB954364926C08AABB07BFF,PinName="LoopBody",Direction="EGPD_Output",PinType.PinCategory="exec",PinType.PinSubCategory="",PinType.PinSubCategoryObject=None,PinType.PinSubCategoryMemberReference=(),PinType.PinValueType=(),PinType.ContainerType=None,PinType.bIsMap=False,PinType.bIsSet=False,PinType.bIsArray=False,PinType.bIsReference=False,PinType.bIsConst=False,PinType.bIsWeakPointer=False,LinkedTo=(K2Node_AddComponent_71 34310B40490B2BAD8DED8ABB2127E66A,),PersistentGuid=00000000000000000000000000000000,bHidden=False,bNotConnectable=False,bDefaultValueIsReadOnly=False,bDefaultValueIsIgnored=False,bAdvancedView=False,bOrphanedPin=False,)
+			   CustomProperties Pin (PinId=E0271977466AEE7454D6A0AD5AABB816,PinName="Array Element",Direction="EGPD_Output",PinType.PinCategory="struct",PinType.PinSubCategory="",PinType.PinSubCategoryObject=ScriptStruct'/Script/CoreUObject.Vector',PinType.PinSubCategoryMemberReference=(),PinType.PinValueType=(),PinType.ContainerType=None,PinType.bIsMap=False,PinType.bIsSet=False,PinType.bIsArray=False,PinType.bIsReference=False,PinType.bIsConst=False,PinType.bIsWeakPointer=False,LinkedTo=(K2Node_CallFunction_3281 72212A6A4BCA13036E0B57BFD90B270E,),PersistentGuid=00000000000000000000000000000000,bHidden=False,bNotConnectable=False,bDefaultValueIsReadOnly=False,bDefaultValueIsIgnored=False,bAdvancedView=False,bOrphanedPin=False,)
+			   CustomProperties Pin (PinId=48EF763F42B4DA5034F2F58F9BBEB00D,PinName="Array Index",Direction="EGPD_Output",PinType.PinCategory="int",PinType.PinSubCategory="",PinType.PinSubCategoryObject=None,PinType.PinSubCategoryMemberReference=(),PinType.PinValueType=(),PinType.ContainerType=None,PinType.bIsMap=False,PinType.bIsSet=False,PinType.bIsArray=False,PinType.bIsReference=False,PinType.bIsConst=False,PinType.bIsWeakPointer=False,DefaultValue="0",AutogeneratedDefaultValue="0",PersistentGuid=00000000000000000000000000000000,bHidden=False,bNotConnectable=False,bDefaultValueIsReadOnly=False,bDefaultValueIsIgnored=False,bAdvancedView=False,bOrphanedPin=False,)
+			   CustomProperties Pin (PinId=7DA82531415B398F5EB95691B0938909,PinName="Completed",Direction="EGPD_Output",PinType.PinCategory="exec",PinType.PinSubCategory="",PinType.PinSubCategoryObject=None,PinType.PinSubCategoryMemberReference=(),PinType.PinValueType=(),PinType.ContainerType=None,PinType.bIsMap=False,PinType.bIsSet=False,PinType.bIsArray=False,PinType.bIsReference=False,PinType.bIsConst=False,PinType.bIsWeakPointer=False,PersistentGuid=00000000000000000000000000000000,bHidden=False,bNotConnectable=False,bDefaultValueIsReadOnly=False,bDefaultValueIsIgnored=False,bAdvancedView=False,bOrphanedPin=False,)
+			End Object
+			Begin Object Class=/Script/BlueprintGraph.K2Node_AddComponent Name="K2Node_AddComponent_71"
+			   TemplateBlueprint="/Game/Blueprints/SteamVR_Pawn_Standing.SteamVR_Pawn_Standing"
+			   TemplateType=Class'/Script/Engine.StaticMeshComponent'
+			   FunctionReference=(MemberName="AddComponent",bSelfContext=True)
+			   NodePosX=1920
+			   NodePosY=-2784
+			   NodeGuid=D899A1934C7F4236190FBB83E80E36E4
+			   CustomProperties Pin (PinId=34310B40490B2BAD8DED8ABB2127E66A,PinName="execute",PinToolTip="\nExec",PinType.PinCategory="exec",PinType.PinSubCategory="",PinType.PinSubCategoryObject=None,PinType.PinSubCategoryMemberReference=(),PinType.PinValueType=(),PinType.ContainerType=None,PinType.bIsMap=False,PinType.bIsSet=False,PinType.bIsArray=False,PinType.bIsReference=False,PinType.bIsConst=False,PinType.bIsWeakPointer=False,LinkedTo=(K2Node_MacroInstance_118 46A798794EB954364926C08AABB07BFF,),PersistentGuid=00000000000000000000000000000000,bHidden=False,bNotConnectable=False,bDefaultValueIsReadOnly=False,bDefaultValueIsIgnored=False,bAdvancedView=False,bOrphanedPin=False,)
+			   CustomProperties Pin (PinId=497F5C1D48B8B3DB2C2BA3AFBDDDA07D,PinName="then",PinToolTip="\nExec",Direction="EGPD_Output",PinType.PinCategory="exec",PinType.PinSubCategory="",PinType.PinSubCategoryObject=None,PinType.PinSubCategoryMemberReference=(),PinType.PinValueType=(),PinType.ContainerType=None,PinType.bIsMap=False,PinType.bIsSet=False,PinType.bIsArray=False,PinType.bIsReference=False,PinType.bIsConst=False,PinType.bIsWeakPointer=False,LinkedTo=(K2Node_CallArrayFunction_227 AB7C29724A4370E374E219B3136B3B14,),PersistentGuid=00000000000000000000000000000000,bHidden=False,bNotConnectable=False,bDefaultValueIsReadOnly=False,bDefaultValueIsIgnored=False,bAdvancedView=False,bOrphanedPin=False,)
+			   CustomProperties Pin (PinId=0942B0BA46BC8263AD4B8DB55CDECCEF,PinName="self",PinFriendlyName=NSLOCTEXT("K2Node", "Target", "Target"),PinToolTip="Target\nActor Object Reference",PinType.PinCategory="object",PinType.PinSubCategory="",PinType.PinSubCategoryObject=Class'/Script/Engine.Actor',PinType.PinSubCategoryMemberReference=(),PinType.PinValueType=(),PinType.ContainerType=None,PinType.bIsMap=False,PinType.bIsSet=False,PinType.bIsArray=False,PinType.bIsReference=False,PinType.bIsConst=False,PinType.bIsWeakPointer=False,PersistentGuid=00000000000000000000000000000000,bHidden=False,bNotConnectable=False,bDefaultValueIsReadOnly=False,bDefaultValueIsIgnored=False,bAdvancedView=False,bOrphanedPin=False,)
+			   CustomProperties Pin (PinId=21E468274D3DBA5F04E63E9919D1B6EB,PinName="TemplateName",PinToolTip="Template Name\nName\n\nThe name of the Component Template to use.",PinType.PinCategory="name",PinType.PinSubCategory="",PinType.PinSubCategoryObject=None,PinType.PinSubCategoryMemberReference=(),PinType.PinValueType=(),PinType.ContainerType=None,PinType.bIsMap=False,PinType.bIsSet=False,PinType.bIsArray=False,PinType.bIsReference=False,PinType.bIsConst=False,PinType.bIsWeakPointer=False,DefaultValue="NODE_AddStaticMeshComponent-5",AutogeneratedDefaultValue="None",PersistentGuid=00000000000000000000000000000000,bHidden=True,bNotConnectable=True,bDefaultValueIsReadOnly=True,bDefaultValueIsIgnored=False,bAdvancedView=False,bOrphanedPin=False,)
+			   CustomProperties Pin (PinId=948A308C4E60D0B8C5EC1D92028E558A,PinName="bManualAttachment",PinToolTip="Manual Attachment\nBoolean\n\nWhether manual or automatic attachment is to be used",PinType.PinCategory="bool",PinType.PinSubCategory="",PinType.PinSubCategoryObject=None,PinType.PinSubCategoryMemberReference=(),PinType.PinValueType=(),PinType.ContainerType=None,PinType.bIsMap=False,PinType.bIsSet=False,PinType.bIsArray=False,PinType.bIsReference=False,PinType.bIsConst=False,PinType.bIsWeakPointer=False,DefaultValue="false",AutogeneratedDefaultValue="false",PersistentGuid=00000000000000000000000000000000,bHidden=False,bNotConnectable=False,bDefaultValueIsReadOnly=False,bDefaultValueIsIgnored=False,bAdvancedView=False,bOrphanedPin=False,)
+			   CustomProperties Pin (PinId=5C74A8634DDD697D1234B38F487C2D6D,PinName="RelativeTransform",PinToolTip="Relative Transform\nTransform \n\nThe relative transform between the new component and its attach parent (automatic only)",PinType.PinCategory="struct",PinType.PinSubCategory="",PinType.PinSubCategoryObject=ScriptStruct'/Script/CoreUObject.Transform',PinType.PinSubCategoryMemberReference=(),PinType.PinValueType=(),PinType.ContainerType=None,PinType.bIsMap=False,PinType.bIsSet=False,PinType.bIsArray=False,PinType.bIsReference=False,PinType.bIsConst=True,PinType.bIsWeakPointer=False,LinkedTo=(K2Node_CallFunction_3281 301B61884B3507108C1051B08AF72EC0,),PersistentGuid=00000000000000000000000000000000,bHidden=False,bNotConnectable=False,bDefaultValueIsReadOnly=False,bDefaultValueIsIgnored=True,bAdvancedView=False,bOrphanedPin=False,)
+			   CustomProperties Pin (PinId=A79F33AE4CF14547CB740F9D0B5ECB67,PinName="ComponentTemplateContext",PinToolTip="Component Template Context\nObject Object Reference\n\nOptional UBlueprintGeneratedClass reference to use to find the template in. If null (or not a BPGC), component is sought in this Actor\'s class",PinType.PinCategory="object",PinType.PinSubCategory="",PinType.PinSubCategoryObject=Class'/Script/CoreUObject.Object',PinType.PinSubCategoryMemberReference=(),PinType.PinValueType=(),PinType.ContainerType=None,PinType.bIsMap=False,PinType.bIsSet=False,PinType.bIsArray=False,PinType.bIsReference=False,PinType.bIsConst=True,PinType.bIsWeakPointer=False,PersistentGuid=00000000000000000000000000000000,bHidden=True,bNotConnectable=True,bDefaultValueIsReadOnly=False,bDefaultValueIsIgnored=False,bAdvancedView=False,bOrphanedPin=False,)
+			   CustomProperties Pin (PinId=6F754DCB4AA8982E55C33D8BEA044333,PinName="ReturnValue",PinToolTip="Return Value\nStatic Mesh Component Object Reference",Direction="EGPD_Output",PinType.PinCategory="object",PinType.PinSubCategory="",PinType.PinSubCategoryObject=Class'/Script/Engine.StaticMeshComponent',PinType.PinSubCategoryMemberReference=(),PinType.PinValueType=(),PinType.ContainerType=None,PinType.bIsMap=False,PinType.bIsSet=False,PinType.bIsArray=False,PinType.bIsReference=False,PinType.bIsConst=False,PinType.bIsWeakPointer=False,LinkedTo=(K2Node_CallArrayFunction_227 EE5E261346A9D948350D08A538206500,),PersistentGuid=00000000000000000000000000000000,bHidden=False,bNotConnectable=False,bDefaultValueIsReadOnly=False,bDefaultValueIsIgnored=False,bAdvancedView=False,bOrphanedPin=False,)
+			End Object
+			Begin Object Class=/Script/BlueprintGraph.K2Node_InputKey Name="K2Node_InputKey_74"
+			   InputKey=MotionController_Right_Trigger
+			   NodePosX=1232
+			   NodePosY=-2768
+			   NodeGuid=115D70DD4D72911C84CDD396206876CB
+			   CustomProperties Pin (PinId=451F705A45E7D71FD4CDD191C52193E1,PinName="Pressed",Direction="EGPD_Output",PinType.PinCategory="exec",PinType.PinSubCategory="",PinType.PinSubCategoryObject=None,PinType.PinSubCategoryMemberReference=(),PinType.PinValueType=(),PinType.ContainerType=None,PinType.bIsMap=False,PinType.bIsSet=False,PinType.bIsArray=False,PinType.bIsReference=False,PinType.bIsConst=False,PinType.bIsWeakPointer=False,LinkedTo=(K2Node_MacroInstance_118 2BDC94F34208A8F557E05F94CE1CE4AE,),PersistentGuid=00000000000000000000000000000000,bHidden=False,bNotConnectable=False,bDefaultValueIsReadOnly=False,bDefaultValueIsIgnored=False,bAdvancedView=False,bOrphanedPin=False,)
+			   CustomProperties Pin (PinId=6F1385DB4EF7010345F655BD3744E632,PinName="Released",Direction="EGPD_Output",PinType.PinCategory="exec",PinType.PinSubCategory="",PinType.PinSubCategoryObject=None,PinType.PinSubCategoryMemberReference=(),PinType.PinValueType=(),PinType.ContainerType=None,PinType.bIsMap=False,PinType.bIsSet=False,PinType.bIsArray=False,PinType.bIsReference=False,PinType.bIsConst=False,PinType.bIsWeakPointer=False,LinkedTo=(K2Node_MacroInstance_120 8E8DF6784C76305868B4A7A4E0029393,),PersistentGuid=00000000000000000000000000000000,bHidden=False,bNotConnectable=False,bDefaultValueIsReadOnly=False,bDefaultValueIsIgnored=False,bAdvancedView=False,bOrphanedPin=False,)
+			   CustomProperties Pin (PinId=022C7579454B610F2F2ED0B282488461,PinName="Key",Direction="EGPD_Output",PinType.PinCategory="struct",PinType.PinSubCategory="",PinType.PinSubCategoryObject=ScriptStruct'/Script/InputCore.Key',PinType.PinSubCategoryMemberReference=(),PinType.PinValueType=(),PinType.ContainerType=None,PinType.bIsMap=False,PinType.bIsSet=False,PinType.bIsArray=False,PinType.bIsReference=False,PinType.bIsConst=False,PinType.bIsWeakPointer=False,PersistentGuid=00000000000000000000000000000000,bHidden=False,bNotConnectable=False,bDefaultValueIsReadOnly=False,bDefaultValueIsIgnored=False,bAdvancedView=False,bOrphanedPin=False,)
+			End Object
+			Begin Object Class=/Script/BlueprintGraph.K2Node_CallFunction Name="K2Node_CallFunction_1086"
+			   FunctionReference=(MemberParent=Class'/Script/Engine.ActorComponent',MemberName="K2_DestroyComponent")
+			   NodePosX=2224
+			   NodePosY=-2560
+			   NodeGuid=559DB1C94E0AA982E874B1BD54B1DFF7
+			   CustomProperties Pin (PinId=0E26C3A84D4201EE457C07B6CE08449B,PinName="execute",PinToolTip="\nExec",PinType.PinCategory="exec",PinType.PinSubCategory="",PinType.PinSubCategoryObject=None,PinType.PinSubCategoryMemberReference=(),PinType.PinValueType=(),PinType.ContainerType=None,PinType.bIsMap=False,PinType.bIsSet=False,PinType.bIsArray=False,PinType.bIsReference=False,PinType.bIsConst=False,PinType.bIsWeakPointer=False,LinkedTo=(K2Node_MacroInstance_119 BF64E45C419D12B713043EB0FD02DA2B,),PersistentGuid=00000000000000000000000000000000,bHidden=False,bNotConnectable=False,bDefaultValueIsReadOnly=False,bDefaultValueIsIgnored=False,bAdvancedView=False,bOrphanedPin=False,)
+			   CustomProperties Pin (PinId=291DBED847173AA169E912AF21CCF5A2,PinName="then",PinToolTip="\nExec",Direction="EGPD_Output",PinType.PinCategory="exec",PinType.PinSubCategory="",PinType.PinSubCategoryObject=None,PinType.PinSubCategoryMemberReference=(),PinType.PinValueType=(),PinType.ContainerType=None,PinType.bIsMap=False,PinType.bIsSet=False,PinType.bIsArray=False,PinType.bIsReference=False,PinType.bIsConst=False,PinType.bIsWeakPointer=False,PersistentGuid=00000000000000000000000000000000,bHidden=False,bNotConnectable=False,bDefaultValueIsReadOnly=False,bDefaultValueIsIgnored=False,bAdvancedView=False,bOrphanedPin=False,)
+			   CustomProperties Pin (PinId=2F4AF3F04BAAB30EEF09D4BFD3383176,PinName="self",PinFriendlyName=NSLOCTEXT("K2Node", "Target", "Target"),PinToolTip="Target\nActor Component Object Reference",PinType.PinCategory="object",PinType.PinSubCategory="",PinType.PinSubCategoryObject=Class'/Script/Engine.ActorComponent',PinType.PinSubCategoryMemberReference=(),PinType.PinValueType=(),PinType.ContainerType=None,PinType.bIsMap=False,PinType.bIsSet=False,PinType.bIsArray=False,PinType.bIsReference=False,PinType.bIsConst=False,PinType.bIsWeakPointer=False,LinkedTo=(K2Node_Knot_65 1967671D4EB26FCEFF4CF28C2C8A9C52,),PersistentGuid=00000000000000000000000000000000,bHidden=False,bNotConnectable=False,bDefaultValueIsReadOnly=False,bDefaultValueIsIgnored=False,bAdvancedView=False,bOrphanedPin=False,)
+			   CustomProperties Pin (PinId=C36C9E034992CB14A304DCA4272AD301,PinName="Object",PinToolTip="Object\nObject Object Reference",PinType.PinCategory="object",PinType.PinSubCategory="",PinType.PinSubCategoryObject=Class'/Script/CoreUObject.Object',PinType.PinSubCategoryMemberReference=(),PinType.PinValueType=(),PinType.ContainerType=None,PinType.bIsMap=False,PinType.bIsSet=False,PinType.bIsArray=False,PinType.bIsReference=False,PinType.bIsConst=False,PinType.bIsWeakPointer=False,PersistentGuid=00000000000000000000000000000000,bHidden=True,bNotConnectable=False,bDefaultValueIsReadOnly=False,bDefaultValueIsIgnored=False,bAdvancedView=False,bOrphanedPin=False,)
+			End Object
+			Begin Object Class=/Script/BlueprintGraph.K2Node_MacroInstance Name="K2Node_MacroInstance_119"
+			   MacroGraphReference=(MacroGraph=EdGraph'/Engine/EditorBlueprintResources/StandardMacros.StandardMacros:IsValid',GraphBlueprint=Blueprint'/Engine/EditorBlueprintResources/StandardMacros.StandardMacros',GraphGuid=64422BCD430703FF5CAEA8B79A32AA65)
+			   NodePosX=1936
+			   NodePosY=-2545
+			   NodeGuid=34F7AC554B511AFB7FBBAAA35D5CA3D9
+			   CustomProperties Pin (PinId=0495F39B41460298D26F09BCFBA110DF,PinName="exec",PinType.PinCategory="exec",PinType.PinSubCategory="",PinType.PinSubCategoryObject=None,PinType.PinSubCategoryMemberReference=(),PinType.PinValueType=(),PinType.ContainerType=None,PinType.bIsMap=False,PinType.bIsSet=False,PinType.bIsArray=False,PinType.bIsReference=False,PinType.bIsConst=False,PinType.bIsWeakPointer=False,LinkedTo=(K2Node_MacroInstance_120 90D1838E43154FBA184E92ADFE172302,),PersistentGuid=00000000000000000000000000000000,bHidden=False,bNotConnectable=False,bDefaultValueIsReadOnly=False,bDefaultValueIsIgnored=False,bAdvancedView=False,bOrphanedPin=False,)
+			   CustomProperties Pin (PinId=F13BFF7648589DEDE21E2D9B116789FB,PinName="InputObject",PinType.PinCategory="object",PinType.PinSubCategory="",PinType.PinSubCategoryObject=Class'/Script/CoreUObject.Object',PinType.PinSubCategoryMemberReference=(),PinType.PinValueType=(),PinType.ContainerType=None,PinType.bIsMap=False,PinType.bIsSet=False,PinType.bIsArray=False,PinType.bIsReference=False,PinType.bIsConst=False,PinType.bIsWeakPointer=False,LinkedTo=(K2Node_MacroInstance_120 3F8EE55345992BAA1A7BD6A63E4BC83F,),PersistentGuid=00000000000000000000000000000000,bHidden=False,bNotConnectable=False,bDefaultValueIsReadOnly=False,bDefaultValueIsIgnored=False,bAdvancedView=False,bOrphanedPin=False,)
+			   CustomProperties Pin (PinId=BF64E45C419D12B713043EB0FD02DA2B,PinName="Is Valid",Direction="EGPD_Output",PinType.PinCategory="exec",PinType.PinSubCategory="",PinType.PinSubCategoryObject=None,PinType.PinSubCategoryMemberReference=(),PinType.PinValueType=(),PinType.ContainerType=None,PinType.bIsMap=False,PinType.bIsSet=False,PinType.bIsArray=False,PinType.bIsReference=False,PinType.bIsConst=False,PinType.bIsWeakPointer=False,LinkedTo=(K2Node_CallFunction_1086 0E26C3A84D4201EE457C07B6CE08449B,),PersistentGuid=00000000000000000000000000000000,bHidden=False,bNotConnectable=False,bDefaultValueIsReadOnly=False,bDefaultValueIsIgnored=False,bAdvancedView=False,bOrphanedPin=False,)
+			   CustomProperties Pin (PinId=E7BA1502454D65452732BB90C917AA3D,PinName="Is Not Valid",Direction="EGPD_Output",PinType.PinCategory="exec",PinType.PinSubCategory="",PinType.PinSubCategoryObject=None,PinType.PinSubCategoryMemberReference=(),PinType.PinValueType=(),PinType.ContainerType=None,PinType.bIsMap=False,PinType.bIsSet=False,PinType.bIsArray=False,PinType.bIsReference=False,PinType.bIsConst=False,PinType.bIsWeakPointer=False,PersistentGuid=00000000000000000000000000000000,bHidden=False,bNotConnectable=False,bDefaultValueIsReadOnly=False,bDefaultValueIsIgnored=False,bAdvancedView=False,bOrphanedPin=False,)
+			End Object
+			Begin Object Class=/Script/BlueprintGraph.K2Node_MacroInstance Name="K2Node_MacroInstance_120"
+			   MacroGraphReference=(MacroGraph=EdGraph'/Engine/EditorBlueprintResources/StandardMacros.StandardMacros:ForEachLoop',GraphBlueprint=Blueprint'/Engine/EditorBlueprintResources/StandardMacros.StandardMacros',GraphGuid=99DBFD5540A796041F72A5A9DA655026)
+			   ResolvedWildcardType=(PinCategory="object",PinSubCategoryObject=Class'/Script/Engine.StaticMeshComponent',ContainerType=Array)
+			   NodePosX=1712
+			   NodePosY=-2545
+			   NodeGuid=60B992404FFA35480EE3A4ABAC6E44CA
+			   CustomProperties Pin (PinId=8E8DF6784C76305868B4A7A4E0029393,PinName="Exec",PinType.PinCategory="exec",PinType.PinSubCategory="",PinType.PinSubCategoryObject=None,PinType.PinSubCategoryMemberReference=(),PinType.PinValueType=(),PinType.ContainerType=None,PinType.bIsMap=False,PinType.bIsSet=False,PinType.bIsArray=False,PinType.bIsReference=False,PinType.bIsConst=False,PinType.bIsWeakPointer=False,LinkedTo=(K2Node_InputKey_74 6F1385DB4EF7010345F655BD3744E632,),PersistentGuid=00000000000000000000000000000000,bHidden=False,bNotConnectable=False,bDefaultValueIsReadOnly=False,bDefaultValueIsIgnored=False,bAdvancedView=False,bOrphanedPin=False,)
+			   CustomProperties Pin (PinId=EDCDE18041E5841D136E4091E04F8FF8,PinName="Array",PinType.PinCategory="object",PinType.PinSubCategory="",PinType.PinSubCategoryObject=Class'/Script/Engine.StaticMeshComponent',PinType.PinSubCategoryMemberReference=(),PinType.PinValueType=(),PinType.ContainerType=Array,PinType.bIsMap=False,PinType.bIsSet=False,PinType.bIsArray=False,PinType.bIsReference=False,PinType.bIsConst=False,PinType.bIsWeakPointer=False,LinkedTo=(K2Node_VariableGet_200 56C1DC54480A023F6C8149818BB067FA,),PersistentGuid=00000000000000000000000000000000,bHidden=False,bNotConnectable=False,bDefaultValueIsReadOnly=False,bDefaultValueIsIgnored=False,bAdvancedView=False,bOrphanedPin=False,)
+			   CustomProperties Pin (PinId=90D1838E43154FBA184E92ADFE172302,PinName="LoopBody",Direction="EGPD_Output",PinType.PinCategory="exec",PinType.PinSubCategory="",PinType.PinSubCategoryObject=None,PinType.PinSubCategoryMemberReference=(),PinType.PinValueType=(),PinType.ContainerType=None,PinType.bIsMap=False,PinType.bIsSet=False,PinType.bIsArray=False,PinType.bIsReference=False,PinType.bIsConst=False,PinType.bIsWeakPointer=False,LinkedTo=(K2Node_MacroInstance_119 0495F39B41460298D26F09BCFBA110DF,),PersistentGuid=00000000000000000000000000000000,bHidden=False,bNotConnectable=False,bDefaultValueIsReadOnly=False,bDefaultValueIsIgnored=False,bAdvancedView=False,bOrphanedPin=False,)
+			   CustomProperties Pin (PinId=3F8EE55345992BAA1A7BD6A63E4BC83F,PinName="Array Element",Direction="EGPD_Output",PinType.PinCategory="object",PinType.PinSubCategory="",PinType.PinSubCategoryObject=Class'/Script/Engine.StaticMeshComponent',PinType.PinSubCategoryMemberReference=(),PinType.PinValueType=(),PinType.ContainerType=None,PinType.bIsMap=False,PinType.bIsSet=False,PinType.bIsArray=False,PinType.bIsReference=False,PinType.bIsConst=False,PinType.bIsWeakPointer=False,LinkedTo=(K2Node_MacroInstance_119 F13BFF7648589DEDE21E2D9B116789FB,K2Node_Knot_64 412537E646194EC45D73ADB0C31250B0,),PersistentGuid=00000000000000000000000000000000,bHidden=False,bNotConnectable=False,bDefaultValueIsReadOnly=False,bDefaultValueIsIgnored=False,bAdvancedView=False,bOrphanedPin=False,)
+			   CustomProperties Pin (PinId=9B8A97E44E9E4118FDDC3798882DF993,PinName="Array Index",Direction="EGPD_Output",PinType.PinCategory="int",PinType.PinSubCategory="",PinType.PinSubCategoryObject=None,PinType.PinSubCategoryMemberReference=(),PinType.PinValueType=(),PinType.ContainerType=None,PinType.bIsMap=False,PinType.bIsSet=False,PinType.bIsArray=False,PinType.bIsReference=False,PinType.bIsConst=False,PinType.bIsWeakPointer=False,DefaultValue="0",AutogeneratedDefaultValue="0",PersistentGuid=00000000000000000000000000000000,bHidden=False,bNotConnectable=False,bDefaultValueIsReadOnly=False,bDefaultValueIsIgnored=False,bAdvancedView=False,bOrphanedPin=False,)
+			   CustomProperties Pin (PinId=3302906A4C6CD110E4D1DA837FC84505,PinName="Completed",Direction="EGPD_Output",PinType.PinCategory="exec",PinType.PinSubCategory="",PinType.PinSubCategoryObject=None,PinType.PinSubCategoryMemberReference=(),PinType.PinValueType=(),PinType.ContainerType=None,PinType.bIsMap=False,PinType.bIsSet=False,PinType.bIsArray=False,PinType.bIsReference=False,PinType.bIsConst=False,PinType.bIsWeakPointer=False,PersistentGuid=00000000000000000000000000000000,bHidden=False,bNotConnectable=False,bDefaultValueIsReadOnly=False,bDefaultValueIsIgnored=False,bAdvancedView=False,bOrphanedPin=False,)
+			End Object
+			Begin Object Class=/Script/BlueprintGraph.K2Node_CallArrayFunction Name="K2Node_CallArrayFunction_227"
+			   FunctionReference=(MemberParent=Class'/Script/Engine.KismetArrayLibrary',MemberName="Array_Add")
+			   NodePosX=2256
+			   NodePosY=-2744
+			   NodeGuid=43700AEC407CB6EBCFD0BBBDB6E7566D
+			   CustomProperties Pin (PinId=AB7C29724A4370E374E219B3136B3B14,PinName="execute",PinToolTip="\nExec",PinType.PinCategory="exec",PinType.PinSubCategory="",PinType.PinSubCategoryObject=None,PinType.PinSubCategoryMemberReference=(),PinType.PinValueType=(),PinType.ContainerType=None,PinType.bIsMap=False,PinType.bIsSet=False,PinType.bIsArray=False,PinType.bIsReference=False,PinType.bIsConst=False,PinType.bIsWeakPointer=False,LinkedTo=(K2Node_AddComponent_71 497F5C1D48B8B3DB2C2BA3AFBDDDA07D,),PersistentGuid=00000000000000000000000000000000,bHidden=False,bNotConnectable=False,bDefaultValueIsReadOnly=False,bDefaultValueIsIgnored=False,bAdvancedView=False,bOrphanedPin=False,)
+			   CustomProperties Pin (PinId=9C49631A4DCB38432F500787BD709DCB,PinName="then",PinToolTip="\nExec",Direction="EGPD_Output",PinType.PinCategory="exec",PinType.PinSubCategory="",PinType.PinSubCategoryObject=None,PinType.PinSubCategoryMemberReference=(),PinType.PinValueType=(),PinType.ContainerType=None,PinType.bIsMap=False,PinType.bIsSet=False,PinType.bIsArray=False,PinType.bIsReference=False,PinType.bIsConst=False,PinType.bIsWeakPointer=False,PersistentGuid=00000000000000000000000000000000,bHidden=False,bNotConnectable=False,bDefaultValueIsReadOnly=False,bDefaultValueIsIgnored=False,bAdvancedView=False,bOrphanedPin=False,)
+			   CustomProperties Pin (PinId=00FB150349E4D66D6F7FE7983CE94DBA,PinName="self",PinFriendlyName=NSLOCTEXT("K2Node", "Target", "Target"),PinToolTip="Target\nKismet Array Library Object Reference",PinType.PinCategory="object",PinType.PinSubCategory="",PinType.PinSubCategoryObject=Class'/Script/Engine.KismetArrayLibrary',PinType.PinSubCategoryMemberReference=(),PinType.PinValueType=(),PinType.ContainerType=None,PinType.bIsMap=False,PinType.bIsSet=False,PinType.bIsArray=False,PinType.bIsReference=False,PinType.bIsConst=False,PinType.bIsWeakPointer=False,DefaultObject="/Script/Engine.Default__KismetArrayLibrary",PersistentGuid=00000000000000000000000000000000,bHidden=True,bNotConnectable=False,bDefaultValueIsReadOnly=False,bDefaultValueIsIgnored=False,bAdvancedView=False,bOrphanedPin=False,)
+			   CustomProperties Pin (PinId=22674C5B4C7D8F21FD49C8A2A4EFF177,PinName="TargetArray",PinToolTip="Target Array\nArray of Static Mesh Component Object References\n\nThe array to add item to",PinType.PinCategory="object",PinType.PinSubCategory="",PinType.PinSubCategoryObject=Class'/Script/Engine.StaticMeshComponent',PinType.PinSubCategoryMemberReference=(),PinType.PinValueType=(),PinType.ContainerType=Array,PinType.bIsMap=False,PinType.bIsSet=False,PinType.bIsArray=False,PinType.bIsReference=True,PinType.bIsConst=True,PinType.bIsWeakPointer=False,LinkedTo=(K2Node_VariableGet_199 92B8D86246B7B64BCB79D093554F79E8,),PersistentGuid=00000000000000000000000000000000,bHidden=False,bNotConnectable=False,bDefaultValueIsReadOnly=False,bDefaultValueIsIgnored=True,bAdvancedView=False,bOrphanedPin=False,)
+			   CustomProperties Pin (PinId=EE5E261346A9D948350D08A538206500,PinName="NewItem",PinToolTip="New Item\nStatic Mesh Component Object Reference (by ref)\n\nThe item to add to the array",PinType.PinCategory="object",PinType.PinSubCategory="",PinType.PinSubCategoryObject=Class'/Script/Engine.StaticMeshComponent',PinType.PinSubCategoryMemberReference=(),PinType.PinValueType=(),PinType.ContainerType=None,PinType.bIsMap=False,PinType.bIsSet=False,PinType.bIsArray=False,PinType.bIsReference=True,PinType.bIsConst=True,PinType.bIsWeakPointer=False,LinkedTo=(K2Node_AddComponent_71 6F754DCB4AA8982E55C33D8BEA044333,),PersistentGuid=00000000000000000000000000000000,bHidden=False,bNotConnectable=False,bDefaultValueIsReadOnly=False,bDefaultValueIsIgnored=False,bAdvancedView=False,bOrphanedPin=False,)
+			   CustomProperties Pin (PinId=BE1A5C404E2C69FB80B60D8497412508,PinName="ReturnValue",PinToolTip="Return Value\nInteger\n\nThe index of the newly added item",Direction="EGPD_Output",PinType.PinCategory="int",PinType.PinSubCategory="",PinType.PinSubCategoryObject=None,PinType.PinSubCategoryMemberReference=(),PinType.PinValueType=(),PinType.ContainerType=None,PinType.bIsMap=False,PinType.bIsSet=False,PinType.bIsArray=False,PinType.bIsReference=False,PinType.bIsConst=False,PinType.bIsWeakPointer=False,DefaultValue="0",AutogeneratedDefaultValue="0",PersistentGuid=00000000000000000000000000000000,bHidden=False,bNotConnectable=False,bDefaultValueIsReadOnly=False,bDefaultValueIsIgnored=False,bAdvancedView=False,bOrphanedPin=False,)
+			End Object
+			Begin Object Class=/Script/BlueprintGraph.K2Node_VariableGet Name="K2Node_VariableGet_199"
+			   VariableReference=(MemberName="baseStationArray",MemberGuid=BD0B12DA483B113DA20838A8B7EF86C6,bSelfContext=True)
+			   NodePosX=2256
+			   NodePosY=-2656
+			   NodeGuid=8441735341894AC0424507AE38C594E1
+			   CustomProperties Pin (PinId=92B8D86246B7B64BCB79D093554F79E8,PinName="baseStationArray",Direction="EGPD_Output",PinType.PinCategory="object",PinType.PinSubCategory="",PinType.PinSubCategoryObject=Class'/Script/Engine.StaticMeshComponent',PinType.PinSubCategoryMemberReference=(),PinType.PinValueType=(),PinType.ContainerType=Array,PinType.bIsMap=False,PinType.bIsSet=False,PinType.bIsArray=False,PinType.bIsReference=False,PinType.bIsConst=False,PinType.bIsWeakPointer=False,LinkedTo=(K2Node_CallArrayFunction_227 22674C5B4C7D8F21FD49C8A2A4EFF177,),PersistentGuid=00000000000000000000000000000000,bHidden=False,bNotConnectable=False,bDefaultValueIsReadOnly=False,bDefaultValueIsIgnored=False,bAdvancedView=False,bOrphanedPin=False,)
+			   CustomProperties Pin (PinId=4653AFB045E0B8B895E1C2B2CCC2283D,PinName="self",PinFriendlyName=NSLOCTEXT("K2Node", "Target", "Target"),PinType.PinCategory="object",PinType.PinSubCategory="",PinType.PinSubCategoryObject=BlueprintGeneratedClass'/Game/Blueprints/SteamVR_Pawn_Standing.SteamVR_Pawn_Standing_C',PinType.PinSubCategoryMemberReference=(),PinType.PinValueType=(),PinType.ContainerType=None,PinType.bIsMap=False,PinType.bIsSet=False,PinType.bIsArray=False,PinType.bIsReference=False,PinType.bIsConst=False,PinType.bIsWeakPointer=False,PersistentGuid=00000000000000000000000000000000,bHidden=True,bNotConnectable=False,bDefaultValueIsReadOnly=False,bDefaultValueIsIgnored=False,bAdvancedView=False,bOrphanedPin=False,)
+			End Object
+			Begin Object Class=/Script/BlueprintGraph.K2Node_VariableGet Name="K2Node_VariableGet_200"
+			   VariableReference=(MemberName="baseStationArray",MemberGuid=BD0B12DA483B113DA20838A8B7EF86C6,bSelfContext=True)
+			   NodePosX=1520
+			   NodePosY=-2497
+			   NodeGuid=6B9F6ADD4B1EC3299164D38C6798C654
+			   CustomProperties Pin (PinId=56C1DC54480A023F6C8149818BB067FA,PinName="baseStationArray",Direction="EGPD_Output",PinType.PinCategory="object",PinType.PinSubCategory="",PinType.PinSubCategoryObject=Class'/Script/Engine.StaticMeshComponent',PinType.PinSubCategoryMemberReference=(),PinType.PinValueType=(),PinType.ContainerType=Array,PinType.bIsMap=False,PinType.bIsSet=False,PinType.bIsArray=False,PinType.bIsReference=False,PinType.bIsConst=False,PinType.bIsWeakPointer=False,LinkedTo=(K2Node_MacroInstance_120 EDCDE18041E5841D136E4091E04F8FF8,),PersistentGuid=00000000000000000000000000000000,bHidden=False,bNotConnectable=False,bDefaultValueIsReadOnly=False,bDefaultValueIsIgnored=False,bAdvancedView=False,bOrphanedPin=False,)
+			   CustomProperties Pin (PinId=C71A09704C9B5D5907C2C9BB475CE4CB,PinName="self",PinFriendlyName=NSLOCTEXT("K2Node", "Target", "Target"),PinType.PinCategory="object",PinType.PinSubCategory="",PinType.PinSubCategoryObject=BlueprintGeneratedClass'/Game/Blueprints/SteamVR_Pawn_Standing.SteamVR_Pawn_Standing_C',PinType.PinSubCategoryMemberReference=(),PinType.PinValueType=(),PinType.ContainerType=None,PinType.bIsMap=False,PinType.bIsSet=False,PinType.bIsArray=False,PinType.bIsReference=False,PinType.bIsConst=False,PinType.bIsWeakPointer=False,PersistentGuid=00000000000000000000000000000000,bHidden=True,bNotConnectable=False,bDefaultValueIsReadOnly=False,bDefaultValueIsIgnored=False,bAdvancedView=False,bOrphanedPin=False,)
+			End Object
+			Begin Object Class=/Script/BlueprintGraph.K2Node_Knot Name="K2Node_Knot_65"
+			   NodePosX=2096
+			   NodePosY=-2385
+			   NodeGuid=57BFDD97498AE8033959D6B59D5D7066
+			   CustomProperties Pin (PinId=412537E646194EC45D73ADB0C31250B0,PinName="InputPin",PinType.PinCategory="object",PinType.PinSubCategory="",PinType.PinSubCategoryObject=Class'/Script/Engine.StaticMeshComponent',PinType.PinSubCategoryMemberReference=(),PinType.PinValueType=(),PinType.ContainerType=None,PinType.bIsMap=False,PinType.bIsSet=False,PinType.bIsArray=False,PinType.bIsReference=False,PinType.bIsConst=False,PinType.bIsWeakPointer=False,LinkedTo=(K2Node_Knot_64 1967671D4EB26FCEFF4CF28C2C8A9C52,),PersistentGuid=00000000000000000000000000000000,bHidden=False,bNotConnectable=False,bDefaultValueIsReadOnly=False,bDefaultValueIsIgnored=True,bAdvancedView=False,bOrphanedPin=False,)
+			   CustomProperties Pin (PinId=1967671D4EB26FCEFF4CF28C2C8A9C52,PinName="OutputPin",Direction="EGPD_Output",PinType.PinCategory="object",PinType.PinSubCategory="",PinType.PinSubCategoryObject=Class'/Script/Engine.StaticMeshComponent',PinType.PinSubCategoryMemberReference=(),PinType.PinValueType=(),PinType.ContainerType=None,PinType.bIsMap=False,PinType.bIsSet=False,PinType.bIsArray=False,PinType.bIsReference=False,PinType.bIsConst=False,PinType.bIsWeakPointer=False,LinkedTo=(K2Node_CallFunction_1086 2F4AF3F04BAAB30EEF09D4BFD3383176,),PersistentGuid=00000000000000000000000000000000,bHidden=False,bNotConnectable=False,bDefaultValueIsReadOnly=False,bDefaultValueIsIgnored=False,bAdvancedView=False,bOrphanedPin=False,)
+			End Object
+			Begin Object Class=/Script/BlueprintGraph.K2Node_Knot Name="K2Node_Knot_64"
+			   NodePosX=1984
+			   NodePosY=-2385
+			   NodeGuid=6379924049F4E5785F49B09C23EB1DF7
+			   CustomProperties Pin (PinId=412537E646194EC45D73ADB0C31250B0,PinName="InputPin",PinType.PinCategory="object",PinType.PinSubCategory="",PinType.PinSubCategoryObject=Class'/Script/Engine.StaticMeshComponent',PinType.PinSubCategoryMemberReference=(),PinType.PinValueType=(),PinType.ContainerType=None,PinType.bIsMap=False,PinType.bIsSet=False,PinType.bIsArray=False,PinType.bIsReference=False,PinType.bIsConst=False,PinType.bIsWeakPointer=False,LinkedTo=(K2Node_MacroInstance_120 3F8EE55345992BAA1A7BD6A63E4BC83F,),PersistentGuid=00000000000000000000000000000000,bHidden=False,bNotConnectable=False,bDefaultValueIsReadOnly=False,bDefaultValueIsIgnored=True,bAdvancedView=False,bOrphanedPin=False,)
+			   CustomProperties Pin (PinId=1967671D4EB26FCEFF4CF28C2C8A9C52,PinName="OutputPin",Direction="EGPD_Output",PinType.PinCategory="object",PinType.PinSubCategory="",PinType.PinSubCategoryObject=Class'/Script/Engine.StaticMeshComponent',PinType.PinSubCategoryMemberReference=(),PinType.PinValueType=(),PinType.ContainerType=None,PinType.bIsMap=False,PinType.bIsSet=False,PinType.bIsArray=False,PinType.bIsReference=False,PinType.bIsConst=False,PinType.bIsWeakPointer=False,LinkedTo=(K2Node_Knot_65 412537E646194EC45D73ADB0C31250B0,),PersistentGuid=00000000000000000000000000000000,bHidden=False,bNotConnectable=False,bDefaultValueIsReadOnly=False,bDefaultValueIsIgnored=False,bAdvancedView=False,bOrphanedPin=False,)
+			End Object
+			Begin Object Class=/Script/BlueprintGraph.K2Node_CallFunction Name="K2Node_CallFunction_3281"
+			   bIsPureFunc=True
+			   FunctionReference=(MemberParent=Class'/Script/Engine.KismetMathLibrary',MemberName="Conv_VectorToTransform")
+			   NodePosX=1808
+			   NodePosY=-2640
+			   NodeGuid=1962718745EDAB1CF846C88774415A68
+			   CustomProperties Pin (PinId=96B155934815308F019857896A1E44D4,PinName="self",PinFriendlyName=NSLOCTEXT("K2Node", "Target", "Target"),PinToolTip="Target\nKismet Math Library Object Reference",PinType.PinCategory="object",PinType.PinSubCategory="",PinType.PinSubCategoryObject=Class'/Script/Engine.KismetMathLibrary',PinType.PinSubCategoryMemberReference=(),PinType.PinValueType=(),PinType.ContainerType=None,PinType.bIsMap=False,PinType.bIsSet=False,PinType.bIsArray=False,PinType.bIsReference=False,PinType.bIsConst=False,PinType.bIsWeakPointer=False,DefaultObject="/Script/Engine.Default__KismetMathLibrary",PersistentGuid=00000000000000000000000000000000,bHidden=True,bNotConnectable=False,bDefaultValueIsReadOnly=False,bDefaultValueIsIgnored=False,bAdvancedView=False,bOrphanedPin=False,)
+			   CustomProperties Pin (PinId=72212A6A4BCA13036E0B57BFD90B270E,PinName="InLocation",PinToolTip="In Location\nVector ",PinType.PinCategory="struct",PinType.PinSubCategory="",PinType.PinSubCategoryObject=ScriptStruct'/Script/CoreUObject.Vector',PinType.PinSubCategoryMemberReference=(),PinType.PinValueType=(),PinType.ContainerType=None,PinType.bIsMap=False,PinType.bIsSet=False,PinType.bIsArray=False,PinType.bIsReference=False,PinType.bIsConst=False,PinType.bIsWeakPointer=False,DefaultValue="0, 0, 0",AutogeneratedDefaultValue="0, 0, 0",LinkedTo=(K2Node_MacroInstance_118 E0271977466AEE7454D6A0AD5AABB816,),PersistentGuid=00000000000000000000000000000000,bHidden=False,bNotConnectable=False,bDefaultValueIsReadOnly=False,bDefaultValueIsIgnored=False,bAdvancedView=False,bOrphanedPin=False,)
+			   CustomProperties Pin (PinId=301B61884B3507108C1051B08AF72EC0,PinName="ReturnValue",PinToolTip="Return Value\nTransform ",Direction="EGPD_Output",PinType.PinCategory="struct",PinType.PinSubCategory="",PinType.PinSubCategoryObject=ScriptStruct'/Script/CoreUObject.Transform',PinType.PinSubCategoryMemberReference=(),PinType.PinValueType=(),PinType.ContainerType=None,PinType.bIsMap=False,PinType.bIsSet=False,PinType.bIsArray=False,PinType.bIsReference=False,PinType.bIsConst=False,PinType.bIsWeakPointer=False,LinkedTo=(K2Node_AddComponent_71 5C74A8634DDD697D1234B38F487C2D6D,),PersistentGuid=00000000000000000000000000000000,bHidden=False,bNotConnectable=False,bDefaultValueIsReadOnly=False,bDefaultValueIsIgnored=False,bAdvancedView=False,bOrphanedPin=False,)
+			End Object
+			END OBJECT
+	[/PARAM]
+[/OBJECT]
+[/EXCERPT:SteamVRC2]
+-->
+
+<!--
+[EXCERPT:SteamVRC3]	
+[OBJECT:Snippet]
+	[PARAMLITERAL:id]
+		SteamVRC3
+	[/PARAMLITERAL]
+	[PARAMLITERAL:type]
+		blueprint
+	[/PARAMLITERAL]
+	[PARAM:image]
+		[REGION:lightbox]
+		[![](SVRChapInteraction_07.png)(w:1000)](SVRChapInteraction_07.png)
+		[/REGION]
+	[/PARAM]
+	[PARAM:snippet]
+			BEGIN OBJECT
+						Begin Object Class=/Script/BlueprintGraph.K2Node_CallFunction Name="K2Node_CallFunction_0"
+		   bIsPureFunc=True
+		   bIsConstFunc=True
+		   FunctionReference=(MemberParent=Class'/Script/SteamVR.SteamVRChaperoneComponent',MemberName="GetBounds")
+		   NodePosX=3472
+		   NodePosY=-1632
+		   ErrorType=1
+		   NodeGuid=F19BE7A840D794F8A26EC6BAE570335B
+		   CustomProperties Pin (PinId=ECF897334074A304613E1DAE43E9C939,PinName="self",PinFriendlyName=NSLOCTEXT("K2Node", "Target", "Target"),PinToolTip="Target\nSteam VRChaperone Component Object Reference",PinType.PinCategory="object",PinType.PinSubCategory="",PinType.PinSubCategoryObject=Class'/Script/SteamVR.SteamVRChaperoneComponent',PinType.PinSubCategoryMemberReference=(),PinType.PinValueType=(),PinType.ContainerType=None,PinType.bIsMap=False,PinType.bIsSet=False,PinType.bIsArray=False,PinType.bIsReference=False,PinType.bIsConst=False,PinType.bIsWeakPointer=False,LinkedTo=(K2Node_VariableGet_0 7FF3CF2B4F054A3D1F35CB89F0C81BEA,),PersistentGuid=00000000000000000000000000000000,bHidden=False,bNotConnectable=False,bDefaultValueIsReadOnly=False,bDefaultValueIsIgnored=False,bAdvancedView=False,bOrphanedPin=False,)
+		   CustomProperties Pin (PinId=D3E7C6854F51B77CCCD1D4A533C9EDB9,PinName="ReturnValue",PinToolTip="Return Value\nArray of Vector s",Direction="EGPD_Output",PinType.PinCategory="struct",PinType.PinSubCategory="",PinType.PinSubCategoryObject=ScriptStruct'/Script/CoreUObject.Vector',PinType.PinSubCategoryMemberReference=(),PinType.PinValueType=(),PinType.ContainerType=Array,PinType.bIsMap=False,PinType.bIsSet=False,PinType.bIsArray=False,PinType.bIsReference=False,PinType.bIsConst=False,PinType.bIsWeakPointer=False,LinkedTo=(K2Node_CallFunction_3 CEA4B4E744ADC1B1675B3FA64E8724DC,),PersistentGuid=00000000000000000000000000000000,bHidden=False,bNotConnectable=False,bDefaultValueIsReadOnly=False,bDefaultValueIsIgnored=False,bAdvancedView=False,bOrphanedPin=False,)
+		End Object
+		Begin Object Class=/Script/BlueprintGraph.K2Node_VariableGet Name="K2Node_VariableGet_0"
+		   VariableReference=(MemberName="SteamVRChaperone",bSelfContext=True)
+		   NodePosX=3296
+		   NodePosY=-1584
+		   NodeGuid=4AE71D05407F17CE6E3D9F9CE446FA35
+		   CustomProperties Pin (PinId=7FF3CF2B4F054A3D1F35CB89F0C81BEA,PinName="SteamVRChaperone",Direction="EGPD_Output",PinType.PinCategory="object",PinType.PinSubCategory="",PinType.PinSubCategoryObject=Class'/Script/SteamVR.SteamVRChaperoneComponent',PinType.PinSubCategoryMemberReference=(),PinType.PinValueType=(),PinType.ContainerType=None,PinType.bIsMap=False,PinType.bIsSet=False,PinType.bIsArray=False,PinType.bIsReference=False,PinType.bIsConst=False,PinType.bIsWeakPointer=False,LinkedTo=(K2Node_CallFunction_0 ECF897334074A304613E1DAE43E9C939,),PersistentGuid=00000000000000000000000000000000,bHidden=False,bNotConnectable=False,bDefaultValueIsReadOnly=False,bDefaultValueIsIgnored=False,bAdvancedView=False,bOrphanedPin=False,)
+		   CustomProperties Pin (PinId=50284E2D4E2065FDEF5FF790F3A616FD,PinName="self",PinFriendlyName=NSLOCTEXT("K2Node", "Target", "Target"),PinType.PinCategory="object",PinType.PinSubCategory="",PinType.PinSubCategoryObject=BlueprintGeneratedClass'/Game/Blueprints/SteamVR_Pawn_Standing.SteamVR_Pawn_Standing_C',PinType.PinSubCategoryMemberReference=(),PinType.PinValueType=(),PinType.ContainerType=None,PinType.bIsMap=False,PinType.bIsSet=False,PinType.bIsArray=False,PinType.bIsReference=False,PinType.bIsConst=False,PinType.bIsWeakPointer=False,PersistentGuid=00000000000000000000000000000000,bHidden=True,bNotConnectable=False,bDefaultValueIsReadOnly=False,bDefaultValueIsIgnored=False,bAdvancedView=False,bOrphanedPin=False,)
+		End Object
+		Begin Object Class=/Script/BlueprintGraph.K2Node_CallFunction Name="K2Node_CallFunction_3"
+		   FunctionReference=(MemberParent=Class'/Script/Engine.KismetMathLibrary',MemberName="MinimumAreaRectangle")
+		   NodePosX=3840
+		   NodePosY=-1776
+		   ErrorType=1
+		   NodeGuid=888F8C0949052A00B47FD0830573ED18
+		   CustomProperties Pin (PinId=27A36946461CB0706E8812873937728A,PinName="execute",PinToolTip="\nExec",PinType.PinCategory="exec",PinType.PinSubCategory="",PinType.PinSubCategoryObject=None,PinType.PinSubCategoryMemberReference=(),PinType.PinValueType=(),PinType.ContainerType=None,PinType.bIsMap=False,PinType.bIsSet=False,PinType.bIsArray=False,PinType.bIsReference=False,PinType.bIsConst=False,PinType.bIsWeakPointer=False,LinkedTo=(K2Node_InputKey_86 451F705A45E7D71FD4CDD191C52193E1,),PersistentGuid=00000000000000000000000000000000,bHidden=False,bNotConnectable=False,bDefaultValueIsReadOnly=False,bDefaultValueIsIgnored=False,bAdvancedView=False,bOrphanedPin=False,)
+		   CustomProperties Pin (PinId=D8DDC8B0419E7157C72694920CD01FB1,PinName="then",PinToolTip="\nExec",Direction="EGPD_Output",PinType.PinCategory="exec",PinType.PinSubCategory="",PinType.PinSubCategoryObject=None,PinType.PinSubCategoryMemberReference=(),PinType.PinValueType=(),PinType.ContainerType=None,PinType.bIsMap=False,PinType.bIsSet=False,PinType.bIsArray=False,PinType.bIsReference=False,PinType.bIsConst=False,PinType.bIsWeakPointer=False,LinkedTo=(K2Node_CallFunction_994 60CD775D47E28A58610F9DBBA97E8501,),PersistentGuid=00000000000000000000000000000000,bHidden=False,bNotConnectable=False,bDefaultValueIsReadOnly=False,bDefaultValueIsIgnored=False,bAdvancedView=False,bOrphanedPin=False,)
+		   CustomProperties Pin (PinId=20481E5742DCD21D0AF0EB95CEFFA209,PinName="self",PinFriendlyName=NSLOCTEXT("K2Node", "Target", "Target"),PinToolTip="Target\nKismet Math Library Object Reference",PinType.PinCategory="object",PinType.PinSubCategory="",PinType.PinSubCategoryObject=Class'/Script/Engine.KismetMathLibrary',PinType.PinSubCategoryMemberReference=(),PinType.PinValueType=(),PinType.ContainerType=None,PinType.bIsMap=False,PinType.bIsSet=False,PinType.bIsArray=False,PinType.bIsReference=False,PinType.bIsConst=False,PinType.bIsWeakPointer=False,DefaultObject="/Script/Engine.Default__KismetMathLibrary",PersistentGuid=00000000000000000000000000000000,bHidden=True,bNotConnectable=False,bDefaultValueIsReadOnly=False,bDefaultValueIsIgnored=False,bAdvancedView=False,bOrphanedPin=False,)
+		   CustomProperties Pin (PinId=069BB5F946DE90B19F86A0840CAE4DC1,PinName="WorldContextObject",PinToolTip="World Context Object\nObject Object Reference",PinType.PinCategory="object",PinType.PinSubCategory="",PinType.PinSubCategoryObject=Class'/Script/CoreUObject.Object',PinType.PinSubCategoryMemberReference=(),PinType.PinValueType=(),PinType.ContainerType=None,PinType.bIsMap=False,PinType.bIsSet=False,PinType.bIsArray=False,PinType.bIsReference=False,PinType.bIsConst=False,PinType.bIsWeakPointer=False,PersistentGuid=00000000000000000000000000000000,bHidden=True,bNotConnectable=False,bDefaultValueIsReadOnly=False,bDefaultValueIsIgnored=False,bAdvancedView=False,bOrphanedPin=False,)
+		   CustomProperties Pin (PinId=CEA4B4E744ADC1B1675B3FA64E8724DC,PinName="InVerts",PinToolTip="In Verts\nArray of Vector s\n\nPoints to enclose in the rectangle",PinType.PinCategory="struct",PinType.PinSubCategory="",PinType.PinSubCategoryObject=ScriptStruct'/Script/CoreUObject.Vector',PinType.PinSubCategoryMemberReference=(),PinType.PinValueType=(),PinType.ContainerType=Array,PinType.bIsMap=False,PinType.bIsSet=False,PinType.bIsArray=False,PinType.bIsReference=True,PinType.bIsConst=True,PinType.bIsWeakPointer=False,LinkedTo=(K2Node_CallFunction_0 D3E7C6854F51B77CCCD1D4A533C9EDB9,),PersistentGuid=00000000000000000000000000000000,bHidden=False,bNotConnectable=False,bDefaultValueIsReadOnly=False,bDefaultValueIsIgnored=True,bAdvancedView=False,bOrphanedPin=False,)
+		   CustomProperties Pin (PinId=76C031DD4BCFAD63300ACEB272F06EA1,PinName="SampleSurfaceNormal",PinToolTip="Sample Surface Normal\nVector  (by ref)",PinType.PinCategory="struct",PinType.PinSubCategory="",PinType.PinSubCategoryObject=ScriptStruct'/Script/CoreUObject.Vector',PinType.PinSubCategoryMemberReference=(),PinType.PinValueType=(),PinType.ContainerType=None,PinType.bIsMap=False,PinType.bIsSet=False,PinType.bIsArray=False,PinType.bIsReference=True,PinType.bIsConst=True,PinType.bIsWeakPointer=False,DefaultValue="0, 0, 0",AutogeneratedDefaultValue="0, 0, 0",LinkedTo=(K2Node_VariableGet_43 764BD867405B9674CC281B84A04BD8D4,),PersistentGuid=00000000000000000000000000000000,bHidden=False,bNotConnectable=False,bDefaultValueIsReadOnly=False,bDefaultValueIsIgnored=True,bAdvancedView=False,bOrphanedPin=False,)
+		   CustomProperties Pin (PinId=990FCC854987628A1F2188A6304C3F94,PinName="OutRectCenter",PinToolTip="Out Rect Center\nVector ",Direction="EGPD_Output",PinType.PinCategory="struct",PinType.PinSubCategory="",PinType.PinSubCategoryObject=ScriptStruct'/Script/CoreUObject.Vector',PinType.PinSubCategoryMemberReference=(),PinType.PinValueType=(),PinType.ContainerType=None,PinType.bIsMap=False,PinType.bIsSet=False,PinType.bIsArray=False,PinType.bIsReference=False,PinType.bIsConst=False,PinType.bIsWeakPointer=False,DefaultValue="0, 0, 0",AutogeneratedDefaultValue="0, 0, 0",PersistentGuid=00000000000000000000000000000000,bHidden=False,bNotConnectable=False,bDefaultValueIsReadOnly=False,bDefaultValueIsIgnored=False,bAdvancedView=False,bOrphanedPin=False,)
+		   CustomProperties Pin (PinId=C4FA178B4D01CFF46CD2DBB72D723B71,PinName="OutRectRotation",PinToolTip="Out Rect Rotation\nRotator ",Direction="EGPD_Output",PinType.PinCategory="struct",PinType.PinSubCategory="",PinType.PinSubCategoryObject=ScriptStruct'/Script/CoreUObject.Rotator',PinType.PinSubCategoryMemberReference=(),PinType.PinValueType=(),PinType.ContainerType=None,PinType.bIsMap=False,PinType.bIsSet=False,PinType.bIsArray=False,PinType.bIsReference=False,PinType.bIsConst=False,PinType.bIsWeakPointer=False,DefaultValue="0, 0, 0",AutogeneratedDefaultValue="0, 0, 0",LinkedTo=(K2Node_Knot_85 365F26FF45B3A67424DF5F9578C950B3,),PersistentGuid=00000000000000000000000000000000,bHidden=False,bNotConnectable=False,bDefaultValueIsReadOnly=False,bDefaultValueIsIgnored=False,bAdvancedView=False,bOrphanedPin=False,)
+		   CustomProperties Pin (PinId=29D905254A0CAA0A449EAD866D3F2AB5,PinName="OutSideLengthX",PinToolTip="Out Side Length X\nFloat",Direction="EGPD_Output",PinType.PinCategory="float",PinType.PinSubCategory="",PinType.PinSubCategoryObject=None,PinType.PinSubCategoryMemberReference=(),PinType.PinValueType=(),PinType.ContainerType=None,PinType.bIsMap=False,PinType.bIsSet=False,PinType.bIsArray=False,PinType.bIsReference=False,PinType.bIsConst=False,PinType.bIsWeakPointer=False,DefaultValue="0.0",AutogeneratedDefaultValue="0.0",LinkedTo=(K2Node_CallFunction_2060 06EEFCC0478897C80F6ED1BB6D90D60B,),PersistentGuid=00000000000000000000000000000000,bHidden=False,bNotConnectable=False,bDefaultValueIsReadOnly=False,bDefaultValueIsIgnored=False,bAdvancedView=False,bOrphanedPin=False,)
+		   CustomProperties Pin (PinId=5990DFBD4CEA02FAF1680FA21EFD3BB1,PinName="OutSideLengthY",PinToolTip="Out Side Length Y\nFloat",Direction="EGPD_Output",PinType.PinCategory="float",PinType.PinSubCategory="",PinType.PinSubCategoryObject=None,PinType.PinSubCategoryMemberReference=(),PinType.PinValueType=(),PinType.ContainerType=None,PinType.bIsMap=False,PinType.bIsSet=False,PinType.bIsArray=False,PinType.bIsReference=False,PinType.bIsConst=False,PinType.bIsWeakPointer=False,DefaultValue="0.0",AutogeneratedDefaultValue="0.0",LinkedTo=(K2Node_CallFunction_2059 06EEFCC0478897C80F6ED1BB6D90D60B,),PersistentGuid=00000000000000000000000000000000,bHidden=False,bNotConnectable=False,bDefaultValueIsReadOnly=False,bDefaultValueIsIgnored=False,bAdvancedView=False,bOrphanedPin=False,)
+		   CustomProperties Pin (PinId=A73AFD834CC2F11C46A9D4B8495EE266,PinName="bDebugDraw",PinToolTip="Debug Draw\nBoolean",PinType.PinCategory="bool",PinType.PinSubCategory="",PinType.PinSubCategoryObject=None,PinType.PinSubCategoryMemberReference=(),PinType.PinValueType=(),PinType.ContainerType=None,PinType.bIsMap=False,PinType.bIsSet=False,PinType.bIsArray=False,PinType.bIsReference=False,PinType.bIsConst=False,PinType.bIsWeakPointer=False,DefaultValue="true",AutogeneratedDefaultValue="false",PersistentGuid=00000000000000000000000000000000,bHidden=False,bNotConnectable=False,bDefaultValueIsReadOnly=False,bDefaultValueIsIgnored=False,bAdvancedView=False,bOrphanedPin=False,)
+		End Object
+		Begin Object Class=/Script/BlueprintGraph.K2Node_VariableGet Name="K2Node_VariableGet_43"
+		   VariableReference=(MemberName="SampleSurface",MemberGuid=461EECB84B8308CA469319B7DFF82820,bSelfContext=True)
+		   NodePosX=3632
+		   NodePosY=-1696
+		   ErrorType=1
+		   NodeGuid=303249A2451A0863E05024A94AD35B84
+		   CustomProperties Pin (PinId=764BD867405B9674CC281B84A04BD8D4,PinName="SampleSurface",Direction="EGPD_Output",PinType.PinCategory="struct",PinType.PinSubCategory="",PinType.PinSubCategoryObject=ScriptStruct'/Script/CoreUObject.Vector',PinType.PinSubCategoryMemberReference=(),PinType.PinValueType=(),PinType.ContainerType=None,PinType.bIsMap=False,PinType.bIsSet=False,PinType.bIsArray=False,PinType.bIsReference=False,PinType.bIsConst=False,PinType.bIsWeakPointer=False,DefaultValue="0, 0, 0",AutogeneratedDefaultValue="0, 0, 0",LinkedTo=(K2Node_CallFunction_3 76C031DD4BCFAD63300ACEB272F06EA1,),PersistentGuid=00000000000000000000000000000000,bHidden=False,bNotConnectable=False,bDefaultValueIsReadOnly=False,bDefaultValueIsIgnored=False,bAdvancedView=False,bOrphanedPin=False,)
+		   CustomProperties Pin (PinId=D5631E82413D5464F7442D93A40A828B,PinName="self",PinFriendlyName=NSLOCTEXT("K2Node", "Target", "Target"),PinType.PinCategory="object",PinType.PinSubCategory="",PinType.PinSubCategoryObject=BlueprintGeneratedClass'/Game/Blueprints/SteamVR_Pawn_Standing.SteamVR_Pawn_Standing_C',PinType.PinSubCategoryMemberReference=(),PinType.PinValueType=(),PinType.ContainerType=None,PinType.bIsMap=False,PinType.bIsSet=False,PinType.bIsArray=False,PinType.bIsReference=False,PinType.bIsConst=False,PinType.bIsWeakPointer=False,PersistentGuid=00000000000000000000000000000000,bHidden=True,bNotConnectable=False,bDefaultValueIsReadOnly=False,bDefaultValueIsIgnored=False,bAdvancedView=False,bOrphanedPin=False,)
+		End Object
+		Begin Object Class=/Script/BlueprintGraph.K2Node_CallFunction Name="K2Node_CallFunction_994"
+		   FunctionReference=(MemberParent=Class'/Script/Engine.KismetSystemLibrary',MemberName="DrawDebugBox")
+		   NodePosX=4512
+		   NodePosY=-1776
+		   EnabledState=DevelopmentOnly
+		   NodeGuid=47EE35E04ECDA4AB6F9A8AB470A3AFD8
+		   CustomProperties Pin (PinId=60CD775D47E28A58610F9DBBA97E8501,PinName="execute",PinToolTip="\nExec",PinType.PinCategory="exec",PinType.PinSubCategory="",PinType.PinSubCategoryObject=None,PinType.PinSubCategoryMemberReference=(),PinType.PinValueType=(),PinType.ContainerType=None,PinType.bIsMap=False,PinType.bIsSet=False,PinType.bIsArray=False,PinType.bIsReference=False,PinType.bIsConst=False,PinType.bIsWeakPointer=False,LinkedTo=(K2Node_CallFunction_3 D8DDC8B0419E7157C72694920CD01FB1,),PersistentGuid=00000000000000000000000000000000,bHidden=False,bNotConnectable=False,bDefaultValueIsReadOnly=False,bDefaultValueIsIgnored=False,bAdvancedView=False,bOrphanedPin=False,)
+		   CustomProperties Pin (PinId=6CEFB95348EF6DB9DCDD27B7C3154B72,PinName="then",PinToolTip="\nExec",Direction="EGPD_Output",PinType.PinCategory="exec",PinType.PinSubCategory="",PinType.PinSubCategoryObject=None,PinType.PinSubCategoryMemberReference=(),PinType.PinValueType=(),PinType.ContainerType=None,PinType.bIsMap=False,PinType.bIsSet=False,PinType.bIsArray=False,PinType.bIsReference=False,PinType.bIsConst=False,PinType.bIsWeakPointer=False,PersistentGuid=00000000000000000000000000000000,bHidden=False,bNotConnectable=False,bDefaultValueIsReadOnly=False,bDefaultValueIsIgnored=False,bAdvancedView=False,bOrphanedPin=False,)
+		   CustomProperties Pin (PinId=E249F42B4DB1BFACC99E76A81D521B62,PinName="self",PinFriendlyName=NSLOCTEXT("K2Node", "Target", "Target"),PinToolTip="Target\nKismet System Library Object Reference",PinType.PinCategory="object",PinType.PinSubCategory="",PinType.PinSubCategoryObject=Class'/Script/Engine.KismetSystemLibrary',PinType.PinSubCategoryMemberReference=(),PinType.PinValueType=(),PinType.ContainerType=None,PinType.bIsMap=False,PinType.bIsSet=False,PinType.bIsArray=False,PinType.bIsReference=False,PinType.bIsConst=False,PinType.bIsWeakPointer=False,DefaultObject="/Script/Engine.Default__KismetSystemLibrary",PersistentGuid=00000000000000000000000000000000,bHidden=True,bNotConnectable=False,bDefaultValueIsReadOnly=False,bDefaultValueIsIgnored=False,bAdvancedView=False,bOrphanedPin=False,)
+		   CustomProperties Pin (PinId=86E38E1C4191B312DC15EB9B5D6A6EC1,PinName="WorldContextObject",PinToolTip="World Context Object\nObject Object Reference",PinType.PinCategory="object",PinType.PinSubCategory="",PinType.PinSubCategoryObject=Class'/Script/CoreUObject.Object',PinType.PinSubCategoryMemberReference=(),PinType.PinValueType=(),PinType.ContainerType=None,PinType.bIsMap=False,PinType.bIsSet=False,PinType.bIsArray=False,PinType.bIsReference=False,PinType.bIsConst=False,PinType.bIsWeakPointer=False,PersistentGuid=00000000000000000000000000000000,bHidden=True,bNotConnectable=False,bDefaultValueIsReadOnly=False,bDefaultValueIsIgnored=False,bAdvancedView=False,bOrphanedPin=False,)
+		   CustomProperties Pin (PinId=A49D89BF47E998D9C831EBB713671243,PinName="Center",PinToolTip="Center\nVector ",PinType.PinCategory="struct",PinType.PinSubCategory="",PinType.PinSubCategoryObject=ScriptStruct'/Script/CoreUObject.Vector',PinType.PinSubCategoryMemberReference=(),PinType.PinValueType=(),PinType.ContainerType=None,PinType.bIsMap=False,PinType.bIsSet=False,PinType.bIsArray=False,PinType.bIsReference=False,PinType.bIsConst=True,PinType.bIsWeakPointer=False,DefaultValue="0, 0, 0",AutogeneratedDefaultValue="0, 0, 0",PersistentGuid=00000000000000000000000000000000,bHidden=False,bNotConnectable=False,bDefaultValueIsReadOnly=False,bDefaultValueIsIgnored=False,bAdvancedView=False,bOrphanedPin=False,)
+		   CustomProperties Pin (PinId=E96F84744CC17AD3CFE69FBF786DDA86,PinName="Extent",PinToolTip="Extent\nVector ",PinType.PinCategory="struct",PinType.PinSubCategory="",PinType.PinSubCategoryObject=ScriptStruct'/Script/CoreUObject.Vector',PinType.PinSubCategoryMemberReference=(),PinType.PinValueType=(),PinType.ContainerType=None,PinType.bIsMap=False,PinType.bIsSet=False,PinType.bIsArray=False,PinType.bIsReference=False,PinType.bIsConst=False,PinType.bIsWeakPointer=False,DefaultValue="0, 0, 0",AutogeneratedDefaultValue="0, 0, 0",SubPins=(K2Node_CallFunction_994 951BB00F41FC65DB32B99B89BA30A88E,K2Node_CallFunction_994 00073CE140652C9D97B18B814C17CB1F,K2Node_CallFunction_994 F36A6CDC4421E7210B7C94BC997F0356,),PersistentGuid=00000000000000000000000000000000,bHidden=True,bNotConnectable=False,bDefaultValueIsReadOnly=False,bDefaultValueIsIgnored=False,bAdvancedView=False,bOrphanedPin=False,)
+		   CustomProperties Pin (PinId=951BB00F41FC65DB32B99B89BA30A88E,PinName="Extent_X",PinFriendlyName="Extent X",PinToolTip="Extent X\nFloat",PinType.PinCategory="float",PinType.PinSubCategory="",PinType.PinSubCategoryObject=None,PinType.PinSubCategoryMemberReference=(),PinType.PinValueType=(),PinType.ContainerType=None,PinType.bIsMap=False,PinType.bIsSet=False,PinType.bIsArray=False,PinType.bIsReference=False,PinType.bIsConst=False,PinType.bIsWeakPointer=False,DefaultValue="0.0",AutogeneratedDefaultValue="0.0",LinkedTo=(K2Node_CallFunction_2060 542EB47B422A16148B175880EB1997D3,),ParentPin=K2Node_CallFunction_994 E96F84744CC17AD3CFE69FBF786DDA86,PersistentGuid=00000000000000000000000000000000,bHidden=False,bNotConnectable=False,bDefaultValueIsReadOnly=False,bDefaultValueIsIgnored=False,bAdvancedView=False,bOrphanedPin=False,)
+		   CustomProperties Pin (PinId=00073CE140652C9D97B18B814C17CB1F,PinName="Extent_Y",PinFriendlyName="Extent Y",PinToolTip="Extent Y\nFloat",PinType.PinCategory="float",PinType.PinSubCategory="",PinType.PinSubCategoryObject=None,PinType.PinSubCategoryMemberReference=(),PinType.PinValueType=(),PinType.ContainerType=None,PinType.bIsMap=False,PinType.bIsSet=False,PinType.bIsArray=False,PinType.bIsReference=False,PinType.bIsConst=False,PinType.bIsWeakPointer=False,DefaultValue="0.0",AutogeneratedDefaultValue="0.0",LinkedTo=(K2Node_CallFunction_2059 542EB47B422A16148B175880EB1997D3,),ParentPin=K2Node_CallFunction_994 E96F84744CC17AD3CFE69FBF786DDA86,PersistentGuid=00000000000000000000000000000000,bHidden=False,bNotConnectable=False,bDefaultValueIsReadOnly=False,bDefaultValueIsIgnored=False,bAdvancedView=False,bOrphanedPin=False,)
+		   CustomProperties Pin (PinId=F36A6CDC4421E7210B7C94BC997F0356,PinName="Extent_Z",PinFriendlyName="Extent Z",PinToolTip="Extent Z\nFloat",PinType.PinCategory="float",PinType.PinSubCategory="",PinType.PinSubCategoryObject=None,PinType.PinSubCategoryMemberReference=(),PinType.PinValueType=(),PinType.ContainerType=None,PinType.bIsMap=False,PinType.bIsSet=False,PinType.bIsArray=False,PinType.bIsReference=False,PinType.bIsConst=False,PinType.bIsWeakPointer=False,DefaultValue="0.0",AutogeneratedDefaultValue="0.0",ParentPin=K2Node_CallFunction_994 E96F84744CC17AD3CFE69FBF786DDA86,PersistentGuid=00000000000000000000000000000000,bHidden=False,bNotConnectable=False,bDefaultValueIsReadOnly=False,bDefaultValueIsIgnored=False,bAdvancedView=False,bOrphanedPin=False,)
+		   CustomProperties Pin (PinId=871095F64585DBFFB0310486606948FE,PinName="LineColor",PinToolTip="Line Color\nLinear Color Structure",PinType.PinCategory="struct",PinType.PinSubCategory="",PinType.PinSubCategoryObject=ScriptStruct'/Script/CoreUObject.LinearColor',PinType.PinSubCategoryMemberReference=(),PinType.PinValueType=(),PinType.ContainerType=None,PinType.bIsMap=False,PinType.bIsSet=False,PinType.bIsArray=False,PinType.bIsReference=False,PinType.bIsConst=False,PinType.bIsWeakPointer=False,DefaultValue="(R=1.000000,G=0.000000,B=0.000000,A=1.000000)",PersistentGuid=00000000000000000000000000000000,bHidden=False,bNotConnectable=False,bDefaultValueIsReadOnly=False,bDefaultValueIsIgnored=False,bAdvancedView=False,bOrphanedPin=False,)
+		   CustomProperties Pin (PinId=3749C2554A8C78750609DD96E362A91B,PinName="Rotation",PinToolTip="Rotation\nRotator ",PinType.PinCategory="struct",PinType.PinSubCategory="",PinType.PinSubCategoryObject=ScriptStruct'/Script/CoreUObject.Rotator',PinType.PinSubCategoryMemberReference=(),PinType.PinValueType=(),PinType.ContainerType=None,PinType.bIsMap=False,PinType.bIsSet=False,PinType.bIsArray=False,PinType.bIsReference=False,PinType.bIsConst=True,PinType.bIsWeakPointer=False,DefaultValue="0, 0, 0",AutogeneratedDefaultValue="0, 0, 0",LinkedTo=(K2Node_Knot_84 84165C3B4A40061E7725728F88817923,),PersistentGuid=00000000000000000000000000000000,bHidden=False,bNotConnectable=False,bDefaultValueIsReadOnly=False,bDefaultValueIsIgnored=False,bAdvancedView=False,bOrphanedPin=False,)
+		   CustomProperties Pin (PinId=CB0E4BF34F77908C09ACE49EDEDFFC69,PinName="Duration",PinToolTip="Duration\nFloat",PinType.PinCategory="float",PinType.PinSubCategory="",PinType.PinSubCategoryObject=None,PinType.PinSubCategoryMemberReference=(),PinType.PinValueType=(),PinType.ContainerType=None,PinType.bIsMap=False,PinType.bIsSet=False,PinType.bIsArray=False,PinType.bIsReference=False,PinType.bIsConst=False,PinType.bIsWeakPointer=False,DefaultValue="5",AutogeneratedDefaultValue="0.000000",PersistentGuid=00000000000000000000000000000000,bHidden=False,bNotConnectable=False,bDefaultValueIsReadOnly=False,bDefaultValueIsIgnored=False,bAdvancedView=False,bOrphanedPin=False,)
+		   CustomProperties Pin (PinId=296488AF475D5417324221A13967232C,PinName="Thickness",PinToolTip="Thickness\nFloat",PinType.PinCategory="float",PinType.PinSubCategory="",PinType.PinSubCategoryObject=None,PinType.PinSubCategoryMemberReference=(),PinType.PinValueType=(),PinType.ContainerType=None,PinType.bIsMap=False,PinType.bIsSet=False,PinType.bIsArray=False,PinType.bIsReference=False,PinType.bIsConst=False,PinType.bIsWeakPointer=False,DefaultValue="5",AutogeneratedDefaultValue="0.000000",PersistentGuid=00000000000000000000000000000000,bHidden=False,bNotConnectable=False,bDefaultValueIsReadOnly=False,bDefaultValueIsIgnored=False,bAdvancedView=False,bOrphanedPin=False,)
+		End Object
+		Begin Object Class=/Script/BlueprintGraph.K2Node_CallFunction Name="K2Node_CallFunction_2060"
+		   bIsPureFunc=True
+		   FunctionReference=(MemberParent=Class'/Script/Engine.KismetMathLibrary',MemberName="Divide_FloatFloat")
+		   NodePosX=4256
+		   NodePosY=-1671
+		   bCommentBubblePinned=True
+		   NodeGuid=782AC20F44449B0E24BE4BB765D5E7B8
+		   CustomProperties Pin (PinId=B1D35F4F45D4A34E6FEAA4AE77994363,PinName="self",PinFriendlyName=NSLOCTEXT("K2Node", "Target", "Target"),PinToolTip="Target\nKismet Math Library Object Reference",PinType.PinCategory="object",PinType.PinSubCategory="",PinType.PinSubCategoryObject=Class'/Script/Engine.KismetMathLibrary',PinType.PinSubCategoryMemberReference=(),PinType.PinValueType=(),PinType.ContainerType=None,PinType.bIsMap=False,PinType.bIsSet=False,PinType.bIsArray=False,PinType.bIsReference=False,PinType.bIsConst=False,PinType.bIsWeakPointer=False,DefaultObject="/Script/Engine.Default__KismetMathLibrary",PersistentGuid=00000000000000000000000000000000,bHidden=True,bNotConnectable=False,bDefaultValueIsReadOnly=False,bDefaultValueIsIgnored=False,bAdvancedView=False,bOrphanedPin=False,)
+		   CustomProperties Pin (PinId=06EEFCC0478897C80F6ED1BB6D90D60B,PinName="A",PinToolTip="A\nFloat",PinType.PinCategory="float",PinType.PinSubCategory="",PinType.PinSubCategoryObject=None,PinType.PinSubCategoryMemberReference=(),PinType.PinValueType=(),PinType.ContainerType=None,PinType.bIsMap=False,PinType.bIsSet=False,PinType.bIsArray=False,PinType.bIsReference=False,PinType.bIsConst=False,PinType.bIsWeakPointer=False,DefaultValue="0.0",AutogeneratedDefaultValue="0.0",LinkedTo=(K2Node_CallFunction_3 29D905254A0CAA0A449EAD866D3F2AB5,),PersistentGuid=00000000000000000000000000000000,bHidden=False,bNotConnectable=False,bDefaultValueIsReadOnly=False,bDefaultValueIsIgnored=False,bAdvancedView=False,bOrphanedPin=False,)
+		   CustomProperties Pin (PinId=6581321D4F0BC79197F13F87A41E6C39,PinName="B",PinToolTip="B\nFloat",PinType.PinCategory="float",PinType.PinSubCategory="",PinType.PinSubCategoryObject=None,PinType.PinSubCategoryMemberReference=(),PinType.PinValueType=(),PinType.ContainerType=None,PinType.bIsMap=False,PinType.bIsSet=False,PinType.bIsArray=False,PinType.bIsReference=False,PinType.bIsConst=False,PinType.bIsWeakPointer=False,DefaultValue="2",AutogeneratedDefaultValue="1.000000",PersistentGuid=00000000000000000000000000000000,bHidden=False,bNotConnectable=False,bDefaultValueIsReadOnly=False,bDefaultValueIsIgnored=False,bAdvancedView=False,bOrphanedPin=False,)
+		   CustomProperties Pin (PinId=542EB47B422A16148B175880EB1997D3,PinName="ReturnValue",PinToolTip="Return Value\nFloat",Direction="EGPD_Output",PinType.PinCategory="float",PinType.PinSubCategory="",PinType.PinSubCategoryObject=None,PinType.PinSubCategoryMemberReference=(),PinType.PinValueType=(),PinType.ContainerType=None,PinType.bIsMap=False,PinType.bIsSet=False,PinType.bIsArray=False,PinType.bIsReference=False,PinType.bIsConst=False,PinType.bIsWeakPointer=False,DefaultValue="0.0",AutogeneratedDefaultValue="0.0",LinkedTo=(K2Node_CallFunction_994 951BB00F41FC65DB32B99B89BA30A88E,),PersistentGuid=00000000000000000000000000000000,bHidden=False,bNotConnectable=False,bDefaultValueIsReadOnly=False,bDefaultValueIsIgnored=False,bAdvancedView=False,bOrphanedPin=False,)
+		End Object
+		Begin Object Class=/Script/BlueprintGraph.K2Node_CallFunction Name="K2Node_CallFunction_2059"
+		   bIsPureFunc=True
+		   FunctionReference=(MemberParent=Class'/Script/Engine.KismetMathLibrary',MemberName="Divide_FloatFloat")
+		   NodePosX=4256
+		   NodePosY=-1584
+		   NodeGuid=E09571CB47DEC7D6C00A0E81C3B20348
+		   CustomProperties Pin (PinId=B1D35F4F45D4A34E6FEAA4AE77994363,PinName="self",PinFriendlyName=NSLOCTEXT("K2Node", "Target", "Target"),PinToolTip="Target\nKismet Math Library Object Reference",PinType.PinCategory="object",PinType.PinSubCategory="",PinType.PinSubCategoryObject=Class'/Script/Engine.KismetMathLibrary',PinType.PinSubCategoryMemberReference=(),PinType.PinValueType=(),PinType.ContainerType=None,PinType.bIsMap=False,PinType.bIsSet=False,PinType.bIsArray=False,PinType.bIsReference=False,PinType.bIsConst=False,PinType.bIsWeakPointer=False,DefaultObject="/Script/Engine.Default__KismetMathLibrary",PersistentGuid=00000000000000000000000000000000,bHidden=True,bNotConnectable=False,bDefaultValueIsReadOnly=False,bDefaultValueIsIgnored=False,bAdvancedView=False,bOrphanedPin=False,)
+		   CustomProperties Pin (PinId=06EEFCC0478897C80F6ED1BB6D90D60B,PinName="A",PinToolTip="A\nFloat",PinType.PinCategory="float",PinType.PinSubCategory="",PinType.PinSubCategoryObject=None,PinType.PinSubCategoryMemberReference=(),PinType.PinValueType=(),PinType.ContainerType=None,PinType.bIsMap=False,PinType.bIsSet=False,PinType.bIsArray=False,PinType.bIsReference=False,PinType.bIsConst=False,PinType.bIsWeakPointer=False,DefaultValue="0.0",AutogeneratedDefaultValue="0.0",LinkedTo=(K2Node_CallFunction_3 5990DFBD4CEA02FAF1680FA21EFD3BB1,),PersistentGuid=00000000000000000000000000000000,bHidden=False,bNotConnectable=False,bDefaultValueIsReadOnly=False,bDefaultValueIsIgnored=False,bAdvancedView=False,bOrphanedPin=False,)
+		   CustomProperties Pin (PinId=6581321D4F0BC79197F13F87A41E6C39,PinName="B",PinToolTip="B\nFloat",PinType.PinCategory="float",PinType.PinSubCategory="",PinType.PinSubCategoryObject=None,PinType.PinSubCategoryMemberReference=(),PinType.PinValueType=(),PinType.ContainerType=None,PinType.bIsMap=False,PinType.bIsSet=False,PinType.bIsArray=False,PinType.bIsReference=False,PinType.bIsConst=False,PinType.bIsWeakPointer=False,DefaultValue="2",AutogeneratedDefaultValue="1.000000",PersistentGuid=00000000000000000000000000000000,bHidden=False,bNotConnectable=False,bDefaultValueIsReadOnly=False,bDefaultValueIsIgnored=False,bAdvancedView=False,bOrphanedPin=False,)
+		   CustomProperties Pin (PinId=542EB47B422A16148B175880EB1997D3,PinName="ReturnValue",PinToolTip="Return Value\nFloat",Direction="EGPD_Output",PinType.PinCategory="float",PinType.PinSubCategory="",PinType.PinSubCategoryObject=None,PinType.PinSubCategoryMemberReference=(),PinType.PinValueType=(),PinType.ContainerType=None,PinType.bIsMap=False,PinType.bIsSet=False,PinType.bIsArray=False,PinType.bIsReference=False,PinType.bIsConst=False,PinType.bIsWeakPointer=False,DefaultValue="0.0",AutogeneratedDefaultValue="0.0",LinkedTo=(K2Node_CallFunction_994 00073CE140652C9D97B18B814C17CB1F,),PersistentGuid=00000000000000000000000000000000,bHidden=False,bNotConnectable=False,bDefaultValueIsReadOnly=False,bDefaultValueIsIgnored=False,bAdvancedView=False,bOrphanedPin=False,)
+		End Object
+		Begin Object Class=/Script/BlueprintGraph.K2Node_Knot Name="K2Node_Knot_85"
+		   NodePosX=4272
+		   NodePosY=-1472
+		   NodeGuid=4D7D28AC4E4210815AB4B7958DAB48A7
+		   CustomProperties Pin (PinId=365F26FF45B3A67424DF5F9578C950B3,PinName="InputPin",PinType.PinCategory="struct",PinType.PinSubCategory="",PinType.PinSubCategoryObject=ScriptStruct'/Script/CoreUObject.Rotator',PinType.PinSubCategoryMemberReference=(),PinType.PinValueType=(),PinType.ContainerType=None,PinType.bIsMap=False,PinType.bIsSet=False,PinType.bIsArray=False,PinType.bIsReference=False,PinType.bIsConst=False,PinType.bIsWeakPointer=False,LinkedTo=(K2Node_CallFunction_3 C4FA178B4D01CFF46CD2DBB72D723B71,),PersistentGuid=00000000000000000000000000000000,bHidden=False,bNotConnectable=False,bDefaultValueIsReadOnly=False,bDefaultValueIsIgnored=True,bAdvancedView=False,bOrphanedPin=False,)
+		   CustomProperties Pin (PinId=84165C3B4A40061E7725728F88817923,PinName="OutputPin",Direction="EGPD_Output",PinType.PinCategory="struct",PinType.PinSubCategory="",PinType.PinSubCategoryObject=ScriptStruct'/Script/CoreUObject.Rotator',PinType.PinSubCategoryMemberReference=(),PinType.PinValueType=(),PinType.ContainerType=None,PinType.bIsMap=False,PinType.bIsSet=False,PinType.bIsArray=False,PinType.bIsReference=False,PinType.bIsConst=False,PinType.bIsWeakPointer=False,LinkedTo=(K2Node_Knot_84 365F26FF45B3A67424DF5F9578C950B3,),PersistentGuid=00000000000000000000000000000000,bHidden=False,bNotConnectable=False,bDefaultValueIsReadOnly=False,bDefaultValueIsIgnored=False,bAdvancedView=False,bOrphanedPin=False,)
+		End Object
+		Begin Object Class=/Script/BlueprintGraph.K2Node_Knot Name="K2Node_Knot_84"
+		   NodePosX=4416
+		   NodePosY=-1472
+		   NodeGuid=4E71C3364746B8B13C3151948349929C
+		   CustomProperties Pin (PinId=365F26FF45B3A67424DF5F9578C950B3,PinName="InputPin",PinType.PinCategory="struct",PinType.PinSubCategory="",PinType.PinSubCategoryObject=ScriptStruct'/Script/CoreUObject.Rotator',PinType.PinSubCategoryMemberReference=(),PinType.PinValueType=(),PinType.ContainerType=None,PinType.bIsMap=False,PinType.bIsSet=False,PinType.bIsArray=False,PinType.bIsReference=False,PinType.bIsConst=False,PinType.bIsWeakPointer=False,LinkedTo=(K2Node_Knot_85 84165C3B4A40061E7725728F88817923,),PersistentGuid=00000000000000000000000000000000,bHidden=False,bNotConnectable=False,bDefaultValueIsReadOnly=False,bDefaultValueIsIgnored=True,bAdvancedView=False,bOrphanedPin=False,)
+		   CustomProperties Pin (PinId=84165C3B4A40061E7725728F88817923,PinName="OutputPin",Direction="EGPD_Output",PinType.PinCategory="struct",PinType.PinSubCategory="",PinType.PinSubCategoryObject=ScriptStruct'/Script/CoreUObject.Rotator',PinType.PinSubCategoryMemberReference=(),PinType.PinValueType=(),PinType.ContainerType=None,PinType.bIsMap=False,PinType.bIsSet=False,PinType.bIsArray=False,PinType.bIsReference=False,PinType.bIsConst=False,PinType.bIsWeakPointer=False,LinkedTo=(K2Node_CallFunction_994 3749C2554A8C78750609DD96E362A91B,),PersistentGuid=00000000000000000000000000000000,bHidden=False,bNotConnectable=False,bDefaultValueIsReadOnly=False,bDefaultValueIsIgnored=False,bAdvancedView=False,bOrphanedPin=False,)
+		End Object
+		Begin Object Class=/Script/BlueprintGraph.K2Node_InputKey Name="K2Node_InputKey_86"
+		   InputKey=MotionController_Right_Trigger
+		   NodePosX=3344
+		   NodePosY=-1776
+		   NodeGuid=25C3050D4AE2B0434095278F92EB43CA
+		   CustomProperties Pin (PinId=451F705A45E7D71FD4CDD191C52193E1,PinName="Pressed",Direction="EGPD_Output",PinType.PinCategory="exec",PinType.PinSubCategory="",PinType.PinSubCategoryObject=None,PinType.PinSubCategoryMemberReference=(),PinType.PinValueType=(),PinType.ContainerType=None,PinType.bIsMap=False,PinType.bIsSet=False,PinType.bIsArray=False,PinType.bIsReference=False,PinType.bIsConst=False,PinType.bIsWeakPointer=False,LinkedTo=(K2Node_CallFunction_3 27A36946461CB0706E8812873937728A,),PersistentGuid=00000000000000000000000000000000,bHidden=False,bNotConnectable=False,bDefaultValueIsReadOnly=False,bDefaultValueIsIgnored=False,bAdvancedView=False,bOrphanedPin=False,)
+		   CustomProperties Pin (PinId=6F1385DB4EF7010345F655BD3744E632,PinName="Released",Direction="EGPD_Output",PinType.PinCategory="exec",PinType.PinSubCategory="",PinType.PinSubCategoryObject=None,PinType.PinSubCategoryMemberReference=(),PinType.PinValueType=(),PinType.ContainerType=None,PinType.bIsMap=False,PinType.bIsSet=False,PinType.bIsArray=False,PinType.bIsReference=False,PinType.bIsConst=False,PinType.bIsWeakPointer=False,PersistentGuid=00000000000000000000000000000000,bHidden=False,bNotConnectable=False,bDefaultValueIsReadOnly=False,bDefaultValueIsIgnored=False,bAdvancedView=False,bOrphanedPin=False,)
+		   CustomProperties Pin (PinId=022C7579454B610F2F2ED0B282488461,PinName="Key",Direction="EGPD_Output",PinType.PinCategory="struct",PinType.PinSubCategory="",PinType.PinSubCategoryObject=ScriptStruct'/Script/InputCore.Key',PinType.PinSubCategoryMemberReference=(),PinType.PinValueType=(),PinType.ContainerType=None,PinType.bIsMap=False,PinType.bIsSet=False,PinType.bIsArray=False,PinType.bIsReference=False,PinType.bIsConst=False,PinType.bIsWeakPointer=False,PersistentGuid=00000000000000000000000000000000,bHidden=False,bNotConnectable=False,bDefaultValueIsReadOnly=False,bDefaultValueIsIgnored=False,bAdvancedView=False,bOrphanedPin=False,)
+		End Object
+
+			END OBJECT
+	[/PARAM]
+[/OBJECT]
+[/EXCERPT:SteamVRC3]
+-->
\ No newline at end of file
