Copying //UE4/Dev-Documentation to //UE4/Samples-Main (Source: //UE4/Dev-Documentation @ 3623665)

diff --git a/Engine/Documentation/Source/Platforms/SteamVR/HowTo/SeatedCamera/HTSteamVRSeatedCamera.INT.udn b/Engine/Documentation/Source/Platforms/SteamVR/HowTo/SeatedCamera/HTSteamVRSeatedCamera.INT.udn
new file mode 100644
--- /dev/null
+++ b/Engine/Documentation/Source/Platforms/SteamVR/HowTo/SeatedCamera/HTSteamVRSeatedCamera.INT.udn
@@ -0,0 +1,212 @@
+Availability: Public
+Crumbs: %ROOT%, Platforms
+Title:Set Up a Seated Camera for Steam VR
+Description:Setting up the UE4 camera to work with a seated Steam VR experience.
+Type: how-to
+Skilllevel:Beginner
+Version:4.17
+Parent:Platforms/SteamVR/HowTo
+Order: 2
+tags:Basics
+topic-image:Platforms\SteamVR\SteamVRDevelopment_topic.png
+
+One of the first things you will need to think about when starting to develop any UE4 powered VR experience to be used with the Steam VR is deciding whether the experience will be a seated one or a standing one. In the following How - To we will take a look at how to setup your UE4 projects VR camera for a seated Steam VR experience.
+
+## Pawn Setup for a Seated Steam VR Experience 
+
+Below you will find instructions on how to setup up any Pawn to work with a seated Steam VR experience.
+
+1. First, open up or create a new Pawn Blueprint and then go to the **Component** section of the **Viewport** tab. From there add the following two components with the following names, making sure that the VRCamera is a child of VRCameraRoot:
+    
+    * **Scene** - VRCameraRoot
+    * **Camera** - VRCamera
+
+	![](HT_SteamVR_Camera_Setup_00.png)
+
+    [region:note]
+    No matter what VR HMD you are using, this is the Epic recommend way to set up a VR camera as it allows you to offset the position of the camera without having to move the actual camera.
+    [/region]
+
+With that completed, the initial setup for the VR Camera is now completed, in the next sections, we will go through setting the needed Blueprint logic for a seated Steam VR experience.
+
+## Blueprint Setup for Seated Steam VR Experience 
+
+Below you will find instructions on how to setup the Blueprint logic for a seated Steam VR experience.  
+
+1. First, open up the Pawn Blueprint if not already open and in the **Event Graph** drag off the **Event Begin Play** node to show the Executable Actions list. In the list search for the **Set Tracking Origin** node and click on it to add it to the Event Graph.
+
+    [REGION:lightbox]
+    [![](HT_Rift_Camera_Setup_09.png)(w:555)](HT_Rift_Camera_Setup_09.png)
+    [/REGION]
+
+    [REGION:caption]
+    Click for full image.
+    [/REGION]
+
+1. The Set Tracking Origin node has two options, **Floor Level** and **Eye Level**. For a seated experience you will need to set the **Origin** of the Set Tracking Origin node to **Eye Level**.
+    
+    [REGION:lightbox]
+    [![](HT_Rift_Camera_Setup_10.png)(w:555)](HT_Rift_Camera_Setup_10.png)
+    [/REGION]
+
+    [REGION:caption]
+    Click for full image.
+    [/REGION]
+
+1. Next, in the **Variables**  section of the **My Blueprint** tab, create a new **Vector** variable called **ViveCameraHeight** and set the **Z** value to **121**.
+
+    ![](HT_Rift_Camera_Setup_05.png)
+
+    [region:note]
+    For a seated Steam VR experience, you will always want to set the height of the camera to match the real world seated height of the user in Centimeters. 
+    [/region]
+
+1. Then, drag off the output of the Set Tracking Origin node and search for the **Set Relative Location** node, selecting the **SetRelativeLocation(VRCameraRoot)** option.
+
+    [REGION:lightbox]
+    [![](HT_Rift_Camera_Setup_11.png)(w:472)](HT_Rift_Camera_Setup_11.png)
+    [/REGION]
+
+    [REGION:caption]
+    Click for full image.
+    [/REGION]
+
+1. Connect the **ViveCameraHeight** variable to the **New Location** input on the Set Relative Location node and then press the Compile button and when completed your Event Graph should look like the following image.
+
+    [REGION:fullwidth]
+    [INCLUDE:Platforms\SteamVR\HowTo\SeatedCamera\#SteamVRSittingVRCamera]
+    [/REGION]
+
+    [region:tip]
+    You can copy the completed Blueprint by clicking on the upper left-hand corner of above and pressing CRTL + C to copy the graph. Once copied, go to your Blueprints event and press CTRL + V to paste it in. 
+    [/region]
+
+1. Drag the Pawn Blueprint from the Content Browser to a level, making sure that it is placed at 0,0,0 in the level.
+    [REGION:lightbox]
+    [![](HT_Rift_Camera_Setup_06.png)(w:555)](HT_Rift_Camera_Setup_06.png)
+    [/REGION]
+
+    [REGION:caption]
+    Click for full image.
+    [/REGION]
+
+1. Select the Pawn Blueprint that was placed in the level and in the **Details** panel under the **Pawn** setting, set the **Auto Possess Player** from **Disabled** to **Player 0**.
+
+    [REGION:lightbox]
+    [![](HT_Rift_Camera_Setup_07.png)(w:225)](HT_Rift_Camera_Setup_07.png)
+    [/REGION]
+
+    [REGION:caption]
+    Click for full image.
+    [/REGION]
+
+1. Finally, go to the **Main Toolbar** and change the **Play Mode** to **VR Preview** and then press the **Play** button. When you put your HTC Vive HMD on and view the level while sitting down, you should see something similar to the following video. 
+
+    [OBJECT:EmbeddedVideo]
+        [PARAMLITERAL:videoid]
+        sO87nTNTATY
+        [/PARAMLITERAL]
+        [PARAMLITERAL:width]
+        100
+        [/PARAMLITERAL]
+        [PARAMLITERAL:height]
+        56.25
+        [/PARAMLITERAL]
+        [PARAMLITERAL:units]
+        %
+        [/PARAMLITERAL]
+        [PARAMLITERAL:type]
+        simple
+        [/PARAMLITERAL]
+    [/OBJECT]
+
+## UE4 Project Downloads
+Below you will find a link to where you can download the UE4 project that was used to create this example. 
+
+* [**Steam VR Seated VR Camera Example Project**](https://epicgames.box.com/s/zxsa658v6t2m917ar1joze2vql5lwprd)
+
+<!--
+[EXCERPT:SteamVRSittingVRCamera]	
+[OBJECT:Snippet]
+	[PARAMLITERAL:id]
+		SteamVRSittingVRCamera
+	[/PARAMLITERAL]
+	[PARAMLITERAL:type]
+		blueprint
+	[/PARAMLITERAL]
+	[PARAM:image]
+		[REGION:lightbox]
+		[![](HT_Rift_Camera_Setup_04.png)(w:1000)](HT_Rift_Camera_Setup_04.png)
+		[/REGION]
+	[/PARAM]
+	[PARAM:snippet]
+			BEGIN OBJECT
+							Begin Object Class=/Script/BlueprintGraph.K2Node_Event Name="K2Node_Event_3"
+			   EventReference=(MemberParent=Class'/Script/Engine.Actor',MemberName="ReceiveBeginPlay")
+			   bOverrideFunction=True
+			   bCommentBubblePinned=True
+			   NodeGuid=71F8712B42D30302CBA9CE84C38862B7
+			   CustomProperties Pin (PinId=63E599D74E38B8776D6E88A4AAB2C2B3,PinName="OutputDelegate",Direction="EGPD_Output",PinType.PinCategory="delegate",PinType.PinSubCategory="",PinType.PinSubCategoryObject=None,PinType.PinSubCategoryMemberReference=(MemberParent=Class'/Script/Engine.Actor',MemberName="ReceiveBeginPlay"),PinType.PinValueType=(),PinType.ContainerType=None,PinType.bIsMap=False,PinType.bIsSet=False,PinType.bIsArray=False,PinType.bIsReference=False,PinType.bIsConst=False,PinType.bIsWeakPointer=False,PersistentGuid=00000000000000000000000000000000,bHidden=False,bNotConnectable=False,bDefaultValueIsReadOnly=False,bDefaultValueIsIgnored=False,bAdvancedView=False,bOrphanedPin=False,)
+			   CustomProperties Pin (PinId=96D13FDB4718AC06F0F132AEEBA98395,PinName="then",Direction="EGPD_Output",PinType.PinCategory="exec",PinType.PinSubCategory="",PinType.PinSubCategoryObject=None,PinType.PinSubCategoryMemberReference=(),PinType.PinValueType=(),PinType.ContainerType=None,PinType.bIsMap=False,PinType.bIsSet=False,PinType.bIsArray=False,PinType.bIsReference=False,PinType.bIsConst=False,PinType.bIsWeakPointer=False,LinkedTo=(K2Node_CallFunction_1 0AF4F2C243DAA259B82A0AB55A91B640,),PersistentGuid=00000000000000000000000000000000,bHidden=False,bNotConnectable=False,bDefaultValueIsReadOnly=False,bDefaultValueIsIgnored=False,bAdvancedView=False,bOrphanedPin=False,)
+			End Object
+			Begin Object Class=/Script/BlueprintGraph.K2Node_CallFunction Name="K2Node_CallFunction_1"
+			   FunctionReference=(MemberParent=Class'/Script/HeadMountedDisplay.HeadMountedDisplayFunctionLibrary',MemberName="SetTrackingOrigin")
+			   NodePosX=224
+			   NodeGuid=401DF7CA48C88D2786ECB8B5C2AC8F46
+			   CustomProperties Pin (PinId=0AF4F2C243DAA259B82A0AB55A91B640,PinName="execute",PinToolTip="\nExec",PinType.PinCategory="exec",PinType.PinSubCategory="",PinType.PinSubCategoryObject=None,PinType.PinSubCategoryMemberReference=(),PinType.PinValueType=(),PinType.ContainerType=None,PinType.bIsMap=False,PinType.bIsSet=False,PinType.bIsArray=False,PinType.bIsReference=False,PinType.bIsConst=False,PinType.bIsWeakPointer=False,LinkedTo=(K2Node_Event_3 96D13FDB4718AC06F0F132AEEBA98395,),PersistentGuid=00000000000000000000000000000000,bHidden=False,bNotConnectable=False,bDefaultValueIsReadOnly=False,bDefaultValueIsIgnored=False,bAdvancedView=False,bOrphanedPin=False,)
+			   CustomProperties Pin (PinId=E9F9736F435EF94B7A4B16A0C9347FFB,PinName="then",PinToolTip="\nExec",Direction="EGPD_Output",PinType.PinCategory="exec",PinType.PinSubCategory="",PinType.PinSubCategoryObject=None,PinType.PinSubCategoryMemberReference=(),PinType.PinValueType=(),PinType.ContainerType=None,PinType.bIsMap=False,PinType.bIsSet=False,PinType.bIsArray=False,PinType.bIsReference=False,PinType.bIsConst=False,PinType.bIsWeakPointer=False,LinkedTo=(K2Node_CallFunction_5 DE9FB716426D0575C24DBD90B0B7F9B0,),PersistentGuid=00000000000000000000000000000000,bHidden=False,bNotConnectable=False,bDefaultValueIsReadOnly=False,bDefaultValueIsIgnored=False,bAdvancedView=False,bOrphanedPin=False,)
+			   CustomProperties Pin (PinId=E21CFDB440B456D7A54A2B991677B221,PinName="self",PinFriendlyName=NSLOCTEXT("K2Node", "Target", "Target"),PinToolTip="Target\nHead Mounted Display Function Library Object Reference",PinType.PinCategory="object",PinType.PinSubCategory="",PinType.PinSubCategoryObject=Class'/Script/HeadMountedDisplay.HeadMountedDisplayFunctionLibrary',PinType.PinSubCategoryMemberReference=(),PinType.PinValueType=(),PinType.ContainerType=None,PinType.bIsMap=False,PinType.bIsSet=False,PinType.bIsArray=False,PinType.bIsReference=False,PinType.bIsConst=False,PinType.bIsWeakPointer=False,DefaultObject="/Script/HeadMountedDisplay.Default__HeadMountedDisplayFunctionLibrary",PersistentGuid=00000000000000000000000000000000,bHidden=True,bNotConnectable=False,bDefaultValueIsReadOnly=False,bDefaultValueIsIgnored=False,bAdvancedView=False,bOrphanedPin=False,)
+			   CustomProperties Pin (PinId=D3A2C5724D17455C5C1DCBAA7E724282,PinName="Origin",PinToolTip="Origin\nEHMDTrackingOrigin Enum",PinType.PinCategory="byte",PinType.PinSubCategory="",PinType.PinSubCategoryObject=Enum'/Script/HeadMountedDisplay.EHMDTrackingOrigin',PinType.PinSubCategoryMemberReference=(),PinType.PinValueType=(),PinType.ContainerType=None,PinType.bIsMap=False,PinType.bIsSet=False,PinType.bIsArray=False,PinType.bIsReference=False,PinType.bIsConst=False,PinType.bIsWeakPointer=False,DefaultValue="Eye",PersistentGuid=00000000000000000000000000000000,bHidden=False,bNotConnectable=False,bDefaultValueIsReadOnly=False,bDefaultValueIsIgnored=False,bAdvancedView=False,bOrphanedPin=False,)
+			End Object
+			Begin Object Class=/Script/BlueprintGraph.K2Node_VariableGet Name="K2Node_VariableGet_0"
+			   VariableReference=(MemberName="VRCameraRoot",bSelfContext=True)
+			   NodePosX=720
+			   NodePosY=48
+			   NodeGuid=DE98398C46FA741DD0ECCFA88342B5DC
+			   CustomProperties Pin (PinId=D3E17C4549E0A126BF8A2D97653C64D5,PinName="VRCameraRoot",Direction="EGPD_Output",PinType.PinCategory="object",PinType.PinSubCategory="",PinType.PinSubCategoryObject=Class'/Script/Engine.SceneComponent',PinType.PinSubCategoryMemberReference=(),PinType.PinValueType=(),PinType.ContainerType=None,PinType.bIsMap=False,PinType.bIsSet=False,PinType.bIsArray=False,PinType.bIsReference=False,PinType.bIsConst=False,PinType.bIsWeakPointer=False,LinkedTo=(K2Node_CallFunction_5 560BA9474CC7FF8D01863CB899932E82,),PersistentGuid=00000000000000000000000000000000,bHidden=False,bNotConnectable=False,bDefaultValueIsReadOnly=False,bDefaultValueIsIgnored=False,bAdvancedView=False,bOrphanedPin=False,)
+			   CustomProperties Pin (PinId=8B50EC6646D9AB54EDFBCAAA7B690B9B,PinName="self",PinFriendlyName=NSLOCTEXT("K2Node", "Target", "Target"),PinType.PinCategory="object",PinType.PinSubCategory="",PinType.PinSubCategoryObject=BlueprintGeneratedClass'/Game/SteamVR_Pawn_Sitting.SteamVR_Pawn_Sitting_C',PinType.PinSubCategoryMemberReference=(),PinType.PinValueType=(),PinType.ContainerType=None,PinType.bIsMap=False,PinType.bIsSet=False,PinType.bIsArray=False,PinType.bIsReference=False,PinType.bIsConst=False,PinType.bIsWeakPointer=False,PersistentGuid=00000000000000000000000000000000,bHidden=True,bNotConnectable=False,bDefaultValueIsReadOnly=False,bDefaultValueIsIgnored=False,bAdvancedView=False,bOrphanedPin=False,)
+			End Object
+			Begin Object Class=/Script/BlueprintGraph.K2Node_CallFunction Name="K2Node_CallFunction_5"
+			   FunctionReference=(MemberParent=Class'/Script/Engine.SceneComponent',MemberName="K2_SetRelativeLocation")
+			   NodePosX=896
+			   NodePosY=-16
+			   NodeGuid=9C212BB84E9A16ED4D937A898C2B6367
+			   CustomProperties Pin (PinId=DE9FB716426D0575C24DBD90B0B7F9B0,PinName="execute",PinToolTip="\nExec",PinType.PinCategory="exec",PinType.PinSubCategory="",PinType.PinSubCategoryObject=None,PinType.PinSubCategoryMemberReference=(),PinType.PinValueType=(),PinType.ContainerType=None,PinType.bIsMap=False,PinType.bIsSet=False,PinType.bIsArray=False,PinType.bIsReference=False,PinType.bIsConst=False,PinType.bIsWeakPointer=False,LinkedTo=(K2Node_CallFunction_1 E9F9736F435EF94B7A4B16A0C9347FFB,),PersistentGuid=00000000000000000000000000000000,bHidden=False,bNotConnectable=False,bDefaultValueIsReadOnly=False,bDefaultValueIsIgnored=False,bAdvancedView=False,bOrphanedPin=False,)
+			   CustomProperties Pin (PinId=E5FB25D74A81ACC9C7E06593A4BC6E98,PinName="then",PinToolTip="\nExec",Direction="EGPD_Output",PinType.PinCategory="exec",PinType.PinSubCategory="",PinType.PinSubCategoryObject=None,PinType.PinSubCategoryMemberReference=(),PinType.PinValueType=(),PinType.ContainerType=None,PinType.bIsMap=False,PinType.bIsSet=False,PinType.bIsArray=False,PinType.bIsReference=False,PinType.bIsConst=False,PinType.bIsWeakPointer=False,PersistentGuid=00000000000000000000000000000000,bHidden=False,bNotConnectable=False,bDefaultValueIsReadOnly=False,bDefaultValueIsIgnored=False,bAdvancedView=False,bOrphanedPin=False,)
+			   CustomProperties Pin (PinId=560BA9474CC7FF8D01863CB899932E82,PinName="self",PinFriendlyName=NSLOCTEXT("K2Node", "Target", "Target"),PinToolTip="Target\nScene Component Object Reference",PinType.PinCategory="object",PinType.PinSubCategory="",PinType.PinSubCategoryObject=Class'/Script/Engine.SceneComponent',PinType.PinSubCategoryMemberReference=(),PinType.PinValueType=(),PinType.ContainerType=None,PinType.bIsMap=False,PinType.bIsSet=False,PinType.bIsArray=False,PinType.bIsReference=False,PinType.bIsConst=False,PinType.bIsWeakPointer=False,LinkedTo=(K2Node_VariableGet_0 D3E17C4549E0A126BF8A2D97653C64D5,),PersistentGuid=00000000000000000000000000000000,bHidden=False,bNotConnectable=False,bDefaultValueIsReadOnly=False,bDefaultValueIsIgnored=False,bAdvancedView=False,bOrphanedPin=False,)
+			   CustomProperties Pin (PinId=4751BC714842E8CEC1B7A294186125C2,PinName="NewLocation",PinToolTip="New Location\nVector \n\nNew location of the component relative to its parent.",PinType.PinCategory="struct",PinType.PinSubCategory="",PinType.PinSubCategoryObject=ScriptStruct'/Script/CoreUObject.Vector',PinType.PinSubCategoryMemberReference=(),PinType.PinValueType=(),PinType.ContainerType=None,PinType.bIsMap=False,PinType.bIsSet=False,PinType.bIsArray=False,PinType.bIsReference=False,PinType.bIsConst=False,PinType.bIsWeakPointer=False,DefaultValue="0, 0, 0",AutogeneratedDefaultValue="0, 0, 0",LinkedTo=(K2Node_VariableGet_47 78E971AE41B9BC4AA335A6BCC9E5AE69,),PersistentGuid=00000000000000000000000000000000,bHidden=False,bNotConnectable=False,bDefaultValueIsReadOnly=False,bDefaultValueIsIgnored=False,bAdvancedView=False,bOrphanedPin=False,)
+			   CustomProperties Pin (PinId=42E17DA8407F1D6AE61E5AB8CE56D509,PinName="bSweep",PinToolTip="Sweep\nBoolean\n\nWhether we sweep to the destination location, triggering overlaps along the way and stopping short of the target if blocked by something. Only the root component is swept and checked for blocking collision, child components move without sweeping. If collision is off, this has no effect.",PinType.PinCategory="bool",PinType.PinSubCategory="",PinType.PinSubCategoryObject=None,PinType.PinSubCategoryMemberReference=(),PinType.PinValueType=(),PinType.ContainerType=None,PinType.bIsMap=False,PinType.bIsSet=False,PinType.bIsArray=False,PinType.bIsReference=False,PinType.bIsConst=False,PinType.bIsWeakPointer=False,DefaultValue="false",AutogeneratedDefaultValue="false",PersistentGuid=00000000000000000000000000000000,bHidden=False,bNotConnectable=False,bDefaultValueIsReadOnly=False,bDefaultValueIsIgnored=False,bAdvancedView=False,bOrphanedPin=False,)
+			   CustomProperties Pin (PinId=C758BE724ED2C35AF55937B09A2B53F4,PinName="SweepHitResult",PinToolTip="Sweep Hit Result\nHit Result Structure\n\nHit result from any impact if sweep is true.",Direction="EGPD_Output",PinType.PinCategory="struct",PinType.PinSubCategory="",PinType.PinSubCategoryObject=ScriptStruct'/Script/Engine.HitResult',PinType.PinSubCategoryMemberReference=(),PinType.PinValueType=(),PinType.ContainerType=None,PinType.bIsMap=False,PinType.bIsSet=False,PinType.bIsArray=False,PinType.bIsReference=False,PinType.bIsConst=False,PinType.bIsWeakPointer=False,PersistentGuid=00000000000000000000000000000000,bHidden=False,bNotConnectable=False,bDefaultValueIsReadOnly=False,bDefaultValueIsIgnored=False,bAdvancedView=False,bOrphanedPin=False,)
+			   CustomProperties Pin (PinId=C8D4DFF74546FCDE47077A8F575C2D36,PinName="bTeleport",PinToolTip="Teleport\nBoolean\n\nWhether we teleport the physics state (if physics collision is enabled for this object). If true, physics velocity for this object is unchanged (so ragdoll parts are not affected by change in location). If false, physics velocity is updated based on the change in position (affecting ragdoll parts). If CCD is on and not teleporting, this will affect objects along the entire sweep volume.",PinType.PinCategory="bool",PinType.PinSubCategory="",PinType.PinSubCategoryObject=None,PinType.PinSubCategoryMemberReference=(),PinType.PinValueType=(),PinType.ContainerType=None,PinType.bIsMap=False,PinType.bIsSet=False,PinType.bIsArray=False,PinType.bIsReference=False,PinType.bIsConst=False,PinType.bIsWeakPointer=False,DefaultValue="false",AutogeneratedDefaultValue="false",PersistentGuid=00000000000000000000000000000000,bHidden=False,bNotConnectable=False,bDefaultValueIsReadOnly=False,bDefaultValueIsIgnored=False,bAdvancedView=False,bOrphanedPin=False,)
+			End Object
+			Begin Object Class=/Script/BlueprintGraph.K2Node_VariableGet Name="K2Node_VariableGet_47"
+			   VariableReference=(MemberName="ViveCameraHeight",MemberGuid=534B2E2F41DF6BFC3DD3ABA221017A9B,bSelfContext=True)
+			   NodePosX=672
+			   NodePosY=272
+			   NodeGuid=EECF89E94031777604797C9F4E162A93
+			   CustomProperties Pin (PinId=78E971AE41B9BC4AA335A6BCC9E5AE69,PinName="ViveCameraHeight",Direction="EGPD_Output",PinType.PinCategory="struct",PinType.PinSubCategory="",PinType.PinSubCategoryObject=ScriptStruct'/Script/CoreUObject.Vector',PinType.PinSubCategoryMemberReference=(),PinType.PinValueType=(),PinType.ContainerType=None,PinType.bIsMap=False,PinType.bIsSet=False,PinType.bIsArray=False,PinType.bIsReference=False,PinType.bIsConst=False,PinType.bIsWeakPointer=False,DefaultValue="0, 0, 0",AutogeneratedDefaultValue="0, 0, 0",LinkedTo=(K2Node_CallFunction_5 4751BC714842E8CEC1B7A294186125C2,),PersistentGuid=00000000000000000000000000000000,bHidden=False,bNotConnectable=False,bDefaultValueIsReadOnly=False,bDefaultValueIsIgnored=False,bAdvancedView=False,bOrphanedPin=False,)
+			   CustomProperties Pin (PinId=9489DF2447B3B0AFE5ED5086A7A47CB8,PinName="self",PinFriendlyName=NSLOCTEXT("K2Node", "Target", "Target"),PinType.PinCategory="object",PinType.PinSubCategory="",PinType.PinSubCategoryObject=BlueprintGeneratedClass'/Game/SteamVR_Pawn_Sitting.SteamVR_Pawn_Sitting_C',PinType.PinSubCategoryMemberReference=(),PinType.PinValueType=(),PinType.ContainerType=None,PinType.bIsMap=False,PinType.bIsSet=False,PinType.bIsArray=False,PinType.bIsReference=False,PinType.bIsConst=False,PinType.bIsWeakPointer=False,PersistentGuid=00000000000000000000000000000000,bHidden=True,bNotConnectable=False,bDefaultValueIsReadOnly=False,bDefaultValueIsIgnored=False,bAdvancedView=False,bOrphanedPin=False,)
+			End Object
+			Begin Object Class=/Script/UnrealEd.EdGraphNode_Comment Name="EdGraphNode_Comment_2"
+			   CommentDepth=-2
+			   NodePosX=608
+			   NodePosY=-80
+			   NodeWidth=592
+			   NodeHeight=448
+			   NodeComment="Offset the VRCameraRoot to move the Vive Camera"
+			   NodeGuid=3C39203E41D17487AD060C9123DA366B
+			End Object
+			Begin Object Class=/Script/UnrealEd.EdGraphNode_Comment Name="EdGraphNode_Comment_3"
+			   NodePosX=624
+			   NodePosY=192
+			   NodeWidth=512
+			   NodeHeight=160
+			   NodeComment="Set this to the standing height of the user."
+			   NodeGuid=C4F04D894FB9E577FF5AF7A8DB3E8FC3
+			End Object
+
+			END OBJECT
+	[/PARAM]
+[/OBJECT]
+[/EXCERPT:SteamVRSittingVRCamera]
+-->
\ No newline at end of file
diff --git a/Engine/Documentation/Source/Platforms/SteamVR/HowTo/StandingCamera/HTSteamVRStandingCamera.INT.udn b/Engine/Documentation/Source/Platforms/SteamVR/HowTo/StandingCamera/HTSteamVRStandingCamera.INT.udn
new file mode 100644
--- /dev/null
+++ b/Engine/Documentation/Source/Platforms/SteamVR/HowTo/StandingCamera/HTSteamVRStandingCamera.INT.udn
@@ -0,0 +1,99 @@
+Availability: Public
+Crumbs: %ROOT%, Platforms
+Title:Set Up a Standing Camera for Steam VR
+Description:Setting up the UE4 camera to work with a standing Steam VR experience.
+Type: how-to
+Skilllevel:Beginner
+Version:4.17
+Parent:Platforms/SteamVR/HowTo
+Order: 2
+tags:Basics
+topic-image:Platforms\SteamVR\SteamVRDevelopment_topic.png
+
+One of the first things you will need to think about when starting to develop any UE4 powered VR experience to be used with the Steam VR is deciding whether the experience will be a seated one or a standing one. In the following How - To we will take a look at how to setup your UE4 projects VR camera for a standing Steam VR experience.
+
+## Pawn Setup for a Standing Steam VR Experience 
+
+Below you will find instructions on how to setup up any Pawn to work with a standing Steam VR experience.
+
+1. First, open up or create a new Pawn Blueprint and then go to the **Component** section of the **Viewport** tab. From there add the following two components with the following names, making sure that the VRCamera is a child of VRCameraRoot:
+    
+    * **Scene** - VRCameraRoot
+    * **Camera** - VRCamera
+
+	![](HT_SteamVR_Camera_Setup_00.png)
+
+    [region:note]
+    No matter what VR HMD you are using, this is the Epic recommend way to set up a VR camera as it allows you to offset the position of the camera without having to move the actual camera.
+    [/region]
+
+With that completed, the initial setup for the VR Camera is now completed, in the next sections, we will go through setting the needed Blueprint logic for a standing Steam VR experience.
+
+## Blueprint Setup for Standing Steam VR Experience 
+
+Below you will find instructions on how to setup the Blueprint logic for a standing Steam VR experience.  
+
+1. First, open up the Pawn Blueprint if not already open and in the **Event Graph** drag off the **Event Begin Play** node to show the Executable Actions list. In the list search for the **Set Tracking Origin** node and click on it to add it to the Event Graph.
+
+    [REGION:lightbox]
+    [![](HT_Rift_Camera_Setup_09.png)(w:555)](HT_Rift_Camera_Setup_09.png)
+    [/REGION]
+
+    [REGION:caption]
+    Click for full image.
+    [/REGION]
+
+1. The Set Tracking Origin node has two options, **Floor Level** and **Eye Level**. For a standing experience you will need to leave the **Origin** of the Set Tracking Origin node to the default of **Floor Level**.
+    
+    [REGION:lightbox]
+    [![](Platforms\Oculus\HowTo\StandingCamera\HT_Rift_Camera_Setup_10.png)(w:555)](Platforms\Oculus\HowTo\StandingCamera\HT_Rift_Camera_Setup_10.png)
+    [/REGION]
+
+    [REGION:caption]
+    Click for full image.
+    [/REGION]
+
+1. Drag the Pawn Blueprint from the Content Browser to a level, making sure that it is placed at 0,0,0 in the level.
+    [REGION:lightbox]
+    [![](HT_Rift_Camera_Setup_06.png)(w:555)](HT_Rift_Camera_Setup_06.png)
+    [/REGION]
+
+    [REGION:caption]
+    Click for full image.
+    [/REGION]
+
+1. Select the Pawn Blueprint that was placed in the level and in the **Details** panel under the **Pawn** setting, set the **Auto Possess Player** from **Disabled** to **Player 0**.
+
+    [REGION:lightbox]
+    [![](HT_Rift_Camera_Setup_07.png)(w:225)](HT_Rift_Camera_Setup_07.png)
+    [/REGION]
+
+    [REGION:caption]
+    Click for full image.
+    [/REGION]
+
+1. Finally, go to the **Main Toolbar** and change the **Play Mode** to **VR Preview** and then press the **Play** button. When you put your HTC Vive HMD on and view the level while sitting down, you should see something similar to the following video. 
+
+    [OBJECT:EmbeddedVideo]
+        [PARAMLITERAL:videoid]
+        sO87nTNTATY
+        [/PARAMLITERAL]
+        [PARAMLITERAL:width]
+        100
+        [/PARAMLITERAL]
+        [PARAMLITERAL:height]
+        56.25
+        [/PARAMLITERAL]
+        [PARAMLITERAL:units]
+        %
+        [/PARAMLITERAL]
+        [PARAMLITERAL:type]
+        simple
+        [/PARAMLITERAL]
+    [/OBJECT]
+
+## UE4 Project Downloads
+Below you will find a link to where you can download the UE4 project that was used to create this example. 
+
+* [**Steam VR Standing VR Camera Example Project**](https://epicgames.box.com/s/8n6rtur0q01jrmgbdacog986tuzfx7js)
+
diff --git a/Engine/Documentation/Source/Platforms/SteamVR/HowTo/SteamVRHowTo.INT.udn b/Engine/Documentation/Source/Platforms/SteamVR/HowTo/SteamVRHowTo.INT.udn
new file mode 100644
--- /dev/null
+++ b/Engine/Documentation/Source/Platforms/SteamVR/HowTo/SteamVRHowTo.INT.udn
@@ -0,0 +1,33 @@
+Availability: Public
+Crumbs: %ROOT%,  Platforms
+Title: Steam VR How To
+Description: How To guides providing step-by-step instructions working with Steam VR and UE4.
+Type: landing
+Version: 4.17
+Parent: Platforms/SteamVR
+Order: 3
+Tags: Platform
+topic-image:Platforms\SteamVR\SteamVRDevelopment_topic.png
+
+![](Platforms\SteamVR\SteamVRDevelopment_topic.png)
+
+The Steam VR How To's contain step-by-step guides that will walk you through how to use UE4 and Steam VR together.
+
+## Beginner
+
+The Steam VR Rift Beginner How - To's will walk you through setting up your UE4 project to work with the Rift HMD, Touch Controllers and the Guardian system.
+
+[DIR(end:"1" output:"topic" type:"how-to" tags:"Basics" parent:"Platforms/SteamVR")]
+
+## Intermediate
+
+The Steam VR Intermediate How - To's will walk you through adjusting the color of the Chaperone System, using Touch Controllers to pick objects up and also how to use Stereo Layers to display information to the user.
+
+[DIR(end:"1" output:"topic" type:"how-to" tags:"SteamVR" parent:"Platforms/SteamVR")]
+
+## Advanced
+
+The Steam VR Advanced How - To's will walk you through detecting if the users has activated the Chaperone system and how to setup and use loading Splash Screens. 
+
+[DIR(end:"1" output:"topic" type:"how-to" tags:"VR" parent:"Platforms/SteamVR")]
+
diff --git a/Engine/Documentation/Source/Platforms/SteamVR/PreRecs/SteamVRPrerequisites.INT.udn b/Engine/Documentation/Source/Platforms/SteamVR/PreRecs/SteamVRPrerequisites.INT.udn
new file mode 100644
--- /dev/null
+++ b/Engine/Documentation/Source/Platforms/SteamVR/PreRecs/SteamVRPrerequisites.INT.udn
@@ -0,0 +1,55 @@
+Availability: Public
+Crumbs: %ROOT%
+Title: SteamVR Prerequisites
+Description: Information about setting up SteamVR to be used with UE4.
+Type:Reference
+Version:4.17
+Parent:Platforms/SteamVR
+Order:2
+Tags:SteamVR
+Tags:Platform
+topic-image:PreRecProfiling_Topic_Image
+
+
+On the following page, you will find information on how to make sure your UE4 project will work with SteamVR. 
+
+## Installing SteamVR 
+
+No matter what HMD you are using, if you are developing a VR project that will use SteamVR you will  have to download and install Steam before you can access SteamVR. In the following section we will go over all you need to know about downloading and installing steam VR.
+
+1. First, you will need to download and install the [Steam Client](http://store.steampowered.com/) from the [Steam website](http://store.steampowered.com/).
+
+1. Once Steam had been installed you can launch SteamVR the following two ways:
+
+	* By pressing the Start SteamVR button located at the top of the Steam Client.
+
+		![](SteamVRPreRec_00.png)
+
+	* By going to **Library** > **Tools** and searching for then clicking on the **SteamVR [beta]** option.
+
+		![](SteamVRPreRec_01.png)
+
+Remember that the SteamVR tools must be running in order for you to view your UE4 projects. You can minime the SteamVR tools window but do not close it. 
+
+## SteamVR Troubleshooting
+
+If you are having issues getting SteamVR to run or detect your HMD, the SteamVR and HTC Vive support pages can help you fix most issues that you will run into.
+
+[**SteamVR Support**](https://support.steampowered.com/kb_article.php?ref=8566-SDZC-9326)
+[**HTC Vive Support**](https://www.vive.com/us/support/)
+
+
+
+
+
+
+
+
+
+
+
+
+
+
+
+
diff --git a/Engine/Documentation/Source/Platforms/SteamVR/ProfilingPerf/SteamVRProfilingPerformance.INT.udn b/Engine/Documentation/Source/Platforms/SteamVR/ProfilingPerf/SteamVRProfilingPerformance.INT.udn
new file mode 100644
--- /dev/null
+++ b/Engine/Documentation/Source/Platforms/SteamVR/ProfilingPerf/SteamVRProfilingPerformance.INT.udn
@@ -0,0 +1,72 @@
+Title: SteamVR Profiling & Performance
+Availability:Public
+Crumbs:%ROOT%, Platforms
+Description:Information profiling the performance of your UE4 SteamVR projects.
+Navigation:platform
+Platform:SteamVR
+Version: 4.17
+type:landing
+parent:Platforms/SteamVR
+order:2
+tags:Performance and Profiling
+type:overview
+topic-image:SteamVR_ProPref_Topic_Image
+
+On this page, you will find information about how to profile the performance of your UE4 projects that utilize SteamVR. 
+
+## SteamVR Frame Timing
+
+![](SteamVR_ProPref_00.png)
+
+The SteamVR Frame Timing tool can be used to help track down what is causing your UE4  project to perform poorly. Whether in the editor or packaged builds, the SteamVR Frame Timing tool can help verify the actual CPU and GPU timings while accounting for things like throttling done by the application. For a more in-depth look at everything you can do with the SteamVR Frame Timing tool, check out the official documentation on the [SteamVR Frame Timing Tool](https://developer.valvesoftware.com/wiki/SteamVR/Frame_Timing).
+
+To display the SteamVR Frame Timing Tool you will need to do the following.
+
+1.  Right-click on the SteamVR tools and from the displayed menu, select the **Settings** option.
+
+    ![](SteamVR_ProPref_01.png)
+
+1. Then from the settings menu, Click on the **Display Frame Timing** button to display the Frame Timing tool.
+
+    [REGION:lightbox]
+    [![](SteamVR_ProPref_02.png)(w:340)](SteamVR_ProPref_02.png)
+    [/REGION]
+
+    [REGION:caption]
+    Click for full image.
+    [/REGION]
+
+1. Once the Frame Timing tool is up, you can then launch your UE4 project and check to see what is happening in the Frame Timing tool. 
+
+    ![](SteamVR_ProPref_00.png)
+
+## Saving SteamVR Frame Timing
+
+You can save the information that is generated from the Frame Timing tool so that it can be reviewed later or sent to someone else for review. To save the SteamVR Frame Timing you will need to do the following.
+
+1.  Right-click on the SteamVR tools and from the displayed menu, select the **Settings** option.
+
+    ![](SteamVR_ProPref_01.png)
+
+1. Then from the settings menu, Click on the **Save Frame Data Now** button to save the Frame Data.
+
+    [REGION:lightbox]
+    [![](SteamVR_ProPref_03.png)(w:340)](SteamVR_ProPref_03.png)
+    [/REGION]
+
+    [REGION:caption]
+    Click for full image.
+    [/REGION]    
+
+1. The Frame Timings will then be saved to a .CSV file called **VRFrames.csv** which can be found under, **C:\Program Files (x86)\Steam\logs**. 
+
+    ![](SteamVR_ProPref_04.png)
+
+
+
+
+
+
+
+
+
diff --git a/Engine/Documentation/Source/Platforms/SteamVR/QuickStart/1/SteamVRQuickStart_1.INT.udn b/Engine/Documentation/Source/Platforms/SteamVR/QuickStart/1/SteamVRQuickStart_1.INT.udn
--- a/Engine/Documentation/Source/Platforms/SteamVR/QuickStart/1/SteamVRQuickStart_1.INT.udn
+++ b/Engine/Documentation/Source/Platforms/SteamVR/QuickStart/1/SteamVRQuickStart_1.INT.udn
@@ -1,65 +1,67 @@
+Availability:Public
 Title: 1. SteamVR Initial Setup
-Crumbs:%ROOT%, Platforms
+Crumbs:%ROOT%, Platforms/SteamVR
 Description: Setting up the require Hardware and Software for SteamVR
-Availability:Public
 Navigation:platform
+Version: 4.17
+Skilllevel: Intermediate
 Platform:SteamVR
+type:multi-step
+parent:Platforms/SteamVR/QuickStart
+order:1
+Checkpoint:SteamVRQS
+tags: Steam VR
+
+[nav]
+
+In the following section, we will take a look at how to get SteamVR setup so that it will work with Unreal Engine 4 (UE4).
 
-[VAR:Steps]
-[OBJECT:Navigation]
-	[PARAM:previous]
-	[/PARAM]
-	[PARAM:current]
-	[/PARAM]
-	[PARAM:home]
-		[SteamVR Quick Start Home](Platforms/SteamVR/QuickStart)
-	[/PARAM]
-	[PARAM:next]
-		[Next Step](Platforms/SteamVR/QuickStart/2 "%Platforms/SteamVR/QuickStart/2:title%")
-	[/PARAM]
-[/OBJECT]
-[/VAR]
-
-%Steps%
-
-In the following section we will take a look at how to get SteamVR setup so that it will work with Unreal Engine 4 (UE4).
+## Steps
 
 [region:note]
-With each SteamVR development kit Valve has provide [detailed instructions](http://media.steampowered.com/apps/steamvr/vr_setup.pdf ) that will show you how to correctly set everything up. 
-If you have not read this document please do so before going any further as the following information is not a substitute for the information contained in the document created by Valve.
+With each SteamVR development kit Valve has provided [detailed instructions](http://media.steampowered.com/apps/steamvr/vr_setup.pdf ) that will show you how to correctly set everything up. 
+If you have not read this document, please do so before going any further as the following information is not a substitute for the information contained in the document created by Valve.
 [/region]
 
-1. Make sure to clear a 12ft x 9ft / 4m x 3m area of all objects that you could trip over or hit while you are using SteamVR.
+1. Ensure that the Vive Head Mounted Display (HMD), Steam Controllers, Breakout Box, and Lighthouse Base stations are all unpacked, powered on, connected and setup according to the [instructions](http://media.steampowered.com/apps/steamvr/vr_setup.pdf) provide by Valve.
 
-1. Ensure that the Vive Head Mounted Display (HMD), Steam Controllers, Breakout Box, and Lighthouse Base stations are all unpacked, powered on, connected and setup according the the [instruction](http://media.steampowered.com/apps/steamvr/vr_setup.pdf) provide by Valve.
-
-1. If you have not done so already, download and install the [Steam Client](http://store.steampowered.com/about/) on your development PC.
+1. If you have not done so already, make sure to download and install the [Steam Client](http://store.steampowered.com/about/) on your development PC.
 
 1. To install the **SteamVR** tools, use your mouse to hover over the Steam **Library** option and from the menu that is displayed, select the **Tools** option. 
 
-	![](T_Launch_SteamVR_Tools.png)
+    ![](T_Launch_SteamVR_Tools.png)
 
 1. Once you are in the Tools section, search for **SteamVR** using the search bar at the top. Once SteamVR is located, double click on it to download and install it.
 
-	![](T_SteamTools_Software.png)
+    [REGION:lightbox]
+    [![](T_SteamTools_Software.png)(w:568)](T_SteamTools_Software.png)
+    [/REGION]
+
+    [REGION:caption]
+    Click for full image.
+    [/REGION]
 
-	[region:note]
-	You can also install SteamVR by clicking on the **VR** icon that is located on the upper right hand corner of the Steam Client and follow the instructions on the screen.
-	![](T_SteamVR_Icon_Install.png)
-	[/region]
+    [region:note]
+    You can also install SteamVR by clicking on the **VR** icon that is located on the upper right-hand corner of the Steam Client and following the provided instructions.
+    ![](T_SteamVR_Icon_Install.png)
+    [/region]
 
 1. Double-clicking on the **SteamVR** option in the Tools menu will launch the SteamVR tools as shown in the image below.
 
-	![](T_Steam_VR_Ready.png)
-	[region:note]
-	When SteamVR shows all devices in green like in the image above that means everything is working correctly. 
-	If a device that shows up in grey then there  is a problem with that device.
-	SteamVR will tell you what is wrong with a greyed out device and possible fixes if you hover your mouse over the greyed out device.
-	[/region]
+    ![](T_Steam_VR_Ready.png)
+    [region:note]
+    When SteamVR shows all devices in green like in the image above that means everything is working correctly. 
+    If a device that shows up in gray then there is a problem with that device.
+    SteamVR will tell you what is wrong with a device that is grayed out and also provides possible fixes. To see the fixes,  hover your mouse over a grayed out device.
+    [/region]
 
-1. Before you can use SteamVR with UE4 you must setup the SteamVR interaction area. 
+1. Before you can use SteamVR with UE4, you must setup the SteamVR interaction area. 
 To do this, **Right - click** on the SteamVR window and select the **Run Room Setup** and follow the provided directions to setup the SteamVR interaction area.
 
-	![](T_Steam_VR_App.png)
+    ![](T_Steam_VR_App.png)
+
+## End Result
+
+When completed you will now have SteamVR setup and ready to be used with UE4.
 
-%Steps%
\ No newline at end of file
+[nav]
\ No newline at end of file
diff --git a/Engine/Documentation/Source/Platforms/SteamVR/QuickStart/2/SteamVRQuickStart_2.INT.udn b/Engine/Documentation/Source/Platforms/SteamVR/QuickStart/2/SteamVRQuickStart_2.INT.udn
--- a/Engine/Documentation/Source/Platforms/SteamVR/QuickStart/2/SteamVRQuickStart_2.INT.udn
+++ b/Engine/Documentation/Source/Platforms/SteamVR/QuickStart/2/SteamVRQuickStart_2.INT.udn
@@ -1,29 +1,22 @@
+Availability:Public
 Title: 2. Setting up UE4 to work with SteamVR
-Crumbs:%ROOT%, Platforms
+Crumbs:%ROOT%, Platforms/SteamVR
 Description: Setting up SteamVR and UE4 to work together
-Availability:Public
 Navigation:platform
-Platform:SteamVR
-
-[VAR:Steps]
-[OBJECT:Navigation]
-	[PARAM:previous]
-	[Previous Step](Platforms/SteamVR/QuickStart/1 "%Platforms/SteamVR/QuickStart/1:title%")
-	[/PARAM]
-	[PARAM:current]
-	[/PARAM]
-	[PARAM:home]
-		[SteamVR Quick Start Home](Platforms/SteamVR/QuickStart)
-	[/PARAM]
-	[PARAM:next]
-	[/PARAM]
-[/OBJECT]
-[/VAR]
+Version: 4.17
+Skilllevel: Beginner
+type:multi-step
+parent:Platforms/SteamVR/QuickStart
+order:2
+Checkpoint:SteamVRQS
+Tags: SteamVR
 
-%Steps%
+[nav]
 
 In the following section, we will go over how to setup a new Unreal Engine 4 (UE4) project to work with SteamVR.
 
+## Steps
+
 [region:note]
 If you have not done so already, make sure to run the SteamVR **Room Setup** to establish and calibrate the VR tracking area. 
 Failing to do this could result in SteamVR and UE4 not working correctly together.
@@ -31,81 +24,41 @@ Failing to do this could result in SteamVR and UE4 not working correctly togethe
 
 1. Create a new blank Blueprint base project with hardware set to **Mobile / Tablet**, graphics set to **Scalable 3D or 2D** and **No Starter Content**.
 
-	![](T_CreateNewProject.png)
+	[REGION:lightbox]
+	[![](T_CreateNewProject.png)(w:555)](T_CreateNewProject.png)
+	[/REGION]
+
+	[REGION:caption]
+	Click for full image.
+	[/REGION]
 
 1. Once the project loads click on the small triangle next to the **Play** button and then from the menu that is displayed, select the **VR Preview** option. 
-	If your VR Headset is properly set up, you should now see the basic level being displayed and have the ability to rotate your head around in any direction when you put the HMD on.
 	
 	![](Check_VR_Working_00.png)
 
-	[region:tip]
-	If your can not get the HTC Vive HMD to work check the Plugins section to make sure that the Steam VR plug-in is enabled and if not click on the checkmark box next to **Enable** to enable it.
-	 ![](T_Enable_Steam_VR.png)
-	 If that still does not get your VR Headset to work with the VR Preview make sure to check out the [VR Documents](Platforms/VR) to ensure that nothing was missed during setup.
-	[/region]
-
-1. With verification that the Head Mounted Display (HMD) is working, right - click inside of the **Content Browser** and select the **Blueprint Class** option from the **Create Basic Asset** section.
-
-	![](Create_New_BP.png)
-
-1. From the **Pick Parent Class** window that is displayed search for, create and name the following Blueprint nodes.
-
-	![](Pick_Parent_Class_Window.png)
-
-	|Blueprint Node|Name to use|
-	|--------------|------|
-	|Game Mode|VR_GameMode|
-	|Pawn|VR_Pawn|
-
-	When completed you should have the following Blueprints in your Content Browser.
-
-	![](Creating_Required_BP_Nodes.png)
-
-1. Open **VR_GameMode** Blueprint and then under the **Classes** section change the **Default Pawn Class** to point to the **VR_Pawn** that was created.
-	
-	![](Setting_Up_Game_Mode.png)
-
-1. Under **World Settings** in the **Game Mode** section, change the **GameMode Override** to the **VR_GameMode**.
-
-	![](Adding_New_GameMode.png)
-
-1. Now open up the **VR_Pawn** Blueprint then click on the **Class Defaults** and in the **Details** panel under the **Camera** section set the **Base Eye Height** to **0.0**.
-
-	![](Set_Base_Eye_Height.png)
-
-1. From the **Components** tab, add the following Components to the **VR_Pawn** Blueprint.
-	|Component Name|Amount|
-	|-------|-----|
-	|Capsule| 1|
-	|Scene|1|
-	|Camera|1|
-	|SteamVRChaperone|1|
-
-	![](Adding_A_Camera.png)
-
-1. Select the **Capsule** Component and under the **Details** panel in the **Shape** section, change the **Capsule Half Height** to **96** and the Capsule Radius to **16**.
-
-	![](Capsule_Setup.png)
-
-1. Click on the **Camera** and then drag it on top of the **Scene** Component, making the Camera a child of the Scene Component.
-
-	![](Parenting_The_Camera.png)
-
-1. Select the new **Scene** Component and go to **Details** > **Transform** > **Location** and set the **Z** value to **-110** then make sure to **Compile** and **Save** the Blueprint.
-
-	![](Change_Scene_Height.png)
-
-1. Inside of the level, select the Static Mesh used for the floor and set the X, Y and Z location to **0,0,0**.
-
-	![](Move_Floor_Down.png)
-
-1. Launch the level using the **VR Preview** button, put your HMD on and then stand up and look around.
-When you look down at your feet, you should now feel as though the virtual floor that you are standing on matches up to the real world floor.
-
-[region:note]
-If it feels like you are too close or too far away from the virtual floor, open the  VR_Pawn Blueprint, select the Scene Component that the Camera is parented to and adjust the Z height value up or down till the floor feels like it is in the correct location for your height.
-[/region]
-
-
-%Steps%
+1. When the VR Preview launches, put on your HMD and you should now see the basic level being displayed and have the ability to rotate your head around in any direction when you put the HMD on like in the video below.
+
+	[OBJECT:EmbeddedVideo]
+		[PARAMLITERAL:videoid]
+		4xuNrwS3o-c
+		[/PARAMLITERAL]
+		[PARAMLITERAL:width]
+		100
+		[/PARAMLITERAL]
+		[PARAMLITERAL:height]
+		56.25
+		[/PARAMLITERAL]
+		[PARAMLITERAL:units]
+		%
+		[/PARAMLITERAL]
+		[PARAMLITERAL:type]
+		simple
+		[/PARAMLITERAL]
+	[/OBJECT]
+
+## End Result
+
+When completed you will be able to view the UE4 project using SteamVR and the HTC Vive HMD.
+
+[nav]
 
diff --git a/Engine/Documentation/Source/Platforms/SteamVR/QuickStart/3/SteamVRQuickStart_3.INT.udn b/Engine/Documentation/Source/Platforms/SteamVR/QuickStart/3/SteamVRQuickStart_3.INT.udn
new file mode 100644
--- /dev/null
+++ b/Engine/Documentation/Source/Platforms/SteamVR/QuickStart/3/SteamVRQuickStart_3.INT.udn
@@ -0,0 +1,25 @@
+Availability:Public
+Title: 3. On Your Own!
+Crumbs:%ROOT%, Platforms/SteamVR
+Description:Using the methods provided during this guide, polish up your level with additional content and read up on related topics. 
+Navigation:platform
+Version: 4.17
+Skilllevel: Beginner
+type:multi-step
+parent:Platforms/SteamVR/QuickStart
+order:3
+Checkpoint:SteamVRQS
+Tags: SteamVR
+
+[nav]
+
+Now that you can view a UE4 project using SteamVR and the HTC Vive try adding the following items to your project.
+
+* Use the [Migrate Tool](Engine/Content/Browser/UserGuide/Migrate/) to move content from one of the mobile content examples into your project for use.
+
+* Add support for [Motion Controllers](Platforms\VR\MotionController) so that the user can move objects around in VR just like they would in real life.
+
+* Use the [GPU Profiler](Engine/Performance/GPU/) to help track your projects performance while you are building it.
+
+[nav]
+
diff --git a/Engine/Documentation/Source/Platforms/SteamVR/QuickStart/SteamVRQuickStart.INT.udn b/Engine/Documentation/Source/Platforms/SteamVR/QuickStart/SteamVRQuickStart.INT.udn
--- a/Engine/Documentation/Source/Platforms/SteamVR/QuickStart/SteamVRQuickStart.INT.udn
+++ b/Engine/Documentation/Source/Platforms/SteamVR/QuickStart/SteamVRQuickStart.INT.udn
@@ -1,120 +1,33 @@
-Title: SteamVR Quick Start
-Crumbs:%ROOT%, Platforms
-Description: A How To guide for setting up SteamVR to work with UE4
-Availability:public
+Availability:Public
+Title: Steam VR Quick Start
+Crumbs:%ROOT%
+Description: Instructions on how to quickly setup and deploy projects for use with Steam VR.
 Navigation:platform
-Platform:SteamVR
+Platform:Steam VR
+Skilllevel: Beginner
+Version:4.17
 type:quick start
 type:landing
-parent:Platforms/SteamVR
 tags:VR
 tags:SteamVR
+order:2
+Checkpoint:SteamVRQS
 
 ![](Platforms\SteamVR\T_Vive_HMD.png)
 
-This "Quick Start" series will walk you through how to setup SteamVR so that it will work with Unreal Engine 4.
-By the end of this "Quick Start" series you will have a new project that will work SteamVR.
-
-[OBJECT:TopicButtonList]
-	[PARAM:icon]
-		![](%ROOT%/build_icon.png)(convert:false)
-	[/PARAM]
-	[PARAM:title]
-		Steps
-	[/PARAM]
-	[PARAM:description]
-	[/PARAM]
-	[PARAM:links]
-	* [](Platforms/SteamVR/QuickStart/1 "%Platforms/SteamVR/QuickStart/1:description%")
-	* [](Platforms/SteamVR/QuickStart/2 "%Platforms/SteamVR/QuickStart/2:description%")
-	[/PARAM]
-[/OBJECT]
-
-
-
-
-
-
-
-
-
-
-
-
-
-
-
-
-
-
-
-
-
-
-
-
-
-
-
-
-
-
-
-
-
-
-
-
-
-
-
-
-
-
-
-
-
-
-
-
-
-
-
-
-
-
-
-
-
-
-
-
-
-
-
-
-
-
-
-
-
-
-
-
-
-
-
-
-
-
-
-
-
-
+## Goals
 
+The Steam VR Quick Start will walk you through how to setup your development PC and Steam VR to work with the Vive Head Mounted Display headset.
 
+## Objectives
+* Create a new UE4 project that is targeted towards Steam VR development.
+* Setting the required project settings so your project can be used with Steam VR.
 
+[DIR(output:"checkpoint" parent:"Platforms/SteamVR/QuickStart")]
 
+[REGION:call_to_action]
+ [Click to Start](Platforms/SteamVR/QuickStart/1 "%Platforms/SteamVR/QuickStart/1:title%")
+[/REGION]
 
 
 
diff --git a/Engine/Documentation/Source/Platforms/SteamVR/Reference/SteamVRCommandReference.INT.udn b/Engine/Documentation/Source/Platforms/SteamVR/Reference/SteamVRCommandReference.INT.udn
new file mode 100644
--- /dev/null
+++ b/Engine/Documentation/Source/Platforms/SteamVR/Reference/SteamVRCommandReference.INT.udn
@@ -0,0 +1,31 @@
+Availability: Public
+Crumbs: %ROOT%, Platforms
+Title: SteamVR Command Reference
+Description: Information on the various commands you can use with SteamVR.
+Type:reference
+Skilllevel: Intermidate
+Version: 4.17
+Parent: Platforms/SteamVR
+Order: 2
+tags: VR
+topic-image:Platforms\SteamVR\SteamVRDevelopment_topic.png
+
+Below you will find various INI and Console commands that can be used to adjust the way UE4 and SteamVR interact with one another. 
+
+## SteamVR VR Mirror Mode
+
+The VR Mirror is the name for the viewport window that allows non VR users to see what the VR user is seeing. When using SteamVR there are a few ways you can change the way the VR Mirror is displayed. Below you will find the commands to change the way the VR Mirror mode and what the effect that will have on the how the VR Mirror looks.
+
+[REGION:note]
+To input and of the following command press the backtick **`** key to open up the UE4 console and then input **VR.MirrorMode** followed by any of the following numbers.
+[/REGION]
+
+
+|Command|Description|Image|
+|-----------|-------------------|----------------|
+|VR.SpectatorScreenMode 0|Disables the VR Mirror|N/A|
+|VR.SpectatorScreenMode 1| Shows a single eye without any distortion.|[REGION:lightbox][![](SteamVRCR_00.png)(w:160)](SteamVRCR_00.png)[/REGION]|
+|VR.SpectatorScreenMode 2| Shows both eyes in stereo without lense distortion.|[REGION:lightbox][![](SteamVRCR_01.png)(w:160)](SteamVRCR_01.png)[/REGION]|
+|VR.SpectatorScreenMode 3| Shows one eye in a cropped rectangle that fills the screen.|[REGION:lightbox][![](SteamVRCR_02.png)(w:160)](SteamVRCR_02.png)[/REGION]|
+|VR.SpectatorScreenMode 4| Shows a single eye in stereo.|[REGION:lightbox][![](SteamVRCR_03.png)(w:160)](SteamVRCR_03.png)[/REGION]|
+|VR.SpectatorScreenMode 7| Shows a single eye in a cropped rectangle with the option to display another texture.|[REGION:lightbox][![](SteamVRCR_04.png)(w:160)](SteamVRCR_04.png)[/REGION]|
\ No newline at end of file
diff --git a/Engine/Documentation/Source/Platforms/SteamVR/SteamVR.INT.udn b/Engine/Documentation/Source/Platforms/SteamVR/SteamVR.INT.udn
--- a/Engine/Documentation/Source/Platforms/SteamVR/SteamVR.INT.udn
+++ b/Engine/Documentation/Source/Platforms/SteamVR/SteamVR.INT.udn
@@ -20,7 +20,7 @@ Topic-image:SteamVRDevelopment_topic.png
 ## SteamVR Topics
 
 [DIR(output:"listbutton" type:"overview" parent:"Platforms/SteamVR")]
-
+[DIR(output:"listbutton" type:"landing" parent:"Platforms/SteamVR")]
 
 
 
diff --git a/Engine/Documentation/Source/Platforms/VR/ContentSetup/VRContent.INT.udn b/Engine/Documentation/Source/Platforms/VR/ContentSetup/VRContent.INT.udn
--- a/Engine/Documentation/Source/Platforms/VR/ContentSetup/VRContent.INT.udn
+++ b/Engine/Documentation/Source/Platforms/VR/ContentSetup/VRContent.INT.udn
@@ -1,12 +1,13 @@
-Availability:Public
+Availability: Public
+Crumbs: %ROOT%
 Title:Virtual Reality Best Practices
-Crumbs:
-Description:Information about developing for VR.
-Navigation:platform
-Version: 4.14
-parent:Platforms/VR
-type:overview
-tags:VR
+Description: Information about developing for VR.
+Type: Overview
+SkillLevel: Beginner
+Version:4.17
+Parent: Platforms/VR
+Order:4
+Tags:VR
 
 
 [TOC(start:2)]
@@ -19,7 +20,13 @@ developing content for your VR projects.
 When creating a new project for VR, whether the project is Blueprint based or C++ based, it is best to create a project that uses the **Mobile / Tablet** option 
 with **Scalable 3D or 2D** graphics and **No Starter Content** to ensure that your VR project will be running at frame right from the start. 
 
-![](VR_Project_Settings.png)
+[REGION:lightbox]
+[![](VR_Project_Settings.png)(w:400)](VR_Project_Settings.png)
+[/REGION]
+
+[REGION:caption]
+Click for full image.
+[/REGION]
 
 ## VR .INI Settings
 
@@ -49,89 +56,6 @@ These . INIsettings, if used, should be added to your project's **Config\Default
 	sg.EffectsQuality=2
 	sg.PostProcessQuality=0
 
-## VR Instanced Stereo
-
-With the release of UE 4.11, Instanced Stereo Rendering was introduced to help lessen the performance impact of VR in UE4. 
-The feature can be enabled by checking the **Instanced Stereo** option that can be found under **Edit > Project Settings > Rendering > VR**. 
-
-![](VR_ISR.png)
-
-After **Instanced Stereo** rendering has been enabled, it will require an engine restart, and shader recompile.
- The base pass, and early-z pass, works with static meshes, skeletal meshes, sprite particles, and mesh particles with the feature turned on (both serial and parallel rendering paths). 
- **Instanced Stereo** currently works on PC (DirectX) and PS4 with support for other platforms coming after the initial 4.11 release. 
- Here is a video showing Standard Stereo Rendering in VR, versus Instanced Stereo Rendering in VR. 
-
-[OBJECT:EmbeddedVideo]
-	[PARAMLITERAL:width]
-	640
-	[/PARAMLITERAL]
-	[PARAMLITERAL:height]
-	360
-	[/PARAMLITERAL]
-	[PARAMLITERAL:videoid]
-	nmajnglJeyE
-	[/PARAMLITERAL]
-[/OBJECT]
-
-### VR Multiview Support for Mobile
-
-Certain (supported) mobile devices can now support VR Multiview. Mobile Multiview is similar to Instanced Stereo rendering that is available to desktop PC's, and works by providing an optimized path for Stereo Rendering on the mobile device's CPU. To enable this feature, you will need to do the following:
-
-[region:warning]
-Due to the experimental nature of this feature, it will only work on modern Mali-based GPUs. As this feature matures, and support for more mobile devices come online, this warning will be removed.
-[/region]
-
-1. From the Main Toolbar, go to **Edit** >  **Projects Settings**, to open the Editor's Project Settings.
-
-1. Then, go to **Engine** > **Rendering** > **VR**, and look for the **Mobile Multi-View (Experimental)** option. 
-
-1. Enable the Mobile Multi-View (Experimental) option by click on the checkmark box next to its name, and then restart Editor for the changes to take effect.
-
-    ![](VR_MobileMultiView.png)
-
-For this feature to work, you will need your Android build settings to be setup like the following:
-
-* Make sure that **Support OpenGL ES2** is enabled in the **Android** > **Build** section. 
-* Make sure that both **Mobile HDR** and **Instanced Stereo Rendering** are disabled.
-* Currently, Mobile Multiview will only work with modern Mali-based GPU's. 
-* If you package a project with this feature on, but don't have a compatible GPU, it will be disabled at runtime. 
-
-## VR Performance Profiling
-
-To track down assets that are too expensive for VR, you need to profile what your project is doing on the CPU and GPU as much as possible over the lifetime of your project. 
-
-* **GPU Profiling** - To activate the GPU profiler, all you need to do is press  **CTRL + SHIFT + ,(Comma Key)** together, while running your project. After the keys have been pressed, you will see a new window come up that should look similar to the image below. 
-	![](VR_GPU_Profiler.png)
-
-* **CPU Profiler** - Profiling what your project is doing on the CPU is a bit more involved than profiling for the GPU. To find out more on how to do this, please refer to the [Performance Profiler](https://docs.unrealengine.com/latest/INT/Engine/Performance/Profiler/index.html) documents. 
-
-## VR Forward Rendering
-
-![](Engine\Performance\ForwardRenderer\ForwardRenderingBanner.png)(w:840)
-
-[INCLUDE:Engine\Performance\ForwardRenderer#intro]
-
-[REGION:note]
-Please see the [](Engine\Performance\ForwardRenderer) documentation for more information on this feature. 
-[/REGION]
-
-## VR Post Process Settings
-
-Due to the demanding rendering requirements of VR, many of the advanced Post Processing features that are enabled by default should be disabled, or your project could suffer
- from performance problems. To accomplish this in your project, you will need to do the following.
-
- 1. Add a Post Process(PP) volume to your level if there is not one already there.
- 1. Select the PP volume, and in the **Post Process Volume** section, enable the **Unbound** option so that the settings in the PP volume will be applied to the entire level. 
-
-	![](VR_PP_Unbound.png)
-
- 1. Expand the **Settings** of the **Post Process Volume**, and then go through each section to disable any active PP settings by enabling that property; clicking on it, and then setting the value from the default (usually 1.0) to **0**, thus disabling the feature.
-	
-	![](VR_Disable_PP_Settings.png)
-	[region:note]
-	When doing this, you won't need to hit every section and set all of the properties to 0. Instead, first disable the really heavy hitting features like Lens Flares, Screen Space reflections, Temporal AA, SSAO, Bloom and anything else that might have an impact on performance.
-	[/region]
-
 ## VR World Scale
 
 Ensuring the correct scale of your world is one of the most important ways to help deliver the best user experience possible 
@@ -163,14 +87,15 @@ could end up having a very unpleasant VR experience.
 * You must maintain framerate, and ideally, a little bit of a buffer to make sure you're always over the HMD's native framerate.
 Low framerates are another trigger for Simulation Sickness, so make sure to optimize your project as much as possible. 
 The following table shows the various HMD's UE4 supports, and the target frame rates your VR project needs to run at on those devices.
-| HMD Device  | Target Frame Rate |
-| --- | --- |
-| DK1 | 60 FPS|
-| DK2 | 75 FPS|
-| Rift Retail | 90 FPS|
-| Vive | 90 FPS|
-| Gear VR | 60 FPS|
-| PSVR | Variable up to 120 FPS|
+
+	| HMD Device  | Target Frame Rate |
+	| --- | --- |
+	| DK1 | 60 FPS|
+	| DK2 | 75 FPS|
+	| Rift Retail | 90 FPS|
+	| Vive | 90 FPS|
+	| Gear VR | 60 FPS|
+	| PSVR | Variable up to 120 FPS|
 
 * Developers make the worst test subjects, because they are often used to using VR devices. Check your game as much as you can, with as many different people as you can, to make sure that you are not causing Simulation Sickness.
 * Avoid cinematic cameras, or anything that takes control of camera movements away from the player, as this tends to be the biggest culprit of your user having a bad VR experience.
@@ -191,6 +116,15 @@ However, for a standing experience, you should make sure that the camera origin
 You can achieve either effects by attaching a camera component to a scene component at the base of the character (ground level), or by setting the Eye Height to half of the 
 (negative) Cylinder Height of the collision capsule on the character.
 
+[OBJECT:ComparisonSlider]
+	[PARAM:before]
+	![Standing VR Camera](VR_Standing_Experiance.png)(w:500)
+	[/PARAM]
+	[PARAM:after]
+	![Seated VR Camera](VR_Seated_Experiance.png)(w:500)
+	[/PARAM]
+[/OBJECT]
+
 ## VR Character Settings
 
 The setup for a character using a VR headset is slightly different than for a standard character. Things like character 
@@ -231,34 +165,38 @@ The VR camera needs to be positioned slightly lower than the base eye height to
 Base Eye Height:|180 cm|160 cm|
 
 
-
 ## VR Content Considerations
 
 When creating VR content, remember that users can look at that content from multiple angles. Here are few things that you might 
 have done in the past, but need to avoid in VR:
 
-* **Scale** - The best thing to do about the scale of the objects in your VR world, is to mimic reality as closely as you can. 
-Making objects bigger or smaller than their real world counterparts could lead to confusion or Simulation Sickness.
+ * **Scale** - The best thing to do about the scale of the objects in your VR world, is to mimic reality as closely as you can. Making objects bigger or smaller than their real world counterparts could lead to confusion or Simulation Sickness.
 
-* **Missing Polygon Faces** - In standard games, it is often acceptable (and preferred) to remove polygon faces from objects that cannot be seen by the player. 
+ * **Missing Polygon Faces** - In standard games, it is often acceptable (and preferred) to remove polygon faces from objects that cannot be seen by the player. 
 However, in VR games, players have much more freedom to look around, and this practice can sometimes lead to players being able to see things that they're not supposed to see.
 
-* ** Which Type of Lighting to use** - You should always use **Static lighting** and **Lightmaps** when making a VR project... this is the cheapest option to render. 
+ * ** Which Type of Lighting to use** - You should always use **Static lighting** and **Lightmaps** when making a VR project... this is the cheapest option to render. 
 If you need to use dynamic lighting, make sure to limit the amount of dynamic lights to as few as possible, and make sure that they never touch one another. 
 If you have an outdoor scene, set your directional light to dynamic instead of stationary, and then turn on Cascaded Shadow Maps (CSM); adjusting the settings to be as simple as possible while still giving you shadows. 
 
-* ** VR & VFX** - Some VFX tricks, like using [SubUV Textures](https://docs.unrealengine.com/latest/INT/Engine/Rendering/ParticleSystems/Reference/Modules/SubUV/index.html) to simulate fire or smoke, do not hold up very well when viewed in VR. 
+ * ** VR & VFX** - Some VFX tricks, like using [SubUV Textures](https://docs.unrealengine.com/latest/INT/Engine/Rendering/ParticleSystems/Reference/Modules/SubUV/index.html) to simulate fire or smoke, do not hold up very well when viewed in VR. 
 In many cases you are going to need to use static meshes instead of 2D particles to simulate VFX's like explosions or smoke trails.  
 Near field effects, or effects that happen very close to the camera, work well in VR, but only when the effects are made up of Static Mesh particles.
 
-* **VR & Transparency** - In 3D graphics, rendering transparency is extremely costly, because transparency will generally have to be re-evaluated per-frame to ensure that nothing has changed. 
+ * **VR & Transparency** - In 3D graphics, rendering transparency is extremely costly, because transparency will generally have to be re-evaluated per-frame to ensure that nothing has changed. 
 Because of this re-evaluation, rendering Transparency in VR can be so costly, that its cost outweighs its benefits. 
 However, to get around this issue, you can use the **DitherTemporalAA** Material Function. 
 This Material Function will allow a Material to look like it is using transparency. Also, this will help you avoid common transparency issues, such as self-sorting.
 
-	![](VR_Dither_Trans_AA.png)
+	[REGION:lightbox]
+	[![](VR_Dither_Trans_AA.png)(w:500)](VR_Dither_Trans_AA.png)
+	[/REGION]
+
+	[REGION:caption]
+	Click for full image.
+	[/REGION]
 
-* **Fake everything you can** - Finding clever ways to recreate expensive rendering options, like dynamic shadows or lighting, can be a huge win for hitting your performance goals in VR. 
+ * **Fake everything you can** - Finding clever ways to recreate expensive rendering options, like dynamic shadows or lighting, can be a huge win for hitting your performance goals in VR. 
 In Showdown, having the characters cast dynamic shadows proved to be too expensive per-frame, which meant that dynamic shadows had to be cut from the project. 
 However, this made the characters look like they were floating while moving. 
 To fix this, fake blob shadows were introduced that could dynamically adjust their position and intensity based on how close the character was to an object in the world. 
@@ -271,7 +209,7 @@ This helped give the illusion that a character was casting shadows when they cam
 	[region:caption]
 	Click for full image.
 	[/region]
-
+<!--
 ## VR Blueprint Functions
  
  Throughout the Blueprint editor, you will find a number of different Blueprint nodes that can be used to interact with VR headsets and motion controllers. 
@@ -305,7 +243,7 @@ Set World to Meters Scale| Sets the World to Meters scale, which changes the sca
 |Get Tracking Space| Gets the tracking space (for example, sitting or standing), which determines the location of the origin.|
 |Get Valid Tracked Device Ids| Returns an array of the currently tracked device IDs.|
 |Set Tracking Space| Sets the tracking space (for example, sitting or standing), changing which space (tracked) positions are returned to.|
-
+-->
 ## Known VR Issues
 
 Due to the manner in which HMD's work, some art techniques that are staples of Video Game development, no longer have the impact that they once did. 
diff --git a/Engine/Documentation/Source/Platforms/VR/DesktopVR/DesktopVR.INT.udn b/Engine/Documentation/Source/Platforms/VR/DesktopVR/DesktopVR.INT.udn
new file mode 100644
--- /dev/null
+++ b/Engine/Documentation/Source/Platforms/VR/DesktopVR/DesktopVR.INT.udn
@@ -0,0 +1,47 @@
+Availability: Public
+Crumbs: %ROOT%
+Title: Desktop Virtual Reality Development
+Description: Information over developing for desktop based virtual reality platforms.
+Type: Landing
+Version: 4.16
+Parent: Platforms
+Order: 1
+topic-icon:%ROOT%/platform_icon.png
+
+![](Header_00.png)
+
+Unreal Engine 4 (UE4) supports a wide range of desktop based Virtual Reality (VR) Head Mounted Displays (HMD's), Motion Controllers and their corresponding Software Development kits. Contained with the pages linked below you will find a multitude of different documents that can help you get any of the supported VR HMD's working with your UE4 project.
+
+## Supported Desktop VR Head Mounted Displays
+
+UE4 natively supports the following desktop based VR HMD's: 
+
+* **Oculus Rift**
+* **HTC Vive** 
+
+
+The fact that UE4 natively supports the above desktop HMD's means that you do not need to do any additional downloading or compiling to view your UE4 project in VR.  Below you will find links to documentation that will explain how to get your UE4 project to work with the supported HMD of your choice.
+
+## Oculus Rift
+
+[DIR(end:"1" output:"topic" type:"landing" tags:"Oculus Rift" parent:"Platforms/Oculus")]
+
+
+## SteamVR
+
+[DIR(end:"1" output:"topic" type:"landing" tags:"SteamVR" parent:"Platforms/SteamVR")]
+
+
+
+
+
+
+
+
+
+
+
+
+
+
+
diff --git a/Engine/Documentation/Source/Platforms/VR/MobileVR/MobileVR.INT.udn b/Engine/Documentation/Source/Platforms/VR/MobileVR/MobileVR.INT.udn
new file mode 100644
--- /dev/null
+++ b/Engine/Documentation/Source/Platforms/VR/MobileVR/MobileVR.INT.udn
@@ -0,0 +1,33 @@
+Availability: Public
+Crumbs: %ROOT%
+Title: Mobile Virtual Reality Development
+Description: Information over developing for mobile based virtual reality platforms.
+Type: Landing
+Version: 4.16
+Parent: Platforms
+Order: 1
+topic-icon:%ROOT%/platform_icon.png
+
+![](Platforms\VR\DesktopVR\Header_00.png)
+
+Unreal Engine 4 (UE4) supports a wide range of mobile based Virtual Reality (VR) Head Mounted Displays (HMD's), Motion Controllers and their corresponding Software Development kits. Contained with the pages linked below you will find a multitude of different documents that can help you get any of the supported mobile VR HMD's working with your UE4 project.
+
+## Supported Mobile VR Head Mounted Displays
+
+UE4 natively supports the following mobile based VR HMD's: 
+
+* **Samsung Gear VR**
+* **Google Cardboard & Daydream** 
+
+
+The fact that UE4 natively supports the above mobile HMD's means that you do not need to do any additional downloading or compiling to view your UE4 project in VR. Below you will find links to documentation that will explain how to get your UE4 project to work with the supported mobile HMD of your choice.
+
+## Samsung Gear VR
+
+[DIR(end:"1" output:"topic" type:"landing" tags:"GearVR" parent:"Platforms/GearVR")]
+
+## Google VR 
+
+[DIR(end:"1" output:"topic" type:"landing" tags:"GoogleVR" parent:"Platforms/GoogleVR")]
+
+
diff --git a/Engine/Documentation/Source/Platforms/VR/MonoFarFieldRendering/MonoFarFieldRendering.INT.udn b/Engine/Documentation/Source/Platforms/VR/MonoFarFieldRendering/MonoFarFieldRendering.INT.udn
new file mode 100644
--- /dev/null
+++ b/Engine/Documentation/Source/Platforms/VR/MonoFarFieldRendering/MonoFarFieldRendering.INT.udn
@@ -0,0 +1,288 @@
+Availability: Public
+Crumbs: %ROOT%
+Title: Monoscopic Far Field Rendering
+Description: Learn about Monoscopic Far Field Rendering, an advanced rendering feature for VR Mobile applications in Unreal Engine 4.
+Type: Overview
+SkillLevel: Advanced
+Version: 4.16
+Parent: Platforms/VR
+Order:
+Tags: Rendering
+Tags: Performance and Profiling
+Topic-Image:hybMonRend_topicImage.png
+related: Engine/Performance
+related: Engine/Performance/GPU
+related: Engine/Rendering
+related: Engine/Performance/StatCommands
+related: Platforms/VR/Profiling
+related: Platforms/VR/Profiling/Considerations
+related: Platforms/VR/Profiling/Overview
+Contributor: Oculus
+
+[TOC(start:2 end:3)]
+
+**Monoscopic Far Field** Rendering (or, more accurately, Hybrid Monoscopic Rendering) is an advanced rendering feature for VR Mobile applications being built in Unreal Engine 4 (UE4). Generally, Monoscopic Far Field Rendering improves 
+the performance of levels containing a significant amount of distant content. The reason that Monoscopic Far Field Rendering improves the performance of levels having a lot of distant content is that distant objects are rendered only 
+once for both eyes. At Epic, we've observed that rendering distant objects this way yields significant performance improvements for VR Mobile applications.
+
+[OBJECT:EmbeddedVideo]
+	[PARAMLITERAL:videoid]
+	bAyqVgv-MqA
+	[/PARAMLITERAL]
+	[PARAMLITERAL:width]
+	100
+	[/PARAMLITERAL]
+	[PARAMLITERAL:height]
+	56.25
+	[/PARAMLITERAL]
+	[PARAMLITERAL:units]
+	%
+	[/PARAMLITERAL]
+	[PARAMLITERAL:type]
+	simple
+	[/PARAMLITERAL]
+[/OBJECT]
+
+With the release of UE4, Version 4.16, Monoscopic Far Field Rendering is enabled with Multiview on GearVR devices, thereby ensuring that the stereoscopic near field rendering portion of your scene is optimized.
+
+## Enabling Far Field Monoscopic Rendering
+
+[REGION:note]
+Make sure that Mobile HDR is disabled before enabling the Monoscopic Far Field Rendering setting in UE4.
+[/REGION]
+
+When you create a new project, Monoscopic Far Field Rendering is disabled by default. To access the Monoscopic Far Field Rendering setting, open the **Project Settings** menu from Unreal Editor, and select **Rendering** from the side bar.
+You should be able to see the **VR** dropdown menu from inside of the **Project-Description** menu.
+
+[REGION:lightbox]
+	[![](uncheckedUI.png)(w:600)](uncheckedUI.png)
+[/REGION]
+
+[REGION:caption]
+	Click for full image.
+[/REGION]
+
+To enable Monoscopic Far Field Rendering, select the checkbox next to the **Monoscopic Far Field** rendering label.
+
+![](checkedUI.png)
+
+After enabling the **Monoscopic Far Field** setting, you'll be prompted to restart UE4. When you're ready, go ahead and restart UE4 to apply the **Monoscopic Far Field** rendering settings to your project.
+
+![](restartWarning.png)  
+
+### Adding a Monoscopic Camera
+
+When Monoscopic Far Field Rendering is disabled, your project has two existing stereoscopic cameras, one for each eye.
+
+[OBJECT:ComparisonSlider]
+	[PARAM:before]
+	![Disabled](Disabled.png)
+	[/PARAM]
+	[PARAM:after]
+	![Enabled](Enabled.png)
+	[/PARAM]
+[/OBJECT]
+
+When you enable Monoscopic Far Field Rendering, UE4 adds a third camera between the existing stereoscopic cameras. The new camera is a Monoscopic camera, which enables the hybrid monoscopic (and stereoscopic) rendering pipeline.
+
+### Stereoscopic Camera Frusta Symmetries
+
+Depending on your device's platform, plane geometry of the stereoscopic camera frusta may vary. For GearVR devices, the stereoscopic and monoscopic camera frusta are symmetric; however, for Rift devices, the stereoscopic and monoscopic camera frusta are asymmetric. 
+Please refer to the following diagrams to see how the camera frusta differ between GearVR and Rift devices.
+
+[REGION:imagetable]
+| ![FrustaGeometryGear.png](FrustaGeometryGear.png)(w:542) | ![FrustaGeometryRift.png](FrustaGeometryRift.png)(w:542) |
+| --- | --- |
+| GearVR (Symmetric Frusta) | Rift (Asymmetric Frusta) |
+[/REGION]
+
+As shown on the left, GearVR stereoscopic cameras have symmetric frusta, which means that the monoscopic camera shares the same projection matrix as the stereoscopic cameras.
+On Rift devices, however, the stereoscopic cameras are asymmetric (the inward angle towards the nose is less than the outward angle), which means that the monoscopic camera's 
+frustum is the union of both the left and right frusta, which makes the monoscopic render target marginally bigger than the stereoscopic render targets.
+
+### Split Plane Properties
+
+When UE4 constructs a monoscopic camera, the renderer adds a split (culling) plane. 
+
+![](AddingSplitPlane.png)
+
+[REGION:caption]
+The green split plane acts as the monoscopic cameras near plane.
+[/REGION]
+
+The following properties apply to the split plane being added to your project:
+
+*	The split plane is parallel to the project's z-axis.
+*	The split plane's default **Mono Culling Distance** is set to `750.0` centimeters (7.5 meters).
+*	Objects on the far side of the plane are rendered with the monoscopic camera.
+*	Objects straddling the culling plane are rendered with both camera types (monoscopic and stereoscopic).
+
+In the following image, the shaded (green) area represents the far side of the split plane, objects in this area are rendered in the monoscopic camera's depth buffer. 
+Objects falling on the near side of the split plane are rendered in the stereoscopic camera's depth buffer.
+
+![](MonoRenderArea.png)
+
+The following image comparison demonstrates the split plane in action:
+
+[REGION:imagetable]
+| ![StereoBuffer_banner.png](StereoBuffer_banner.png)(w:542) | ![MonoBuffer_banner.png](MonoBuffer_banner.png)(w:542) |
+| --- | --- |
+| Stereoscopic Depth Buffer | Monoscopic Depth Buffer |
+[/REGION]
+
+## Performance Consideration
+
+Thus far, you've learned that UE4 utilizes a depth buffer rendering technique to separate content being drawn by either stereoscopic or monoscopic cameras. From a high level, the UE4 mobile 
+forward renderer was [modified](#Bookmark1) to use depth buffers, determining if stereoscopic pixels have less depth than monoscopic pixels. When compared to an object-based rendering technique, 
+utilizing a pixel-based split (culling) plane can yield performance advantages, however, when objects go through the split plane, they have to be drawn monoscopically and stereoscopically 
+(even if the object's pixels are far enough to be drawn monoscopically). These extra draw calls typically generate unnecessary bandwidth, increasing GPU and CPU vertex costs.
+
+A good example of this performance hit can be seen in the following image: where the cliff's bounding sphere collides with the split plane, thereby causing the stereoscopic near field cameras to 
+render the cliff (even though the cliff's pixels should only be rendered by the monoscopic far-field camera).
+
+![](StereoNoCull.png)
+
+## Optimizing Performance
+
+To assess your project's performance, refer to the statistical output being generated by the **Rendering Hardware Interface (RHI)**. To view statistics being generated by RHI, perform the following steps:
+
+1.	From the Editor, click the **Viewport Options** button.
+
+	[REGION:lightbox]
+	[![](RHIStat_Step1.png)(w:600)](RHIStat_Step1.png)
+	[/REGION]
+	
+	[REGION:caption]
+		Click for full image.
+	[/REGION]
+
+1.	Hover your mouse cursor over the **Stat** option to open the Statistics menu.
+
+	[REGION:lightbox]
+	[![](RHIStat_Step2.png)(w:600)](RHIStat_Step2.png)
+	[/REGION]
+	
+	[REGION:caption]
+		Click for full image.
+	[/REGION]
+
+1.	Navigate to **Stat > Advanced**.
+
+	[REGION:lightbox]
+	[![](RHIStat_Step3.png)(w:600)](RHIStat_Step3.png)
+	[/REGION]
+	
+	[REGION:caption]
+		Click for full image.
+	[/REGION]
+
+1.	Finally, select the **RHI** option to display your project's RHI [Statistics](Engine/Performance/StatCommands).
+
+	[REGION:lightbox]
+	[![](RHIStat_Step4.png)(w:600)](RHIStat_Step4.png)
+	[/REGION]
+	
+	[REGION:caption]
+		Click for full image.
+	[/REGION]
+
+#### End Result
+
+At this point, you should see your project's RHI [Statistics](Engine/Performance/StatCommands).
+
+[REGION:lightbox]
+[![](RHIStat_Profiling.png)(w:600)](RHIStat_Profiling.png)
+[/REGION]
+
+[REGION:caption]
+	Click for full image.
+[/REGION]
+
+If you determine that your project's RHI profile is unacceptable, there are a few strategies you can take to optimize your project's performance, including:
+
+*	**Updating the Split Plane Distance**. There may be scenes where you have to modify the split plane distance to accommodate large (or small) rendering areas.
+*	**Updating Monoscopic Rendering Modes**. Changing Monoscopic Rendering Modes allows you to determine if certain objects need to be forced to render in mono.
+*	**Forcing Far Objects to Render in Mono**. Inevitably, there will be use cases to force certain objects (such as skyboxes or background meshes) to render in mono.
+
+### Updating the Split Plane Distance
+
+If you'd like to update the split plane's distance, open your project's **World Settings** in the Editor, locate the **VR** menu, and update the **Mono Culling Distance** to your desired distance.
+
+![](MonoCullingDistance.png)
+
+### Updating Monoscopic Rendering Modes
+
+After enabling monoscopic rendering, you may want to update your project's monoscopic rendering mode. Rendering a scene in different modes can help you debug potential performance issues with your project. 
+To update your application's monoscopic rendering mode, go ahead and enter `vr.MonoscopicFarFieldMode [0-4]` into the Editor's Command Console.
+
+![](UpdateMonoMode_console.png)
+
+By default, `vr.MonoscopicFarFieldMode` is set to `1`, however, if you want to change the default value, please refer to the following table.
+
+|Value|Mode|
+|-----|---|
+|**0**| Disable (Hybrid) Monoscopic Far Field Rendering. |
+|**1**| Enable (Hybrid) Monoscopic Far Field Rendering. This mode is set by default when Monoscopic Far Field Rendering is enabled. |
+|**2**| Only enable Stereoscopic Near Field rendering. |
+|**3**| Enable Stereoscopic Near Field rendering while also disabling the Far Field Pixel Depth Test. Basically, this option disables the split (culling) plane. |
+|**4**| Only enable Monoscopic Far Field Rendering. |
+
+### Forcing Far Objects to Render in Mono
+
+Let's say that you're working on the following scene.
+
+![](DebugScene.png)
+
+Now, you input the `vr.MonoscopicFarFieldMode 3` console command, enabling stereoscopic far field rendering while also disabling the split plane.
+
+![](StereoNoCull.png)
+
+[REGION:caption]
+Although you expect the cliff to be culled, you notice that the cliff mesh is still being rendered by the stereoscopic cameras, which can negatively affect performance.
+[/REGION]
+
+If (after observing your scene) you decide that want to force the cliff's mesh to render in Monoscopic Far Field mode, open your mesh using the **Blueprint Editor** and select the **Render in Mono** option under the **Rendering** menu in the **Details** panel.
+
+[REGION:lightbox]
+	[![](ForceRenderMono.png)(w:600)](ForceRenderMono.png)
+[/REGION]
+
+[REGION:caption]
+	Click for full image.
+[/REGION]
+
+[REGION:tip]
+When you set **Render in Mono** to true, the cliff component will only be rendered with the Monoscopic Near Field camera if a Head Mounted Device (HMD) is connected and **Monoscopic Far Field** rendering is enabled for the project.
+[/REGION]
+
+## Stereopsis
+
+The inspiration behind Hybrid Monoscopic Rendering is the ocular process of stereopsis. According to the [American Academy of Ophthalmology](https://www.aao.org/eye-health/ask-ophthalmologist-q/stereopsis), stereopsis is synonymous with depth perception. 
+When you look at an object, each eye looks at the object from a slightly different angle, with each eye sending images back to the brain, where the images are "composited" into a single image. The different images being sent back from each eye form a 
+binocular disparity, which the brain uses to extract depth information. 
+
+[REGION:note]
+The further an object is, the harder it is for the brain to extract the object's depth information.
+[/REGION]
+
+(#Bookmark1)
+###Leveraging Binocular Disparity
+
+Leveraging the fact that the brain has a hard time extracting depth information from far away objects, the UE4 mobile forward renderer was modified to perform the following (in order):
+
+1.	Render opaque objects with stereoscopic cameras.
+
+	![](Stereo.png)
+
+1.	Shift and combine the output to create a monoscopic occlusion mask, pre-populating the monoscopic depth buffer.
+
+1.	Render opaque objects with the monoscopic camera.
+
+	![](MonoOnly.png)
+
+1.	Composite the monoscopic camera's result into the stereo depth buffers.
+
+	![](MonoStereo.png)
+
+1.	Render all transparent content and perform all post-processing with stereoscopic cameras.
+
diff --git a/Engine/Documentation/Source/Platforms/VR/MotionController/MotionControllerComponent.INT.udn b/Engine/Documentation/Source/Platforms/VR/MotionController/MotionControllerComponent.INT.udn
--- a/Engine/Documentation/Source/Platforms/VR/MotionController/MotionControllerComponent.INT.udn
+++ b/Engine/Documentation/Source/Platforms/VR/MotionController/MotionControllerComponent.INT.udn
@@ -1,33 +1,38 @@
-Availability:Public
-Title:Motion Controller Component Setup
-Crumbs:%ROOT%, Platforms
-Description:Information on how to setup Motion Controllers for VR interaction.
-Navigation:platform
-Version:4.9
-Skilllevel:Intermediate
+Availability: Public
+Crumbs: %ROOT%
+Title: Motion Controller Component Setup
+Description: Information on how to setup Motion Controllers for VR interaction.
+Type: Overview
+SkillLevel:Beginner
+Version:4.17
 parent:Platforms/VR
-type:overview
+Order: 2
+Tags: VR
 
 [TOC(start:2)]
 
-You can add a Motion Controller as a component to your Pawn / Character, set which "Hand" it is supposed to represent.
-In the following document, we will cover how you can set up your Virtual Reality project to make use of Motion Controllers.
+No matter which Virtual Reality platform you are developing for, adding support for Motion Controllers can add a level of immersion and interaction that is just not possible to achieve with a controller or mouse and keyboard. In the following How - To we will take a look at how to add Motion Controller support to the VR Platforms that support it. 
 
-[region:tip]
-Before going any further make sure that your VR Headset and Motion Controllers are plugged into your PC and required the drivers installed. 
-See the [VR Documents](Platforms/VR) for information on how to get various VR headsets to work with UE4.
-[/region]
+## Supported Platforms 
+
+The Motion Controller component that is found in the Components tab will work on the following VR platforms.
+
+* Oculus VR
+* Steam VR
+* Gear VR
+* Playstation VR
+
+If you do not see the platform you are developing for listed above, make sure to check that platform's documentation for how to go about setting use of Motion Controllers. 
 
 ## Motion Controller Setup
 
-Setting up a VR project to make use of Motion Controllers requires adding a Motion Controller Component to your Pawn and configuring it. 
-In this section, we will cover how you can quickly get Motion Controllers working with your project.
+In the following section we will take a look at how to add and setup the components that are needed for Motion Controllers to work.
 
 [region:note]
-Please note this setup can be applied to any Pawn in any UE4 project.
+This How - To has been written assuming that you have setup your Pawn to work with the VR Head Mounted Display (HMD) that you are developing for. If you are not sure how to do this, check please check out the getting [started guide](Platforms/VR) for the VR Head Mounted Display (HMD) you are developing for. 
 [/region]
 
-1. In the **Content Browser** open the **Player Pawn** Blueprint.
+1. First, inside the  **Content Browser** locate and open up your **Player Pawn** Blueprint.
 
     [REGION:lightbox]
     [![](Open_VR_Pawn.png)(w:427)](Open_VR_Pawn.png)
@@ -37,76 +42,78 @@ Please note this setup can be applied to any Pawn in any UE4 project.
     Click for full image.
     [/REGION]
 
-1. Once the Blueprint is opened, navigate to the **Viewport** tab and in the **Components** section click on the **Add Component** button to expose the components that can be added to this Blueprint. 
+1.  In the **Components** section click on the **Add Component** button to expose the components that can be added to this Blueprint. 
 
     [REGION:lightbox]
-    [![](Find_Components.png)(w:418)](Find_Components.png)
+    [![](Find_Components.png)(w:435)](Find_Components.png)
     [/REGION]
 
     [REGION:caption]
     Click for full image.
     [/REGION]
 
-1. From the list that is shown search for the Motion Controller using the word **Motion** as the search term. 
-    When the Motion Controller component has been located, click on it to add it to the components list.
+1. Input **Motion** in the search box and then click on the **Motion Controller** component to add it to the components list, giving it a name of **MC_Left**..
 
-    [OBJECT:EmbeddedVideo]
-    [PARAMLITERAL:width]
-    640
-    [/PARAMLITERAL]
-    [PARAMLITERAL:height]
-    360
-    [/PARAMLITERAL]
-    [PARAMLITERAL:videoid]
-    kcyHnfpmWbU
-    [/PARAMLITERAL]
-    [/OBJECT]
+	![](Add_MC_Component.png)
 
-1. Once the Motion Controller has been added, select it by clicking on it and then in the **Details** panel under the **Motion Controller** section make sure the **Hand** is set to **Left**.
+1. Select the newly added Motion Controller component and over in the **Details** panel under the **Motion Controller** section make sure the **Hand** is set to **Left**.
 
     [REGION:lightbox]
-    [![](Setting_Up_Components.png)(w:418)](Setting_Up_Components.png)
+    [![](Setting_Up_Components.png)(w:787)](Setting_Up_Components.png)
     [/REGION]
 
     [REGION:caption]
     Click for full image.
     [/REGION]
 
-1. To give the Motion Controller a visual representation in the world, select the **Motion Controller Component** in the **Components panel**    and click the **Add Component** button to add a **Static Mesh Component**.
-    Drag the Static Mesh Component on to the Motion Controller to parent it to the Motion Controller Component then, 
-    in the Details panel, assign a Static Mesh asset to the Static Mesh property of the new component.
+1. Next, select the **Motion Controller Component** in the **Components panel** and click the **Add Component** button and then search for and add a **Static Mesh Component**, calling it **SM_Left**. 
 
-    [OBJECT:EmbeddedVideo]
-    [PARAMLITERAL:width]
-    640
-    [/PARAMLITERAL]
-    [PARAMLITERAL:height]
-    360
-    [/PARAMLITERAL]
-    [PARAMLITERAL:videoid]
-    x4vjDqnHQnk
-    [/PARAMLITERAL]
-    [/OBJECT]
+	![](Add_SM_To_MC.png)
 
-    [region:tip]
-    If your project does not currently have any Static Mesh assets, import one of your meshes or click the Marketplace button in the toolbar to visit the Marketplace and download some assets to use in your project.
-    If the Static Mesh you selected is too big or too small, you can scale to your likening using the Scale section of the Transform section in the Details panel.
-    [/region]
+	[region:note]
+	Make sure that the Static Mesh Component is a child of the Motion Controller Component otherwise the Static Mesh will not follow along when the Motion Controller is moved.
+	[/region]
 
-1. Add another Motion Controller Component parented to the root and set it to be the **Right Hand**. Add a Static Mesh Component parented to the new Motion Controller Component and applied the right-hand mesh to it.
+1. Now, in the Details panel of the Static Mesh Component under the **Static Mesh** section input a Static Mesh to represent what the Motion Controllers will look like. For this example we are using a simple box, but feel free to use any Static Mesh you have available.
 
     [REGION:lightbox]
-    [![](Finshed_MC_Setup.png)(w:418)](Finshed_MC_Setup.png)
+    [![](SM_MC_Static_Mesh.png)(w:787)](SM_MC_Static_Mesh.png)
     [/REGION]
 
     [REGION:caption]
     Click for full image.
     [/REGION]
 
-    [region:note]
-    In the image above, the two cone shaped meshes have been moved apart to show that there are two of them. 
-    When setting up your project leave both Static Mesh Components at 0,0,0 so that they are located right on top of one another.
-    [/region]
+1. Now, duplicate the entire Left hand Motion Controller setup and re-place the word **Left** with **Right**. Also make sure to change the which hand this Motion Controller is used for by going to the Motion Controller component and then changing **Hand** from Left to **Right**.
+
+    [REGION:lightbox]
+    [![](Finshed_MC_Setup.png)(w:787)](Finshed_MC_Setup.png)
+    [/REGION]
+
+    [REGION:caption]
+    Click for full image.
+    [/REGION]
+
+1. Compile and save your Pawn Blueprint, make sure it is placed in your test level and then launch your project. When you put on your HMD and pick up your Motion Controlers you should be able to now do what is shown in the following video.
+
+	[OBJECT:EmbeddedVideo]
+		[PARAMLITERAL:videoid]
+		CD6-Mje1y6g
+		[/PARAMLITERAL]
+		[PARAMLITERAL:width]
+		100
+		[/PARAMLITERAL]
+		[PARAMLITERAL:height]
+		56.25
+		[/PARAMLITERAL]
+		[PARAMLITERAL:units]
+		%
+		[/PARAMLITERAL]
+		[PARAMLITERAL:type]
+		simple
+		[/PARAMLITERAL]
+	[/OBJECT]
+
 
 ## Working with Motion Controllers
 
@@ -118,12 +125,11 @@ If you search inside of the Blueprint Editor **Palette** using **MotionControlle
 
 ![](Working_With_MC.png)
 
-When looking at the available Motion Controller input options, there is a capital **L** or **R** next to the name of what that value is supposed to control. 
-This L or R means that this is for the **Left** or **Right** Motion Controller.
+Make sure to note that there are seperate inputs for both the Left and Right Motion Controllers. The captial **L** or **R** next to the name of the input node signfies if this is for the left or right Motion Controller.
 
 ![](MC_Left_Right.png)
 
-
+<!--
 ## Using Motion Controllers
 
 In the following section, we will set the Motion Controllers up so that when you press the trigger button on the **Right-Handed** Motion Controller, a Blueprint containing a Static Mesh is spawned into the world, while pressing the trigger button on the **Left-Handed** Motion Controller will remove the mesh from the world. 
@@ -211,6 +217,7 @@ In the following section, we will set the Motion Controllers up so that when you
         VXUowSUKT9I
         [/PARAMLITERAL]
     [/OBJECT]
+-->
 
 ## Training Streams
 
diff --git a/Engine/Documentation/Source/Platforms/VR/Profiling/Tools/ProfilingTools.INT.udn b/Engine/Documentation/Source/Platforms/VR/Profiling/Tools/ProfilingTools.INT.udn
--- a/Engine/Documentation/Source/Platforms/VR/Profiling/Tools/ProfilingTools.INT.udn
+++ b/Engine/Documentation/Source/Platforms/VR/Profiling/Tools/ProfilingTools.INT.udn
@@ -20,8 +20,12 @@ For GPU profiling, there are more options.  First of all, and most specifically,
 
 Additionally, you could go further drilling into more systemic problems using both GPUView and Nsight. 
 
+
 ## Oculus PerfHud
 
+[Include:Platforms\VR\Profiling\Tools\#OculusPerfHud]
+<!--
+[EXCERPT:OculusPerfHud]
 Oculus PerfHud is a tool that can be displayed in game, whether in editor or in packaged builds, that will help verify the actual CPU and GPU timings, excluding any throttling done by the application.  There's a great overview of the screens 
 and what the values mean in Oculus' [documentation](https://developer3.oculus.com/documentation/pcsdk/latest/concepts/dg-hud/).
 
@@ -41,7 +45,8 @@ where N is:
 |	3	| Shows render timing info for application |
 |	4	| Shows render timing info for OVR compositor |
 |	5	| Shows SDK & HMD version Info |
-
+[/EXCERPT:OculusPerfHud]
+-->
 ## SteamVR Frame Timing Tool
 
 The SteamVR Frame Timing will allow you see how your VR project is performing in both the UE4 Editor or in a packaged build. This tool will help verify the actual CPU and GPU timings, excluding any throttling done by the application. There's a great overview of how to open the tool as well as what all of the values relate to on the 
diff --git a/Engine/Documentation/Source/Platforms/VR/SplashScreens/VRSplashScreenNodeReference/VRSplashScreenNodeReference.INT.udn b/Engine/Documentation/Source/Platforms/VR/SplashScreens/VRSplashScreenNodeReference/VRSplashScreenNodeReference.INT.udn
new file mode 100644
--- /dev/null
+++ b/Engine/Documentation/Source/Platforms/VR/SplashScreens/VRSplashScreenNodeReference/VRSplashScreenNodeReference.INT.udn
@@ -0,0 +1,131 @@
+Availability: Public 
+Crumbs: %ROOT%
+Title: VR Splash Screen Node Referance
+Description: Descriptions of what each node for VR Splash Screens performs.
+Type: Reference
+Version: 4.17
+Parent: Platforms/VR/SplashScreens
+Order:
+Tags: VR 
+
+[TOC(start:1)]
+
+## VR Splash Screen Node Reference
+
+The following breakdown covers what each one of the properties on the VR splash screen nodes does. When looking at the VR splash screen nodes, you will notice that there is a generic type and a type specific to Oculus. You should only use the Oculus-specific types when you are targeting Oculus HMDs. All other HMDs should use the generic type.
+
+## Generic Splash Screens
+
+* ### Enable Auto Loading Splash Screen
+
+    ![](Platforms\VR\SplashScreens\VRLS_EnableAutoLoading.png)
+
+    |Setting  Name| Description|
+    |---------------|----------------------------------------------------|
+    |In Auto Show Enabled|If true, load the splash screen when a level is loaded.|
+
+* ### Hide Splash Screen
+
+    ![](Platforms\VR\SplashScreens\VRLS_Hide_Splash.png)
+
+    |Setting  Name| Description|
+    |---------------|----------------------------------------------------|
+    |Hide Splash Screen|Hides the splash screen.|
+
+* ### Set Splash Screen
+
+    ![](Platforms\VR\SplashScreens\VRLS_Set_Splash.png)
+
+    |Setting  Name| Description|
+    |---------------|----------------------------------------------------|
+    |Texture|Texture to use for the splash. Must be in B8R8G8A8 format.|
+    |Scale|Scale of the Texture.|
+    |Offest|Position from which to start rendering the Texture.|
+    |Show Loading Movie|Show a "loading" movie.|
+    |Show on Set|Show the Texture specified when the compositor is called.|
+
+* ### Show Splash Screen
+
+    ![](Platforms\VR\SplashScreens\VRLS_Show_Splash.png)
+
+    |Setting  Name| Description|
+    |---------------|----------------------------------------------------|
+    |Show Splash Screen|Shows the splash screen.|
+
+## Oculus Specific Splash Screens
+
+The following breakdown covers what each of the Oculus-specific splash screen nodes will do. Keep in mind that the following nodes should only be used when you are targeting Oculus HMDs.
+
+* ### Add Loading Splash Screen
+
+    ![](Platforms\VR\SplashScreens\VRLS_OC_Add_Loading.png)
+
+    |Setting  Name| Description|
+    |---------------|----------------------------------------------------|
+    |Texture|A Texture asset to be used for the splash screen. Gear VR uses it as a path for loading icon; all other params are currently ignored by Gear VR.|
+    |Translation in Meters|Initial translation of the center of the splash screen in meters.|
+    |Rotation|Initial rotation of the splash screen, with the pivot at the center of the splash screen.|
+    |Size in Meters|Size, in meters, of the quad with the splash screen.|
+    |Delta Rotation|Incremental rotation added each 2nd frame to the quad's transform. The quad pivots around its center.|
+    |Clear Before Add|If true, clears splash before adding a new one.|
+
+* ### Clear Loading Splash Screen
+
+    ![](Platforms\VR\SplashScreens\VRLS_OC_Clear_Loading_Splash.png)
+
+    |Setting  Name| Description|
+    |---------------|----------------------------------------------------|
+    |Clear Loading Splash Screen|Removes all splash screens.|
+
+* ### Enable Auto Loading Splash Screen
+
+    ![](Platforms\VR\SplashScreens\VRLS_OC_Enable_Auto_Load.png)
+
+    |Setting  Name| Description|
+    |---------------|----------------------------------------------------|
+    |Auto Show Enabled|If true, automatic showing of splash screens is enabled when a map is being loaded.|
+
+* ### Hide Loading Icon
+
+    ![](Platforms\VR\SplashScreens\VRLS_OC_Hide_Loading.png)
+
+    |Setting  Name| Description|
+    |---------------|----------------------------------------------------|
+    |Hide Loading Icon|Clears the loading icon. This call will clear all the splashes.|
+
+* ### Hide Loading Splash Screen
+
+    ![](Platforms\VR\SplashScreens\VRLS_OC_Hide_Loading_SS.png)
+
+    |Setting  Name| Description|
+    |---------------|----------------------------------------------------|
+    |Clear|Clear all splash screen after hiding.|
+
+* ### Is Auto Loading Splash Screen Enabled
+
+    ![](Platforms\VR\SplashScreens\VRLS_OC_IsAutoLoasding.png)
+
+    |Setting  Name| Description|
+    |---------------|----------------------------------------------------|
+    |Is Auto Loading Splash Screen Enabled| Returns true if the splash screen is automatically shown when LoadMap is called.|
+
+* ### Is Loading Icon Enabled
+
+    ![](Platforms\VR\SplashScreens\VRLS_OC_IsLoadingIconEnabled.png)
+
+    |Setting  Name| Description|
+    |---------------|----------------------------------------------------|
+    |Is Loading Icon Enabled| Returns true if the splash screen is in loading icon mode.|
+
+* ### Show Loading Splash Screen
+
+    ![](Platforms\VR\SplashScreens\VRLS_OC_ShowLoadingSplashScreen.png)
+
+    |Setting  Name| Description|
+    |---------------|----------------------------------------------------|
+    |Show Loading Splash Screen| Shows loading splash screen.|
+
+
+
+
+
