Copying //UE4/Dev-Documentation to //UE4/Samples-Main (Source: //UE4/Dev-Documentation @ 3623665)

diff --git a/Engine/Documentation/Source/Platforms/VR/SplashScreens/VRSplashScreens.INT.udn b/Engine/Documentation/Source/Platforms/VR/SplashScreens/VRSplashScreens.INT.udn
--- a/Engine/Documentation/Source/Platforms/VR/SplashScreens/VRSplashScreens.INT.udn
+++ b/Engine/Documentation/Source/Platforms/VR/SplashScreens/VRSplashScreens.INT.udn
@@ -1,92 +1,143 @@
-Availability:Docs
+Availability:Public
 Title:VR Splash Screens
-Crumbs:%ROOT%, Platforms, Platforms/VR
-Description:Setting up and working with VR splash screens and movies.
-Navigation:platform
-Version: 4.14
+Crumbs:%ROOT%
+Description:Setting up and working with VR Splash Screens and movies.
+Version: 4.17
 parent:Platforms/VR
 Type:overview
 Skilllevel: Intermediate
+Tags:VR
+Order:2
 
 [TOC(start:2)]
 
-Unreal Engine 4 (UE4) offers support for displaying movies or images while a level loads when using Oculus, Gear VR, SteamVR, and PSVR Head Mounted Displays (HMDs). The **loading** movie or image will run on the rendering thread and can also help to mask frame rate issues while your content is loading. In the following document, we will take a look at how you can set these VR splash screens in your UE4 projects.
+Unreal Engine 4 (UE4) offers support for displaying videos or images while a level loads when using Oculus, Gear VR, SteamVR, and PSVR Head Mounted Displays (HMDs). The loading movie or image will run on the rendering thread and can also help to mask frame rate issues while your content is loading. In the following document, we will take a look at how you can set these VR **Splash Screens** in your UE4 projects.
 
 ## VR Splash Screen Supported Platforms 
 
-VR Splash Screens will only work when used with the following VR platforms. 
+VR Splash Screens will work on the following VR platforms. 
 
-* **Oculus VR**
-* **Gear VR**
-* **Steam VR**
-* **PSVR**
+* [Oculus VR](https://www.oculus.com/)
+* [Steam VR](http://store.steampowered.com/steamvr)
+* [PSVR](https://www.playstation.com/en-us/explore/playstation-vr/)
+* [Gear VR](http://www.samsung.com/global/galaxy/gear-vr/)
+* [Google Daydream](https://vr.google.com/daydream/)
 
 ## VR Splash Screen Project Setup
 
-Before you can setup and use the VR splash screens you will first need to create a new **Game Mode** and **Pawn** Blueprint to make it easier to work with the VR splash screens. In the following section, we will go over what Blueprints need to be created and how they need to be set up.
+Before you can setup and use the VR Splash Screens you will first need to create a new **Game Mode** and **Pawn** Blueprint to make it easier to work with the VR Splash Screens. In the following section, we will go over what Blueprints need to be created and how they need to be set up.
 
-1. Right - click in the **Content Browser**, and from the **Create Basic Asset** list select the **Blueprint Class** option.
+1. Right-click in the **Content Browser**, and from the **Create Basic Asset** list select the **Blueprint Class** option.
 
     ![](VRLS_Intro_01.png)
 
 1. In the **All Classes** section of the **Pick Parent Class** dialogue box, search for **Game Mode**. When Game Mode is located, select it and then press the **Select** button and give it a name of **VR_GameMode**.
 
-    ![](VRLS_Intro_02.png)
+    [REGION:lightbox]
+    [![](VRLS_Intro_02.png)(w:272)](VRLS_Intro_02.png)
+    [/REGION]
+
+    [REGION:caption]
+    Click for full image.
+    [/REGION]
 
 1. Repeat the above process, but this time select the **Pawn** option from the Pick Parent Class dialogue box and call this new Pawn, **VR_Pawn**.
 
-    ![](VRLS_Intro_03.png)
+    [REGION:lightbox]
+    [![](VRLS_Intro_03.png)(w:272)](VRLS_Intro_03.png)
+    [/REGION]
+
+    [REGION:caption]
+    Click for full image.
+    [/REGION]
 
 1. Open up the VR_GameMode Blueprint and in the **Classes** section, set the **Default Pawn Class** to the VR_Pawn that was created. When completed, Compile and Save the VR_Gamemode Blueprint.
 
-    ![](VRLS_Intro_04.png)
+    [REGION:lightbox]
+    [![](VRLS_Intro_04.png)(w:352)](VRLS_Intro_04.png)
+    [/REGION]
+
+    [REGION:caption]
+    Click for full image.
+    [/REGION]
 
-Now that the Pawn and Game Mode have been set up, we will take a look at how to set up the VR splash screens in the Pawn's Event Graph.
+Now that the Pawn and Game Mode have been set up, we will take a look at how to set up a VR Splash Screen in your Pawn's **Event Graph**. 
 
 ## VR Splash Screen Setup
 
-In the following section, we will go over how you set up a VR splash screen in your Pawns Event Graph.
+In the following section, we will go over how you set up a VR Splash Screen in your Pawns Event Graph. 
 
 [region:note]
-Please note that the following is showing the simplest way to set VR splash screens up. While each project has unique needs, the overall setup for VR splash screens will be very similar to what is described.
+Please note that this example shows the quickest way to get the Auto Loading Splash screens working in your project. When using this in production you will need to make sure to manually show and hide the Splash Screens during level transitions to make sure that you account for the time it takes for a level to finish streaming in.
 [/region]
-1. Open up VR_Pawn and go to the **Event Graph**. Right-click inside the Event Graph and search for and add the following Blueprint nodes:
 
-    * Enable Auto Loading Splash Screen
-    * Set Splash Screen
-    * Hide Splash Screen
+1. Open up VR_Pawn to the **Event Graph** and the following Blueprint nodes to the Event Graph:
+
+    * [**Enable Auto Loading Splash Screen**](Platforms\VR\SplashScreens\VRSplashScreenNodeReference#EnableAutoLoadingSplashScreen)
+    * [**Set Splash Screen**](Platforms\VR\SplashScreens\VRSplashScreenNodeReference#SetSplashScreen)
+    * [**Hide Splash Screen**](Platforms\VR\SplashScreens\VRSplashScreenNodeReference#HideSplashScreen)
 
-    ![](VRLS_Setup_01.png)
+    [REGION:lightbox]
+    [![](VRLS_Setup_01.png)(w:249)](VRLS_Setup_01.png)
+    [/REGION]
+
+    [REGION:caption]
+    Click for full image.
+    [/REGION]
 
 1. With the nodes now added to the Event graph, arrange and connect the nodes so that they match the setup in the following image.
 
-    ![](VRLS_Setup_02.png)
+    [REGION:lightbox]
+    [![](VRLS_Setup_02.png)(w:480)](VRLS_Setup_02.png)
+    [/REGION]
+
+    [REGION:caption]
+    Click for full image.
+    [/REGION]
 
-1. Since we want the VR splash screen to load any time a level is loaded, on the Enable Auto Loading Splash Screen node, check the **In Auto Show Enabled**.
+1. Since we want the VR Splash Screen to load any time a level is loaded, on the Enable Auto Loading Splash Screen node, check the **In Auto Show Enabled**.
+
+    [REGION:lightbox]
+    [![](VRLS_Setup_03.png)(w:480)](VRLS_Setup_03.png)
+    [/REGION]
+
+    [REGION:caption]
+    Click for full image.
+    [/REGION]
 
-    ![](VRLS_Setup_03.png)
+1. The Set Splash Screen node is where you specify which Texture to use for the VR Splash Screen. To set a VR Splash Screen Texture, select a Texture in the Content Browser and then press the small white arrow in the Texture section of the Set Splash Screen node to use that Texture as the VR splash screen.
 
-1. The Set Splash Screen node is where you specify which Texture to use for the VR splash screen. To set a VR splash screen Texture, select a Texture in the Content Browser and then press the small white arrow in the Texture section of the Set Splash Screen node to use that Texture as the VR splash screen.
+    [REGION:lightbox]
+    [![](VRLS_Setup_04.png)(w:408)](VRLS_Setup_04.png)
+    [/REGION]
 
-    ![](VRLS_Setup_04.png)
+    [REGION:caption]
+    Click for full image.
+    [/REGION]
     [region:note]
     For this example, a noise Texture was created in Photoshop and then imported into UE4 and used. Feel free to use any Texture you want for this section.
     [/region]
     [region:warning]
-    If you are targeting **SteamVR**, you will need to make sure that you set the **Compression Settings** of the Texture you are using for the splash screen to **UserInterface2D(RGBA)** or it will not be displayed when the splash screen is called.
+    If you are targeting **SteamVR**, you will need to make sure that you set the **Compression Settings** of the Texture you are using for the Splash Screen to **UserInterface2D(RGBA)** or it will not be displayed when the Splash Screen is called.
 
     ![](VIL_SteamVR_Comp_Settings.png)
     [/region]
 
 1. Finally, on the Set Splash Screen node check the **Show on Set** option to make sure that the Texture you specified is displayed when a level is loaded.
 
-    ![](VRLS_Setup_05.png)
+    [REGION:lightbox]
+    [![](VRLS_Setup_05.png)(w:408)](VRLS_Setup_05.png)
+    [/REGION]
+
+    [REGION:caption]
+    Click for full image.
+    [/REGION]
 
-With the VR splash screen now set up in the next section, we will take a look at how to go about calling the VR splash screen when a new level is loaded.
+With the VRSplash Screen now set up in the next section, we will take a look at how to go about calling the VR Splash Screen when a new level is loaded.
 
 ## VR Splash Screen Use
 
-Now that we have the Pawn set up to use the VR splash screens, we will go over how to display the VR Splash Screens when a level is loaded.
+Now that we have the Pawn set up to use the VR Splash Screens, we will go over how to display the VR Splash Screens when a level is loaded.
 
 1. For this to work, you will need to create **two** levels that we can load. To create the levels, go to the **Main Toolbar** and select the **File** > **New Level** option. When the **New Level** popup box is displayed, select the **Default** level option.
 
@@ -112,7 +163,7 @@ Now that we have the Pawn set up to use the VR splash screens, we will go over h
 
     ![](VLI_NewLevel_06.png)
 
-1. Open up the VR_Pawn Blueprint, if not already open, and in the Event Graph, add key inputs for the **R** and **T** keys.
+1. Open up the VR_Pawn Blueprint and in the Event Graph, add key inputs for the **R** and **T** keys.
 
     [REGION:lightbox]
     [![](VLI_NewLevel_02.png)(w:746)](VLI_NewLevel_02.png)
@@ -142,7 +193,7 @@ Now that we have the Pawn set up to use the VR splash screens, we will go over h
     Click for full image.
     [/REGION]
 
-1. Press the **Compile** and **Save** buttons to compile and save your Pawn Blueprint and then launch your project using the VR Preview.
+1. Push the **Compile** and **Save** buttons to compile and save your Pawn Blueprint and then launch your project using the VR Preview.
 
     [REGION:lightbox]
     [![](VLI_NewLevel_05.png)(w:746)](VLI_NewLevel_05.png)
@@ -152,138 +203,31 @@ Now that we have the Pawn set up to use the VR splash screens, we will go over h
     Click for full image.
     [/REGION]
 
-1. Put your HMD on and then press the **R** key to load Level_01. While the level is being loaded, you should see the image that you specified in the Texture input of the Set Splash Screen node. Once Level_01 loads, press the **T** key to load Level_00, and you should now see the same image being displayed.
+1. Put your HMD on and then press the **R** key to load Level_01. While the level is being loaded, you should see the image that you specified in the Texture input of the Set Splash Screen node. Once Level_01 loads press the **T** key to load Level_00, and you should now see the same image being displayed.
 
     [OBJECT:EmbeddedVideo]
+        [PARAMLITERAL:videoid]
+        5x5zH8JE0pQ
+        [/PARAMLITERAL]
         [PARAMLITERAL:width]
-        640
+        100
         [/PARAMLITERAL]
         [PARAMLITERAL:height]
-        360
+        56.25
         [/PARAMLITERAL]
-        [PARAMLITERAL:videoid]
-        DWkuJhEOtwc
+        [PARAMLITERAL:units]
+        %
+        [/PARAMLITERAL]
+        [PARAMLITERAL:type]
+        simple
         [/PARAMLITERAL]
     [/OBJECT]
 
-With the ability to change levels on button presses now working, the next step would be to try and get this working when a player enters a level streaming volume or trigger as you would in any game or interactive project.
-
-## VR Splash Screen Node Reference
-
-The following breakdown covers what each one of the properties on the VR splash screen nodes does. When looking at the VR splash screen nodes, you will notice that there is a generic type and a type specific to Oculus. You should only use the Oculus-specific types when you are targeting Oculus HMDs. All other HMDs should use the generic type.
-
-### Generic Splash Screens
-
-* Enable Auto Loading Splash Screen
-
-    ![](VRLS_EnableAutoLoading.png)
-
-    |Setting  Name| Description|
-    |---------------|----------------------------------------------------|
-    |In Auto Show Enabled|If true, load the splash screen when a level is loaded.|
-
-* Hide Splash Screen
-
-    ![](VRLS_Hide_Splash.png)
-
-    |Setting  Name| Description|
-    |---------------|----------------------------------------------------|
-    |Hide Splash Screen|Hides the splash screen.|
-
-* Set Splash Screen
-
-    ![](VRLS_Set_Splash.png)
-
-    |Setting  Name| Description|
-    |---------------|----------------------------------------------------|
-    |Texture|Texture to use for the splash. Must be in B8R8G8A8 format.|
-    |Scale|Scale of the Texture.|
-    |Offest|Position from which to start rendering the Texture.|
-    |Show Loading Movie|Show a "loading" movie.|
-    |Show on Set|Show the Texture specified when the compositor is called.|
-
-* Show Splash Screen
-
-    ![](VRLS_Show_Splash.png)
-
-    |Setting  Name| Description|
-    |---------------|----------------------------------------------------|
-    |Show Splash Screen|Shows the splash screen.|
-
-### Oculus Specific Splash Screens
-
-The following breakdown covers what each of the Oculus-specific splash screen nodes will do. Keep in mind that the following nodes should only be used when you are targeting Oculus HMDs.
-
-* Add Loading Splash Screen
-
-    ![](VRLS_OC_Add_Loading.png)
-
-    |Setting  Name| Description|
-    |---------------|----------------------------------------------------|
-    |Texture|A Texture asset to be used for the splash screen. GearVR uses it as a path for loading icon; all other params are currently ignored by GearVR.|
-    |Translation in Meters|Initial translation of the center of the splash screen in meters.|
-    |Rotation|Initial rotation of the splash screen, with the pivot at the center of the splash screen.|
-    |Size in Meters|Size, in meters, of the quad with the splash screen.|
-    |Delta Rotation|Incremental rotation added each 2nd frame to the quad's transform. The quad pivots around its center.|
-    |Clear Before Add|If true, clears splash before adding a new one.|
-
-* Clear Loading Splash Screen
-
-    ![](VRLS_OC_Clear_Loading_Splash.png)
-
-    |Setting  Name| Description|
-    |---------------|----------------------------------------------------|
-    |Clear Loading Splash Screen|Removes all splash screens.|
-
-* Enable Auto Loading Splash Screen
-
-    ![](VRLS_OC_Enable_Auto_Load.png)
-
-    |Setting  Name| Description|
-    |---------------|----------------------------------------------------|
-    |Auto Show Enabled|If true, automatic showing of splash screens is enabled when a map is being loaded.|
-
-* Hide Loading Icon
-
-    ![](VRLS_OC_Hide_Loading.png)
-
-    |Setting  Name| Description|
-    |---------------|----------------------------------------------------|
-    |Hide Loading Icon|Clears the loading icon. This call will clear all the splashes.|
-
-* Hide Loading Splash Screen
-
-    ![](VRLS_OC_Hide_Loading_SS.png)
-
-    |Setting  Name| Description|
-    |---------------|----------------------------------------------------|
-    |Clear|Clear all splash screen after hiding.|
-
-* Is Auto Loading Splash Screen Enabled
-
-    ![](VRLS_OC_IsAutoLoasding.png)
-
-    |Setting  Name| Description|
-    |---------------|----------------------------------------------------|
-    |Is Auto Loading Splash Screen Enabled| Returns true if the splash screen is automatically shown when LoadMap is called.|
-
-* Is Loading Icon Enabled
-
-    ![](VRLS_OC_IsLoadingIconEnabled.png)
-
-    |Setting  Name| Description|
-    |---------------|----------------------------------------------------|
-    |Is Loading Icon Enabled| Returns true if the splash screen is in loading icon mode.|
-
-* Show Loading Splash Screen
-
-    ![](VRLS_OC_ShowLoadingSplashScreen.png)
-
-    |Setting  Name| Description|
-    |---------------|----------------------------------------------------|
-    |Show Loading Splash Screen| Shows loading splash screen.|
-
+With the ability to change levels on button presses now working, the next step would be to set this up so that it will work when a player enters a Level Streaming Volume or trigger as you would in any game or interactive project.
 
 
+## VR Splash Screen Node Reference
 
+Clicking on the following link will take you to the VR Splash Screen node reference page. On this page you will find detailed descriptions for each node that is available for use with the VR Splash Screens.
 
+[**VR Splash Screen Node Reference**](Platforms\VR\SplashScreens\VRSplashScreenNodeReference)
diff --git a/Engine/Documentation/Source/Platforms/VR/StereoLayers/StereoLayers.INT.udn b/Engine/Documentation/Source/Platforms/VR/StereoLayers/StereoLayers.INT.udn
--- a/Engine/Documentation/Source/Platforms/VR/StereoLayers/StereoLayers.INT.udn
+++ b/Engine/Documentation/Source/Platforms/VR/StereoLayers/StereoLayers.INT.udn
@@ -1,9 +1,8 @@
-Availability:Docs
+Availability:Public
 Title:VR Stereo Layers
-Crumbs:%ROOT%, Platforms, Platforms/VR
+Crumbs:%ROOT%,
 Description:Setting up and working with VR Stereo Layers.
-Navigation:platform
-Version: 4.14
+Version: 4.17
 parent:Platforms/VR
 type:overview
 Skilllevel: Intermediate
@@ -12,26 +11,102 @@ tags:Oculus Rift
 tags:GearVR
 tags:SteamVR
 tags:PSVR
+order:1
+Topic-Image:VRSL_Topic_Image.png
+social-image:![](Platforms\VR\VirtualRealityDevelopment_topic.png)
 
 [TOC(start:2)]
 
-Virtuality Reality (VR) Stereo Layers allow you to send a separate Texture to the VR Head Mounted Display (HMD) and have it re-projected on its own in a separate rendering pass than the rest of the project.  This is extremely useful for User Interface (UI) elements that you want to attach to the HMD but not be affected by post processing or anti-aliasing. In the following document, we will go over how you set up and use Stereo Layers across all VR Platforms that are supported.  
+[EXCERPT:VRStereoLayersIntro]
 
+Virtual Reality (VR) **Stereo Layers** allow you to send a separate Texture to the VR Head Mounted Display (HMD) and have it re-projected in a separate rendering pass than the rest of the project.  This is extremely useful for User Interface (UI) elements that you want to attach to the HMD but not be affected by post processing or anti-aliasing. In the following document, we will go over how you set up and use Stereo Layers across all VR Platforms that are supported.  
+
+[/EXCERPT:VRStereoLayersIntro]
 
 ## Supported Platforms 
 
-VR Stereo Layers will only work when used with the following VR platforms. 
+VR Stereo Layers are supported on all VR platforms that Unreal Engine 4 (UE4) supports which include:
+
+* [**Oculus VR**](https://www.oculus.com/)
+* [**Gear VR**](http://www.samsung.com/global/galaxy/gear-vr/)
+* [**Google VR**](https://vr.google.com/)
+* [**Steam VR**](http://store.steampowered.com/steamvr)
+* [**PSVR**](https://www.playstation.com/en-us/explore/playstation-vr/)
+
+
+## Platform Differences
+
+Since the native VR platforms expose different features, the behavior of Stereo Layers can vary between platforms. This section lists most of them by the implementation.
+
+* **World Locked:** The layer position is in game world coordinates.
+* **Tracker Locked:** The layer position is relative to the real-world tracking space around the player
+* **Face Locked:** The layer stays in the same place in the HMD view regardless of position or rotation.
+
+Another way to view the position types is by which aspects of the player position and rotation are taken into account when calculating the layer position relative to the view point:
+
+|Position Type | World location|Head Position|Head Rotation|
+|--------------|--------------|-|-|
+|World Locked  | Yes | Yes|Yes|
+|Tracker Locked| No  | Yes|Yes|
+|Face Locked   | No  | No |No|
+
+### Steam VR
+
+* Only supports the Quad stereo layer shape.
+
+* Does not support the No Alpha flag to ignore the alpha channel. A workaround is to ensure that the texture does not have an alpha channel or is one-hundred percent opaque.
+
+* Performs backface culling on all layers, the layer is not visible from the reverse side.
+
+* When multiple layers have the same priority, layers are rendered by average distance from the HMD view.
+
+* Inside UE4 Textures used in Stereo Layers must be set to UserInterFace2D (RGBA) or they will not show up.
+
+### Oculus VR
+
+* Only supports the Quad stereo layer shape.
+
+* Supports the Support Depth flag on World and Tracker locked layers.
+
+* When multiple layers have the same priority, layers are rendered in a non-defined order. If you care about the ordering, use different priorities.
+
+### Gear VR 
+
+* Implements all three shape types.
+
+* Supports the Support Depth flag on World and Tracker locked layers.
+
+* When multiple layers have the same priority, layers are rendered in a non-defined order. If you care about the ordering, use different priorities.
+
+### Other VR Platforms
+
+Currently, all other VR platforms share a default layers implementation which will support the following.
+
+* Only supports the Quad stereo layer shape.
+
+* Mixing Face locked layer priorities with World and Tracker locked layers by setting the vr.StereoLayers.bMixLayerPriorities console variable is not supported.
+
+* When multiple layers have the same priority, layers are rendered in a non-defined order. If you care about the ordering, use different priorities.
+
+[EXCERPT:VRStereoLayers]
+
+## Layer Shapes
+
+Only layer shape supported on all platforms are quadrilinear layers. In addition to those Gear VR supports Cylinder and Cube layers.
+
+## Layer Draw Order
+
+The Stereo Layers exist outside the game world and are layered based on the Priority setting on each layer. Layers with higher priority are rendered on top of layers of lower priority. Additionally, face locked layers are always rendered on top of World or Tracker locked layers regardless of priority. This is different from UE 4.16 and earlier, where all layer types were mixed together. To get the old behavior, set the **vr.StereoLayers.bMixLayerPriorities** console variable to 1.
 
-* **Oculus VR**
-* **Gear VR**
-* **Steam VR**
-* **PSVR**
+[region:note]
+Layers will always be rendered on top of the 3D scene, Oculus Rift and Gear VR support mixing World locked and Tracker locked layers with the scene however by toggling a Supports Depth flag. This flag is ignored on all other platforms.
+[/region]
 
 ## Stereo Layer Setup
 
 To use Stereo Layers with your **VR Pawn**, you will need to do the following.
 
-1. Open up the **Pawn** that you want to use Stereo Layers with and then go to the **Viewport** pannel.
+1. Open up the **Pawn** that you want to use Stereo Layers with and then go to the **Viewport** panel.
 
     [REGION:lightbox]
     [![](VRSL_Pawn_Setup_01.png)(w:667)](VRSL_Pawn_Setup_01.png)
@@ -41,7 +116,7 @@ To use Stereo Layers with your **VR Pawn**, you will need to do the following.
     Click for full image.
     [/REGION]
 
-1. Click on the **Add Components** button and from the drop-down list that is displayed and search for the **Stereo Layer** Component. When you find it, click on it to add it to your **Pawn**.
+1. Click on the Add Components button from the drop-down list that is displayed and search for the Stereo Layer Component.
 
     [REGION:lightbox]
     [![](VRSL_Pawn_Setup_02.png)(w:667)](VRSL_Pawn_Setup_02.png)
@@ -65,7 +140,7 @@ To use Stereo Layers with your **VR Pawn**, you will need to do the following.
     Click for full image.
     [/REGION]
 
-1. Over in the **Details Panel**  expand the **Stereo Layer** section.
+1. Over in the **Details Panel**, expand the **Stereo Layer** section.
 
     [REGION:lightbox]
     [![](VRSL_Pawn_Setup_05.png)(w:667)](VRSL_Pawn_Setup_05.png)
@@ -87,21 +162,28 @@ To use Stereo Layers with your **VR Pawn**, you will need to do the following.
 
 1. Next press the **Play in VR** button to launch your project in VR. When you put your HMD on you should see the Texture you specified in Texture input being displayed in front of you. The Texture should also follow your HMD precisely and as quickly as you can move your head around like in the video below.
 
-    [OBJECT:EmbeddedVideo]
-        [PARAMLITERAL:width]
-        640
-        [/PARAMLITERAL]
-        [PARAMLITERAL:height]
-        360
-        [/PARAMLITERAL]
-        [PARAMLITERAL:videoid]
-        Eb_kndx6tac
-        [/PARAMLITERAL]
-    [/OBJECT]
+	[OBJECT:EmbeddedVideo]
+		[PARAMLITERAL:videoid]
+		Eb_kndx6tac
+		[/PARAMLITERAL]
+		[PARAMLITERAL:width]
+		100
+		[/PARAMLITERAL]
+		[PARAMLITERAL:height]
+		56.25
+		[/PARAMLITERAL]
+		[PARAMLITERAL:units]
+		%
+		[/PARAMLITERAL]
+		[PARAMLITERAL:type]
+		simple
+		[/PARAMLITERAL]
+	[/OBJECT]
 
     [region:note]
     If the Texture you selected does not show up when your project is viewed in VR, make sure you are using the correct Texture Compressions settings. If you are unsure what those settings should be for the HMD you are using, please check out the Stereo Layers Texture Compression Settings section.
     [/region]
+[/EXCERPT:VRStereoLayers]
 
 ## Stereo Layers Texture Compression Settings
 
@@ -117,30 +199,30 @@ Click for full image.
 
 |HMD Name| Required Texture Compression Settings|
 |---------------|----------------------------------------------------|
-|Oculus Rift| None|
-|Gear VR| None|
-| HTC Vive| UserInterFace2D (RGBA)|
-|PSVR| None|
+|**Oculus Rift**| None|
+|**Gear VR**| None|
+|**HTC Vive**| UserInterFace2D (RGBA)|
+|**PSVR**| None|
 
 ## Stereo Layer Properties Breakdown
 
-Stereo Layers have a number of properties that affect the way in which it is rendered. The following table breakdown each one of these settings and it's affect on the the Stereo Layers.
+Stereo Layers have some properties that affect the way in which it is rendered. The following table breaks down each one of these settings and its effect on the Stereo Layers.
 
 ![](VRSL_Prop_Breakdown_00.png)
 
 |Setting  Name| Description|
 |---------------|----------------------------------------------------|
-|Live Texture|True if the stereo layer texture needs to update itself every frame(scene capture, video, etc.)|
-|Supports Depth|True if the stereo layer needs to support depth intersections with the scene geometry, if available on the platform.|
-|No Alpha Channel|True if the texture should not use its own alpha channel(1.0 will be substituted).|
-|Texture|Texture displayed on the stereo layer (is stereoscopic textures are supported on the platform and more than one texture is provided, this will be the right eye).|
-|Quad Preserve Texture Ratio|True if the quad should internally set it's Y value based on the set texture's dimensions.|
-|Stereo Layer Type|Specifies how and where the quad is rendered to the screen.|
-|Stereo Layer Shape|Specifies which type of layer it is. Note that some shapes will be supported only on certain platforms!|
-|Priority|Render priority among all stereo layers, higher priority render on top of lower priority.|
-|Left Texture|Texture displayed on the stereo layer for left eye, if stereoscopic textures are supported on the platform.|
-|Quad Size|Size of the rendered stereo layer quad.|
-|UVRect|UV coordinates mapped to the quad face.|
-|Cylinder Radius|Radial size of the rendered stereo layer cylinder.|
-|Cylinder Overly Arc|Arc angle for the stereo layer cylinder.|
-|Cylinder Height|Height of the stereo layer cylinder.|
\ No newline at end of file
+|**Live Texture**|True if the stereo layer texture needs to update itself every frame (scene capture, video, etc.). |
+|**Supports Depth**|True if the stereo layer needs to support depth intersections with the scene geometry, if available on the platform.|
+|**No Alpha Channel**|True if the texture should not use its own alpha channel (1.0 will be substituted).|
+|**Texture**|Texture displayed on the stereo layer (if stereoscopic textures are supported on the platform, and more than one texture is provided, this will be the right eye).|
+|**Quad Preserve Texture Ratio**|True if the quad should internally set its Y value based on the set texture's dimensions.|
+|**Stereo Layer Type**|Specifies how and where the quad is rendered to the screen.|
+|**Stereo Layer Shape**|Specifies which type of layer it is. Note that some shapes will be supported only on certain platforms!|
+|**Priority**|Render priority among all stereo layers, higher priority render on top of lower priority.|
+|**Left Texture**|Texture displayed on the stereo layer for left eye, if stereoscopic textures are supported on the platform.|
+|**Quad Size**|Size of the rendered stereo layer quad.|
+|**UVRect**|UV coordinates mapped to the quad face.|
+|**Cylinder Radius**|Radial size of the rendered stereo layer cylinder.|
+|**Cylinder Overly Arc**|Arc angle for the stereo layer cylinder.|
+|**Cylinder Height**|Height of the stereo layer cylinder.|
