Copying //UE4/Dev-Documentation to //UE4/Samples-Main (Source: //UE4/Dev-Documentation @ 3623665)

diff --git a/Engine/Documentation/Source/Platforms/VR/VR.INT.udn b/Engine/Documentation/Source/Platforms/VR/VR.INT.udn
--- a/Engine/Documentation/Source/Platforms/VR/VR.INT.udn
+++ b/Engine/Documentation/Source/Platforms/VR/VR.INT.udn
@@ -3,7 +3,7 @@ Title:Virtual Reality Development
 Crumbs:%ROOT%, Platforms
 Description:Information over developing for virtual reality platforms.
 Navigation:platform
-Version: 4.9
+Version: 4.17
 Type:landing
 tags:VR
 tags:Platform
@@ -12,13 +12,13 @@ order:2
 topic-icon:%ROOT%/platform_icon.png
 topic-image:VirtualRealityDevelopment_topic.png
 
-## Virtual Reality Platforms
+## Virtual Reality Desktop Platforms
 
-[DIR(end:"1" output:"topic" type:"landing" tags:"platform" parent:"Platforms/VR")]
+[DIR(end:"1" output:"topic" type:"landing" tags:"platform" parent:"Platforms\VR\DesktopVR")]
 
-## Virtual Reality Topics
+## Virtual Reality Mobile Platforms
 
-[DIR(end:"1" output:"topic" parent:"Platforms/VR" tags:"!platform")]
+[DIR(end:"1" output:"topic" type:"landing" tags:"platform" parent:"Platforms\VR\MobileVR")]
 
 ## Training Streams
 
diff --git a/Engine/Documentation/Source/Platforms/VR/VRCameraRefactor/VRCameraRefactor.INT.udn b/Engine/Documentation/Source/Platforms/VR/VRCameraRefactor/VRCameraRefactor.INT.udn
--- a/Engine/Documentation/Source/Platforms/VR/VRCameraRefactor/VRCameraRefactor.INT.udn
+++ b/Engine/Documentation/Source/Platforms/VR/VRCameraRefactor/VRCameraRefactor.INT.udn
@@ -1,100 +1,77 @@
-Availability:Public
-Title:VR Camera Refactor
-Crumbs:%ROOT%, Platforms
-Description:Information covering the new VR camera setup.
-Navigation:platform
-Version:4.11
-Skilllevel:Intermediate
+Availability: Public
+Crumbs: %ROOT%
+Title: Attaching Items To the HMD
+Description:Information on attaching items to any HMD.
+Type: Overview
+SkillLevel:Beginner
+Version:4.17
 parent:Platforms/VR
-type:overview
+Order: 4
+Tags:VR
 
 [TOC(start:2)]
 
-
-
-With the release of Unreal Engine 4.11 Virtual Reality(VR) cameras now work a bit different than in previous versions of UE4. 
-The biggest change is that objects will now follow the camera simply by making the camera the objects parent.
-This also means that you will no longer be required to offset tracking objects by the Head Mounted Display(HMD) transform.
-Components attached to the Camera Actor will now keep their position relative to the camera, making it easy to add a helmet, collision, or other primitives to the player. The attached renderable components will now also receive the same **late update** transforms used in low-latency rendering to keep them perfectly in-sync with HMD motion.
-In the following document, we will go over all you need to know about how these changes will affect your past that are upgraded as well as the effects on future UE4 projects.
-
-## VR Camera Setup in UE 4.11 and Later
-
-[region:note]
-Make sure to note that the Camera Manager no longer dictates when a camera should track the HMD, that is all set directly on the camera itself through the **Lock to HMD** option.
-[/region]
-
-If you are creating a new VR project using UE 4.11 or later, the camera will already be setup to be in sync with the HMD as this functionality is now enabled by default. 
-However if you wish to disable this all you will need to do  is select the **Camera** in your **Player Blueprint** and then in the **Details** panel under the **Camera** 
-section uncheck the **Lock to HMD** option which you can see highlighted in the image below.
-
- ![](VR_Camera_Refactor_00.png)
-
-
-## Upgrading Previous VR Cameras to 4.11
-
-If you are upgrading your UE4's project from a previous release to the 4.11 release there are a few things you will need to do to ensure that your previous VR camera setup will work with the new 4.11 changes. 
-Depending on how your Camera is created and setup you will need to ensure that either of the following changes are made in order to get your VR camera working correctly. 
-
-### Base Camera Component
-
-In the case of a bare Camera Component, the camera transform is the HMD offset and a "base" will be required to move the player's head through a scene.
- If you see your VR camera on the ground after migrating to 4.11 you will need to insert a new scene transform to set the placement of the camera base as detailed in the image below.
-
- ![](VR_Camera_Refactor_01.png)
-
-
-### Camera Actor Component
-
-In the case of a Camera Actor, a base scene transform automatically gets inserted as a parent to the actor's own Camera Component giving it a complete transform chain as detailed in the image below. 
-
- ![](VR_Camera_Refactor_02.png)
-
-
+No matter which Head Mounted Display (HMD) you are developing for Unreal Engine 4 (UE4) offers a unified approach to attaching items to the HMD. Not only will this unified approach work with any HMD, the items that are attached will move perfectly in-sync with the HMD. In the following document, we will go over all you need to know about attaching items to the HMD.
 
 [region:note]
 Now if you want to get the position of the player in the world you can just get the Camera Actor's position.
 [/region]
 
 
-## Setting up objects to follow the HMD
+## Setting Up Objects To Follow The HMD
+
+You can setup a Static Mesh, Particle Effect, UI elements and many more items to follow the movment of the HMD by doing the following. 
 
-Setting up 3D Mesh, Particle Effects, or UI elements so that they can follow the movement of the HMD can be accomplished in the following few steps.
+1. First open up your Pawn Blueprint and navigate to the **Viewport** tab.
 
-1. First open up your Character Blueprint and then navigate to the **Viewport** tab of the Blueprint Editor.
+	[REGION:lightbox]
+	[![](VR_Follow_HMD_Setup_00.png)(w:737)](VR_Follow_HMD_Setup_00.png)
+	[/REGION]
 
-	![](VR_Follow_HMD_Setup_00.png)
+	[REGION:caption]
+	Click for full image.
+	[/REGION]
 
-1. Next from the **Component** tab, click on the **Add Component** option and then input the word **Cube** into the search box then click on the **Cube** mesh that is displayed to add it as a component
+1.  From the **Component** tab, click on  **Add Component** and then input the word **Cube** into the search box clicking on the **Cube** mesh that is displayed to add it as a component
 	
-	![](VR_Follow_HMD_Setup_01.png)
+	[REGION:lightbox]
+	[![](VR_Follow_HMD_Setup_01.png)(w:737)](VR_Follow_HMD_Setup_01.png)
+	[/REGION]
 
-1. With the Cube Static Mesh now added, select the Cube Static Mesh and then drag it onto the camera so that camera is now the parent of them mesh.
+	[REGION:caption]
+	Click for full image.
+	[/REGION]
+
+1. Select the Cube Static Mesh and drag it onto the Camera so that the Cmera is now the parent of the Cube Static Mesh.
 
 	![](VR_Follow_HMD_Setup_02.png)
 
-1. Now compile your Blueprint and then launch your project using the VR preview option from inside the UE4 editor. 
-When you put an HMD on then move your head around the Cube you attached will now accurately follow your head movements like in the video below.
+1. Now compile your Blueprint and then launch your project using the VR preview option from inside the UE4 editor. When you put an HMD on then move your head around the Cube you attached will now accurately follow your head movements like in the video below.
 
 	[OBJECT:EmbeddedVideo]
-	[PARAMLITERAL:width]
-	640
-	[/PARAMLITERAL]
-	[PARAMLITERAL:height]
-	360
-	[/PARAMLITERAL]
-	[PARAMLITERAL:videoid]
-	xzdA-JU3-js
-	[/PARAMLITERAL]
+		[PARAMLITERAL:videoid]
+		5brQDHHFR0w
+		[/PARAMLITERAL]
+		[PARAMLITERAL:width]
+		100
+		[/PARAMLITERAL]
+		[PARAMLITERAL:height]
+		56.25
+		[/PARAMLITERAL]
+		[PARAMLITERAL:units]
+		%
+		[/PARAMLITERAL]
+		[PARAMLITERAL:type]
+		simple
+		[/PARAMLITERAL]
 	[/OBJECT]
 
-## HMD and Player Position
+## HMD and Player World Position
 
-One of the other advantages the new camera refactor gives you is the ability to quickly get the world location of the camera / HMD by getting the location of the Camera Component. 
-In the image below you can see an example of this in action.
+You can also quickly get the exact world location player and their HMD by getting the location of the Camera Component using the following Blueprint setup.
 
 ![](VR_Get_Player_Location_00.png)
 
 In the above image whenever a user presses the T key on the keyboard the X, Y and Z position of the Camera in the world will then be output to the screen and log window.
 
-
+![](VR_Follow_HMD_Setup_03.png)
diff --git a/Engine/Documentation/Source/Platforms/VR/VRPerformance/VRPerformance.INT.udn b/Engine/Documentation/Source/Platforms/VR/VRPerformance/VRPerformance.INT.udn
--- a/Engine/Documentation/Source/Platforms/VR/VRPerformance/VRPerformance.INT.udn
+++ b/Engine/Documentation/Source/Platforms/VR/VRPerformance/VRPerformance.INT.udn
@@ -1,97 +1,118 @@
-Availability:Doc
-Title: VR Performance Testing
-Crumbs:%ROOT%, Platforms
-Description:
-Navigation:platform
-Version:4.10
-Skilllevel:Advanced
+Availability: Public
+Crumbs: %ROOT%
+Title:VR Performance Features
+Description: Information on different performance optmizations for VR.
+Type: Overview
+SkillLevel:Advanced
+Version:4.17
+Parent: Platforms/VR
+Order:4
+Tags:VR
 
 [TOC(start:2)]
 
-## VR Performance Overview
-Trying to figure out which computer parts are capable of producing the breathtaking imagery and immersive experiences that Virtual Reality (VR) Head Mounted Displays (HMD's) 
-like the Oculus Rift or HTC Vive produce can be a daunting and somewhat frustrating experience due to the many hardware choices available. 
-To help alleviate this frustration and ensure that any personal computer (PC) hardware that anyone is planning to purchase will meet the demanding rendering requirements of a UE4 powered VR experience, the following document has been put together.
-Inside of this  document, you will find detailed information about what PC hardware is recommended for VR development with UE4 and what type of performance you should expect to see on such hardware. 
-You will also find performance tracking data that tracks how different version of UE4 perform with different different HMD's and PC hardware.
+Due to the demanding rendering requirements Virtual Reality (VR) requires a number of performance optimizations and new rendering modes that are specific to VR have been introduced to Unreal Engine 4 (UE4). In the following document we will take a look at these features and how to use them in your UE4 VR projects.
 
-## Recommended PC Hardware for VR & UE4
 
-Due to the demanding rendering requirements of VR, the PC that is used to run a VR experience needs to have certain hardware inside of it to give users a smooth VR experience as possible. 
-According to [Oculus](https://www.oculus.com/en-us/blog/powering-the-rift/), for your PC to be able to run the smoothest VR experience as possible, the hardware inside of that PC needs to meet or exceed the following requirements.
+## VR Instanced Stereo
 
-* **CPU:** Intel i5-4590 equivalent or greater.
-* **GPU:** NVIDIA GTX 970 / AMD 290 equivalent or greater.
-* **RAM:** 8 GB of RAM or more.
-* **Opperating System:** Windows 7 SP1 or newer.
-* **USB Ports:** At least 2x USB 3.0 ports or more.
-* **HDMI Output:** HDMI 1.3 video output supporting a 297MHz clock via a direct output architecture.
+Instanced Stereo Rendering was introduced to help lessen the performance impact of VR in UE4. The feature can be enabled by checking the **Instanced Stereo** option that can be found under **Edit > Project Settings > Rendering > VR**. 
 
-Taking the recommended hardware spec's from Oculus into consideration the following computer parts list was generated. 
-The computer parts that you will find in the list below were picked based on how well the part in question met the following criteria. 
+[REGION:lightbox]
+[![](Platforms\VR\ContentSetup\VR_ISR.png)(w:376)](Platforms\VR\ContentSetup\VR_ISR.png)
+[/REGION]
 
-* How closely the part matched the recommend hardware specifications set forth by Oculus.
-* Number of positive customer reviews the part had. 
-* What the price of the part was.
+[REGION:caption]
+Click for full image.
+[/REGION]
 
-[region:note]
-Please note that while it might be possible to run VR experiences on hardware that does not meet the following specifications, it is not advisable as it could result in a poor VR experiance.
-[/region]
-
-### Recomened GPU
-While there are a plethora of different video cards on the market to choose from, there are only a few that have enough power to run a VR experience. 
-Choosing any of the following video cards, or video cards that have the same spec as the ones listed below will ensure that you have enough power to run a UE4 powered VR experience.
-
-**Minimum GPU**
-* [NVIDIA GTX 970 w 4GB of VRam](http://www.newegg.com/Product/Product.aspx?Item=N82E16814487136&cm_re=NVIDIA_GTX_970-_-14-487-136-_-Product)
-* [AMD Radeon R9 w 4GB of VRam](http://www.newegg.com/Product/Product.aspx?Item=N82E16814150724&cm_re=AMD_Radeon_R9_270-_-14-150-724-_-Product)
+After **Instanced Stereo** rendering has been enabled, it will require an engine restart, and shader recompile.
+ The base pass, and early-z pass, works with static meshes, skeletal meshes, sprite particles, and mesh particles with the feature turned on (both serial and parallel rendering paths). 
+ **Instanced Stereo** currently works on PC (DirectX) and PS4 with support for other platforms coming after the initial 4.11 release. 
+ Here is a video showing Standard Stereo Rendering in VR, versus Instanced Stereo Rendering in VR. 
 
-**Recommended GPU**
-* [NVIDIA EVGA GeForce GTX 980 w 4GB of VRam](http://www.newegg.com/Product/Product.aspx?Item=N82E16814487161&cm_re=NVIDIA_GTX_980-_-14-487-161-_-Product)
-* [AMD Radeon R9 290X w 4GB of VRam](http://www.newegg.com/Product/Product.aspx?Item=N82E16814150696&cm_re=AMD_Radeon_R9_290x-_-14-150-696-_-Product)
+[OBJECT:EmbeddedVideo]
+	[PARAMLITERAL:videoid]
+	nmajnglJeyE
+	[/PARAMLITERAL]
+	[PARAMLITERAL:width]
+	100
+	[/PARAMLITERAL]
+	[PARAMLITERAL:height]
+	56.25
+	[/PARAMLITERAL]
+	[PARAMLITERAL:units]
+	%
+	[/PARAMLITERAL]
+	[PARAMLITERAL:type]
+	simple
+	[/PARAMLITERAL]
+[/OBJECT]
 
-** Preferred GPU**
-* [NVIDIA GTX Titian w 12GB of VRam](http://www.newegg.com/Product/Product.aspx?Item=N82E16814487129)
-* [AMD Radeon R9 390 w 8GB of VRam](http://www.newegg.com/Product/Product.aspx?Item=N82E16814202164&cm_re=AMD_Radeon_R9_390X-_-14-202-164-_-Product)
+## VR Multiview Support for Mobile
 
-### Recomened CPU
-* **Intel:** [Intel Core i5-4590 Processor](http://www.amazon.com/Intel-i5-4590-BX80646I54590-Processor-Cache/dp/B00J2LIFBY)
-* **AMD:** [AMD FX-9590 Vishera Processor](http://www.newegg.com/Product/Product.aspx?Item=N82E16819113347)
+Certain (supported) mobile devices can now support VR Multiview. Mobile Multiview is similar to Instanced Stereo rendering that is available to desktop PC's, and works by providing an optimized path for Stereo Rendering on the mobile device's CPU. To enable this feature, you will need to do the following:
 
-### Recomened RAM
-* **8 GB RAM:** [G.SKILL Ripjaws Series](http://www.newegg.com/Product/Product.aspx?Item=N82E16820231314)
-
-### Recomened Hard Drive
-* **SSD Harddrive:** [Mushkin 2.5" Internal 256GB SSD](http://www.newegg.com/Product/Product.aspx?Item=N82E16820226688)
-* ** Standard Harddrive:** [WD 3.5" Internal 7200 RPM 1 TB](http://www.newegg.com/Product/Product.aspx?Item=N82E16822236339&cm_re=1_TB_7200_RMP_drive-_-22-236-339-_-Product)
+[region:warning]
+Due to the experimental nature of this feature, it will only work on modern Mali-based GPUs. As this feature matures, and support for more mobile devices come online, this warning will be removed.
+[/region]
 
-### Recomened Motherboard
-* **Intel Based Motherboard:** [MSI Z97 PC Intel](http://www.newegg.com/Product/Product.aspx?Item=N82E16813130779&cm_re=Z97_motherboards-_-13-130-779-_-Product)
-* **AMD Based Motherboard:** [MSI 970A-G46 AMD](http://www.newegg.com/Product/Product.aspx?Item=N82E16813130637&cm_re=socket_am3-_-13-130-637-_-Product)
+1. From the Main Toolbar, go to **Edit** >  **Projects Settings**, to open the Editor's Project Settings.
 
-[region:note]
-If you are picking a different motherboard than the ones listed above, make sure that it has and can use USB 3.0 as this is a requirement of all HMD's.
-[/region]
+1. Then, go to **Engine** > **Rendering** > **VR**, and look for the **Mobile Multi-View (Experimental)** option. 
 
-## Oculus VR Performance
+1. Enable the Mobile Multi-View (Experimental) option by click on the checkmark box next to its name, and then restart Editor for the changes to take effect.
 
+	[REGION:lightbox]
+	[![](Platforms\VR\ContentSetup\VR_MobileMultiView.png)(w:379)](Platforms\VR\ContentSetup\VR_MobileMultiView.png)
+	[/REGION]
 
-## Vive VR Performance
+	[REGION:caption]
+	Click for full image.
+	[/REGION]
 
+For this feature to work, you will need your Android build settings to be setup like the following:
 
+* Make sure that **Support OpenGL ES2** is enabled in the **Android** > **Build** section. 
+* Make sure that both **Mobile HDR** and **Instanced Stereo Rendering** are disabled.
+* Currently, Mobile Multiview will only work with modern Mali-based GPU's. 
+* If you package a project with this feature on, but don't have a compatible GPU, it will be disabled at runtime. 
 
+## VR Forward Rendering
 
+![](Engine\Performance\ForwardRenderer\ForwardRenderingBanner.png)(w:840)
 
+[INCLUDE:Engine\Performance\ForwardRenderer#intro]
 
+[REGION:note]
+Please see the [](Engine\Performance\ForwardRenderer) documentation for more information on this feature. 
+[/REGION]
 
+## VR Performance Profiling
 
+To track down assets that are too expensive for VR, you need to profile what your project is doing on the CPU and GPU as much as possible over the lifetime of your project. 
 
+* **GPU Profiling** - To activate the GPU profiler, all you need to do is press  **CTRL + SHIFT + ,(Comma Key)** together, while running your project. After the keys have been pressed, you will see a new window come up that should look similar to the image below. 
+	![](Platforms\VR\ContentSetup\VR_GPU_Profiler.png)
 
+* **CPU Profiler** - Profiling what your project is doing on the CPU is a bit more involved than profiling for the GPU. To find out more on how to do this, please refer to the [Performance Profiler](https://docs.unrealengine.com/latest/INT/Engine/Performance/Profiler/index.html) documents. 
 
+## VR Post Process Settings
 
+Due to the demanding rendering requirements of VR, many of the advanced Post Processing features that are enabled by default should be disabled, or your project could suffer
+ from performance problems. To accomplish this in your project, you will need to do the following.
 
+ 1. Add a Post Process(PP) volume to your level if there is not one already there.
+ 1. Select the PP volume, and in the **Post Process Volume** section, enable the **Unbound** option so that the settings in the PP volume will be applied to the entire level. 
 
+	![](Platforms\VR\ContentSetup\VR_PP_Unbound.png)
 
+ 1. Expand the **Settings** of the **Post Process Volume**, and then go through each section to disable any active PP settings by enabling that property; clicking on it, and then setting the value from the default (usually 1.0) to **0**, thus disabling the feature.
+	
+	![](Platforms\VR\ContentSetup\VR_Disable_PP_Settings.png)
+	[region:note]
+	When doing this, you won't need to hit every section and set all of the properties to 0. Instead, first disable the really heavy hitting features like Lens Flares, Screen Space reflections, Temporal AA, SSAO, Bloom and anything else that might have an impact on performance.
+	[/region]
 
 
 
diff --git a/Engine/Documentation/Source/Platforms/VR/VRSpectatorScreen/SetSpectatorScreenMode/SetSpectatorScreenMode.INT.udn b/Engine/Documentation/Source/Platforms/VR/VRSpectatorScreen/SetSpectatorScreenMode/SetSpectatorScreenMode.INT.udn
new file mode 100644
--- /dev/null
+++ b/Engine/Documentation/Source/Platforms/VR/VRSpectatorScreen/SetSpectatorScreenMode/SetSpectatorScreenMode.INT.udn
@@ -0,0 +1,12 @@
+Availability: Public 
+Crumbs: %ROOT%
+Title: Set Spectator Screen Mode
+Description: This node sets the Social (Spectator) Screen mode.
+Type: Reference
+Version: 4.17
+Parent: Platforms/VR/VRSpectatorScreen
+Order: 
+Tags: VR 
+
+[TOC(start:2 end:2)]
+[INCLUDE:Platforms/PS4/PSVR/SpectatorScreen/SetSpectatorScreenMode#SetSpectatorScreenModeNode_Ref]
\ No newline at end of file
diff --git a/Engine/Documentation/Source/Platforms/VR/VRSpectatorScreen/SetSpectatorScreenModeTexturePlusEyeLayout/SetSpectatorScreenModeTexturePlusEyeLayout.INT.udn b/Engine/Documentation/Source/Platforms/VR/VRSpectatorScreen/SetSpectatorScreenModeTexturePlusEyeLayout/SetSpectatorScreenModeTexturePlusEyeLayout.INT.udn
new file mode 100644
--- /dev/null
+++ b/Engine/Documentation/Source/Platforms/VR/VRSpectatorScreen/SetSpectatorScreenModeTexturePlusEyeLayout/SetSpectatorScreenModeTexturePlusEyeLayout.INT.udn
@@ -0,0 +1,12 @@
+Availability: Public 
+Crumbs: %ROOT%
+Title: Set Spectator Screen Mode Texture Plus Eye Layout
+Description: This node sets up the layout for the TexturePlusEye function in ESpectatorScreenMode.
+Type: Reference
+Version: 4.17
+Parent: Platforms/VR/VRSpectatorScreen
+Order: 
+Tags: VR 
+
+[TOC(start:2 end:2)]
+[INCLUDE:Platforms/PS4/PSVR/SpectatorScreen/SetSpectatorScreenModeTexturePlusEyeLayout#SetSpectatorScreenModeTexturePlusEyeLayoutNode_Ref]
diff --git a/Engine/Documentation/Source/Platforms/VR/VRSpectatorScreen/SetSpectatorScreenTexture/SetSpectatorScreenTexture.INT.udn b/Engine/Documentation/Source/Platforms/VR/VRSpectatorScreen/SetSpectatorScreenTexture/SetSpectatorScreenTexture.INT.udn
new file mode 100644
--- /dev/null
+++ b/Engine/Documentation/Source/Platforms/VR/VRSpectatorScreen/SetSpectatorScreenTexture/SetSpectatorScreenTexture.INT.udn
@@ -0,0 +1,12 @@
+Availability: Public 
+Crumbs: %ROOT%
+Title: Set Spectator Screen Texture 
+Description: This node changes the texture being displayed on the Social (Spectator) Screen.
+Type: Reference
+Version: 4.17
+Parent: Platforms/VR/VRSpectatorScreen
+Order:
+Tags: VR 
+
+[TOC(start:2 end:2)]
+[INCLUDE:Platforms/PS4/PSVR/SpectatorScreen/SetSpectatorScreenTexture#SetSpectatorScreenTextureNode_Ref]
\ No newline at end of file
diff --git a/Engine/Documentation/Source/Platforms/VR/VRSpectatorScreen/VRSpectatorScreen.INT.udn b/Engine/Documentation/Source/Platforms/VR/VRSpectatorScreen/VRSpectatorScreen.INT.udn
new file mode 100644
--- /dev/null
+++ b/Engine/Documentation/Source/Platforms/VR/VRSpectatorScreen/VRSpectatorScreen.INT.udn
@@ -0,0 +1,31 @@
+Availability: Public
+Crumbs: %ROOT%
+Title: Virtual Reality Spectator Screen
+Description: This document provides an Overview of the Virtual Reality Spectator Screen in Unreal Engine 4. 
+Type: Overview
+SkillLevel: Advanced
+Version: 4.17
+Parent: Platforms/VR 
+Order: 
+Tags: VR
+topic-image:VRSpecScreen_topic_image.png
+social-image:VRSpecScreen_topic_image.png
+Contributors: Ghost Paint
+
+[TOC(start:2 end:2)]
+
+[REGION:tip]
+Please refer to our PlayStation VR (PS VR) Spectator Screen documentation if you're developing for PlayStation.
+[/REGION]
+
+[INCLUDE:Platforms/PS4/PSVR/SpectatorScreen#SpectatorScreen_Overview_1]
+##Spectator Screen Blueprint Nodes
+
+For the Example Usage, we're using the following Blueprint interface:
+
+| Node                                                                                  | Description                                                                                 |
+| ------------------------------------------------------------------------------------- | ------------------------------------------------------------------------------------------- |
+| **[](Platforms/VR/VRSpectatorScreen/SetSpectatorScreenTexture)**                  | %Platforms/VR/VRSpectatorScreen/SetSpectatorScreenTexture:description%                  |
+| **[](Platforms/VR/VRSpectatorScreen/SetSpectatorScreenModeTexturePlusEyeLayout)** | %Platforms/VR/VRSpectatorScreen/SetSpectatorScreenModeTexturePlusEyeLayout:description% |
+| **[](Platforms/VR/VRSpectatorScreen/SetSpectatorScreenMode)**                     | %Platforms/VR/VRSpectatorScreen/SetSpectatorScreenMode:description%                    |
+[INCLUDE:Platforms/PS4/PSVR/SpectatorScreen#SpectatorScreen_Overview_2]
\ No newline at end of file
diff --git a/Engine/Documentation/Source/Platforms/iOS/GettingStarted/iOSGettingStarted.INT.udn b/Engine/Documentation/Source/Platforms/iOS/GettingStarted/iOSGettingStarted.INT.udn
--- a/Engine/Documentation/Source/Platforms/iOS/GettingStarted/iOSGettingStarted.INT.udn
+++ b/Engine/Documentation/Source/Platforms/iOS/GettingStarted/iOSGettingStarted.INT.udn
@@ -1,4 +1,4 @@
-Availability: Docs
+Availability: Public
 Title: Advanced iOS PC/Mac Workflow
 Crumbs:%ROOT%, Platforms, Platforms/iOS
 Description:Setting up for development for the iOS platform.
@@ -8,6 +8,7 @@ SkillLevel: Intermediate
 parent:Platforms/iOS
 order:0
 type:overview
+related:Platforms/iOS/Windows
 
 [TOC(start:2 end:2)]
 
@@ -77,8 +78,8 @@ XCode is [available as a free download](https://itunes.apple.com/us/app/xcode/id
 
 ## Getting Started as an Apple iOS Developer
 
-You will need to sign up for the [Apple IOS Developer Program](https://developer.apple.com/programs/ios/).
-* Individual - If you are working on your own, or will be the only one on a small them that needs Admin access to the IOS account
+You will need to sign up for the [Apple iOS Developer Program](https://developer.apple.com/programs/ios/).
+* Individual - If you are working on your own, or will be the only one on a small them that needs Admin access to the iOS account
 * Company - If you are a small team with Company / This allows you to add people to your team, and grant them different kinds of access
 
 When you create your Apple developer account, you will have to do the following (using Apple's instructions):
@@ -228,15 +229,15 @@ Once you have completed the above steps for setting up your Windows and Mac syst
 * Open the UE4.sln it creates.
 * Open the Games solution folder.
 * Select the Development Editor configuration, and Win64 platform.
-* **Right-click** on the game of choice (QAGame, for instance) and select Build.
-* Now, select the Development configuration, and IOS platform.
+* Right-click on the game of choice (QAGame, for instance) and select Build.
+* Now, select the Development configuration, and iOS platform.
 * Again, build your game.
     * This will connect to your Mac running UnrealRemoteTool, copy files over, and then run remote compilation commands on the Mac.
 	* It will also perform some signing at the end, so make sure that is working correctly. If not, double check the provision file and the certificate installed on the Mac.
 * Start the cooker.
     * In the UE4/Engine/Binaries/Win64 directory (in cmd.exe), run this (replacing QAGame with whatever your game is):
 	    * ue4editor-cmd QAGame -run=cook -targetplatform=ios -cookonthefly
-	* This will start up a cooking server that the IOS device will talk to over the network to retrieve cooked files.
+	* This will start up a cooking server that the iOS device will talk to over the network to retrieve cooked files.
 
 You can look on your Mac drive in the /UE4/Builds directory will have a directory named with your PC's computer name. Inside of there is a mirrored directory structure that matches your PC directory structure. When the compile finishes, the executable will be in the game's Binaries directory (or Engine/Binaries if you compiled the game-less UE4 project)
 
@@ -264,10 +265,10 @@ You can compile for IOS on a Mac. This is pretty new, so is more likely to have
 * In Xcode, open UE4/UE4.xcodeproj
 * Select the QAGame - Mac. 
 * In the scheme (Option-Command-R) select Development, then click done. Command-B to build the Mac executable.
-* Run the IOS cooker from Terminal.
+* Run the iOS cooker from Terminal.
     * Launch Terminal.app and cd to the UE4/Engine/Binaries/Mac directory, then run this (replacing QAGame with whatever your game is):
 	    * open UE4Editor.app --args QAGame -run=cook -targetplatform=ios -cookonthefly
-* Select <GameName> - IOS (Run) target, and select a device to run it on.
+* Select **[GameName] - IOS (Run)** target, and select a device to run it on.
 * In the scheme (Option-Command-R), select Development in Info, and then go to Arguments.
     * Add any necessary commandline options for your game, and also
 	    * -filehostip=xx.xx.xx.xx
diff --git a/Engine/Documentation/Source/Platforms/iOS/QuickStart/2/iOSQuickStart_2.INT.udn b/Engine/Documentation/Source/Platforms/iOS/QuickStart/2/iOSQuickStart_2.INT.udn
--- a/Engine/Documentation/Source/Platforms/iOS/QuickStart/2/iOSQuickStart_2.INT.udn
+++ b/Engine/Documentation/Source/Platforms/iOS/QuickStart/2/iOSQuickStart_2.INT.udn
@@ -70,8 +70,8 @@ ultimately launch or install on your iOS device.
 
 [REGION:warning]
 The PC-only workflow described in this guide requires a pure Blueprint project. If you wish to use code in your 
-project, you will also need a Mac in order to compile the code. You can follow the workflow outlined on the 
-[](Platforms/iOS/GettingStarted) page to develop code-based projects for iOS on a PC.
+project, you will also need a Mac in order to compile the code. You can follow the [](Platforms/iOS/Windows) 
+page to develop code-based projects for iOS on a PC.
 [/REGION]
 [/EXCERPT:windows]
 [EXCERPT:macos]
diff --git a/Engine/Documentation/Source/Platforms/iOS/Windows/Windows.INT.udn b/Engine/Documentation/Source/Platforms/iOS/Windows/Windows.INT.udn
--- a/Engine/Documentation/Source/Platforms/iOS/Windows/Windows.INT.udn
+++ b/Engine/Documentation/Source/Platforms/iOS/Windows/Windows.INT.udn
@@ -10,6 +10,7 @@ parent:Platforms/iOS
 tags:Mobile
 tags:iOS
 prereq:Platforms/iOS/QuickStart
+related:Platforms/iOS/GettingStarted
 
 While building for iOS requires a Mac, using remote compilation means that a whole team that primarily uses Windows machines can build for iOS using one Mac. 
 
diff --git a/Engine/Documentation/Source/Programming/Assets/ReferencingAssets/ReferencingAssets.INT.udn b/Engine/Documentation/Source/Programming/Assets/ReferencingAssets/ReferencingAssets.INT.udn
--- a/Engine/Documentation/Source/Programming/Assets/ReferencingAssets/ReferencingAssets.INT.udn
+++ b/Engine/Documentation/Source/Programming/Assets/ReferencingAssets/ReferencingAssets.INT.udn
@@ -47,7 +47,7 @@ The declaration of the UPROPERTY appears the same as the previous hard reference
 
 ## Indirect Property Reference
 
-One easy way to control when an asset is loaded is to use a TSoftObjectPtr. To the designer, they can work with it just like the direct property reference. However, instead of a direct pointer reference, the property is stored as a string with template code to enable safe checking of whether the asset has been loaded yet or not. Use the IsPending() method to check if the asset is ready to be accessed or not. Note that using TSoftObjectPtr requires you to manually load the asset when you want to use it. You can use the templated LoadObject<>() method, StaticLoadObject(), or the FStreamingManager to load your object (see [](Programming/Assets/AsyncLoading) for more information). The first two methods load the asset synchronously which can cause frame rate spikes, so should only be used if you know it won't affect the gameplay.
+One easy way to control when an asset is loaded is to use a TSoftObjectPtr. To the designer, they can work with it just like the direct property reference. However, instead of a direct pointer reference, the property is stored as a string with template code to enable safe checking of whether the asset has been loaded yet or not. Use the IsPending() method to check if the asset is ready to be accessed or not. Note that using TSoftObjectPtr requires you to manually load the asset when you want to use it. You can use the templated `LoadObject<>()` method, StaticLoadObject(), or the FStreamingManager to load your object (see [](Programming/Assets/AsyncLoading) for more information). The first two methods load the asset synchronously which can cause frame rate spikes, so should only be used if you know it won't affect the gameplay.
 
 	UPROPERTY(EditDefaultsOnly, BlueprintReadWrite, Category=Building)
 	TSoftObjectPtr<UStaticMesh> BaseMesh;
@@ -68,7 +68,7 @@ If you want to defer loading a UClass, you use the same approach as the TSoftObj
 
 ## Find/Load Object
 
-So far these examples have all been UPROPERTY based. However, what if you want to build a string at runtime and use that to get a reference to an object? There are two options that you can use. If you want to use the UObject only if it has already been loaded or created, then the right choice is FindObject<>(). If you want to load the object if it is not already loaded, then using LoadObject<>() is the right choice. Note that LoadObject<>() does the equivalent of FindObject under the covers, so it's not necessary to try to find the object first and then load it. Here are some examples of using each function.
+So far these examples have all been UPROPERTY based. However, what if you want to build a string at runtime and use that to get a reference to an object? There are two options that you can use. If you want to use the UObject only if it has already been loaded or created, then the right choice is `FindObject<>()`. If you want to load the object if it is not already loaded, then using `LoadObject<>()` is the right choice. Note that `LoadObject<>()` does the equivalent of FindObject under the covers, so it's not necessary to try to find the object first and then load it. Here are some examples of using each function.
 
 	AFunctionalTest* TestToRun = FindObject<AFunctionalTest>(TestsOuter, *TestName);
 	GridTexture = LoadObject<UTexture2D>(NULL, TEXT("/Engine/EngineMaterials/DefaultWhiteGrid.DefaultWhiteGrid"), NULL, LOAD_None, NULL);
diff --git a/Engine/Documentation/Source/Programming/Development/BuildGraph/ScriptAnatomy/Conditions/BuildGraphScriptConditions.INT.udn b/Engine/Documentation/Source/Programming/Development/BuildGraph/ScriptAnatomy/Conditions/BuildGraphScriptConditions.INT.udn
--- a/Engine/Documentation/Source/Programming/Development/BuildGraph/ScriptAnatomy/Conditions/BuildGraphScriptConditions.INT.udn
+++ b/Engine/Documentation/Source/Programming/Development/BuildGraph/ScriptAnatomy/Conditions/BuildGraphScriptConditions.INT.udn
@@ -53,18 +53,18 @@ A list of operators is specified below:
 
 | **Operator**            | **Description**                                                             | **Precedence**  |
 | ----------------------- | --------------------------------------------------------------------------- | --------------- |
-| **(x)**                 | Subexpression                                                               | 1               |
-| **!x**                  | Not operator                                                                | 1               |
-| **Exists(x)**           | True if the x file exists.                                                  | 1               |
-| **HasTrailingSlash(x)** | True if x ends with a slash or backslash.                                   | 1               |
-| **x == y**              | Tests two atoms for equality (case insensitive).                            | 2               |
-| **x != y**              | Tests two atoms for inequality (case insensitive).                          | 2               |
-| **x < y**               | Compares whether the integer x is less than the integer y.                  | 2               |
-| **x <= y**              | Compares whether the integer x is less than or equal to the integer y.      | 2               |
-| **x > y**               | Compares whether the integer x is greater than the integer y.               | 2               |
-| **x >= y**              | Compares whether the integer x is greater than or equal to the integer y.   | 2               |
-| **x and y**             | True if both arguments are `true`.                                          | 3               |
-| **x or y**              | True if either argument is `true`.                                          | 4               |
+| `(x)`                 | Subexpression                                                               | 1               |
+| `!x`                  | Not operator                                                                | 1               |
+| `Exists(x)`           | True if the x file exists.                                                  | 1               |
+| `HasTrailingSlash(x)` | True if x ends with a slash or backslash.                                   | 1               |
+| `x == y`              | Tests two atoms for equality (case insensitive).                            | 2               |
+| `x != y`              | Tests two atoms for inequality (case insensitive).                          | 2               |
+| `x < y`               | Compares whether the integer x is less than the integer y.                  | 2               |
+| `x <= y`              | Compares whether the integer x is less than or equal to the integer y.      | 2               |
+| `x > y`               | Compares whether the integer x is greater than the integer y.               | 2               |
+| `x >= y`              | Compares whether the integer x is greater than or equal to the integer y.   | 2               |
+| `x and y`             | True if both arguments are `true`.                                          | 3               |
+| `x or y`              | True if either argument is `true`.                                          | 4               |
 
 [REGION:note]
 The `'<'` and `'>'` characters must be escaped as `"&lt;"` and `"&gt;"` in XML.
diff --git a/Engine/Documentation/Source/Programming/Development/CodingStandard/CodingStandard.INT.udn b/Engine/Documentation/Source/Programming/Development/CodingStandard/CodingStandard.INT.udn
--- a/Engine/Documentation/Source/Programming/Development/CodingStandard/CodingStandard.INT.udn
+++ b/Engine/Documentation/Source/Programming/Development/CodingStandard/CodingStandard.INT.udn
@@ -489,7 +489,7 @@ Enum classes should always be used as a replacement for old-style namespaced enu
 			Thing2
 		};
 
-These are also supported as UPROPERTYs, as long they are based on uint8 - this replaces the old TEnumAsByte<> workaround:
+These are also supported as UPROPERTYs, as long they are based on uint8 - this replaces the old `TEnumAsByte<>` workaround:
 
 		// Old property
 		UPROPERTY()
diff --git a/Engine/Documentation/Source/Programming/Introduction/IntroToProgramming.INT.udn b/Engine/Documentation/Source/Programming/Introduction/IntroToProgramming.INT.udn
--- a/Engine/Documentation/Source/Programming/Introduction/IntroToProgramming.INT.udn
+++ b/Engine/Documentation/Source/Programming/Introduction/IntroToProgramming.INT.udn
@@ -243,7 +243,7 @@ This function is called like any other C++ function. Under the covers, the Unrea
 	UFUNCTION(BlueprintNativeEvent, Category="Damage")
 	void CalledFromCpp();
 
-This version still generates the thunking method to call into the Blueprint VM. So how do you provide the default implementation? The tools also generate a new function declaration that looks like <function name>_Implementation(). You must provide this version of the function or your project will fail to link. Here is the implementation code for the declaration above.
+This version still generates the thunking method to call into the Blueprint VM. So how do you provide the default implementation? The tools also generate a new function declaration that looks like `<function name>_Implementation()`. You must provide this version of the function or your project will fail to link. Here is the implementation code for the declaration above.
 
 	void AMyActor::CalledFromCpp_Implementation()
 	{
@@ -533,7 +533,7 @@ Since different platforms have different sizes for basic types such as **short**
 
 Floating point numbers are also supported with the standard **float **(32-bit)** **and **double** (64-bit) types.
 
-Unreal Engine has a template, **TNumericLimits<t>**, for finding the minimum and maximum ranges value types can hold. For more information follow this [link](https://docs.unrealengine.com/latest/INT/API/Runtime/Core/Math/TNumericLimits/index.html).
+Unreal Engine has a template, `TNumericLimits<t>`, for finding the minimum and maximum ranges value types can hold. For more information follow this [link](https://docs.unrealengine.com/latest/INT/API/Runtime/Core/Math/TNumericLimits/index.html).
 
 ### Strings
 
@@ -593,7 +593,7 @@ The **FChar** type provides a set of static utility functions for working with i
 	TCHAR Lower = FChar::ToLower(Upper); // 'a'
 
 [REGION:note]
-The FChar type is defined as **TChar<TCHAR>** (as it is listed in the API).
+The FChar type is defined as `TChar<TCHAR>` (as it is listed in the API).
 [/REGION]
 
 [Full Topic: TChar API](https://docs.unrealengine.com/latest/INT/API/Runtime/Core/Misc/TChar/index.html)
diff --git a/Engine/Documentation/Source/Programming/Online/OnlineSubsystem.INT.udn b/Engine/Documentation/Source/Programming/Online/OnlineSubsystem.INT.udn
--- a/Engine/Documentation/Source/Programming/Online/OnlineSubsystem.INT.udn
+++ b/Engine/Documentation/Source/Programming/Online/OnlineSubsystem.INT.udn
@@ -74,7 +74,7 @@ Additional services are loaded on demand when called from this function with the
 
 Much like Unreal Engine 3 (UE3), the online subsystem will make heavy use of delegates when calling functions with asynchronous side effects.   It is important to respect the delegates and wait for the proper delegate to be called before calling functions further down a chain.  Failure to wait for an asynchronous task can cause crashes and unexpected, undefined behavior.   Waiting for delegates is especially important during connectivity failures such as a cable pull or other disconnect events.   The amount of time a task may take to finish may seem instantaneous in the ideal case, but can be upwards of almost a minute in the timeout case.
 
-The delegate interface is fairly straightforward, with each delegate clearly defined at the top of each interface header.   Every delegate has an Add<DelegateName>, Clear<DelegateName>, and Trigger<DelegateName> function (although Triggering delegates manually is discouraged).   Common practice is to Add() the delegate right before calling the appropriate function, and then Clear() the delegate from within itself.   
+The delegate interface is fairly straightforward, with each delegate clearly defined at the top of each interface header.   Every delegate has an `Add<DelegateName>`, `Clear<DelegateName>`, and `Trigger<DelegateName>` function (although Triggering delegates manually is discouraged).   Common practice is to Add() the delegate right before calling the appropriate function, and then Clear() the delegate from within itself.   
 
 
 ### Interfaces
diff --git a/Engine/Documentation/Source/Programming/Online/Steam/OnlineSubsystemSteam.INT.udn b/Engine/Documentation/Source/Programming/Online/Steam/OnlineSubsystemSteam.INT.udn
--- a/Engine/Documentation/Source/Programming/Online/Steam/OnlineSubsystemSteam.INT.udn
+++ b/Engine/Documentation/Source/Programming/Online/Steam/OnlineSubsystemSteam.INT.udn
@@ -1,22 +1,27 @@
 Availability:Public
 Title: Online Subsystem Steam
 Crumbs:%ROOT%, Programming
-Description:Information specific to the Steam Online Subsystem 
-Version: 4.9
-
-[TOC(start:2)]
-
+Description:An overview of Online Subsystem Steam, including how to set up your project for distribution on Valve's Steam platform.
+Version: 4.17
+Type: Overview
+SkillLevel:Advanced
+Parent: Programming/Online
+Order:
+Tags:Steam
+Tags: Online Subsystem
+Contributors: Valve
+
+[TOC(start:2 end:2)]
 
 [EXCERPT:Intro] 
-Information contained here pertains specifically to the Steam implementation.   Any additional setup steps, tips & tricks, and workarounds can be found here.
+The **Online Subsystem Steam API** enables you to ship Unreal Engine 4 (UE4) applications to Valve's Steam platform. 
+The main purpose of the **Steam** module is to help you distribute your application with a set of features (such as matchmaking and leaderboards) to Steam users.
+Additionally, the Steam module implements several of the interfaces being exposed by the [Online Subsystem](Programming/Online), supporting most of what is offered by the Steamworks Software Development Kit (SDK). 
 [/EXCERPT:Intro]
 
+Some of the available Steam Interfaces include: 
 
-
-The **Steam** module implements most of the interfaces exposed through the online subsystem and supports most of what Valve offers through the Steamworks SDK. 
-
-Steam Interfaces
-* Matchmaking (Lobbies / GameServer APIs)
+* Matchmaking (Lobbies and GameServer APIs)
 * Leaderboards
 * Achievements
 * Voice
@@ -24,99 +29,196 @@ Steam Interfaces
 * SharedCloud
 * External UI
 
+Reference our [Online Subsystem Steam API Reference](https://docs.unrealengine.com/latest/INT/API/Plugins/OnlineSubsystemSteam/FOnlineSubsystemSteam/index.html) for a more complete listing of currently available Steam Interfaces.
+
+## Meeting Valve's Requirements
+
+The Steam Subsystem requires additional setup through [Valve Steamworks](https://partner.steamgames.com/?goto=%2Fhome). 
+Please contact [Valve](http://www.valvesoftware.com/contact/) and reference their [Steamworks SDK Documentation](https://partner.steamgames.com/doc/home) to make sure that your application meets Valve's requirements before attempting to use Steam with UE4.
+
+## Downloading Steamworks
+
+If your application meets Valve's requirements, go ahead and download the latest version of the [Steamworks SDK](https://partner.steamgames.com/). The SDK needs to be unzipped and copied to `/YourUnrealEnginePath/Engine/Source/ThirdParty/Steamworks/Steam[Current Version]/sdk`
+
+[REGION:note]
+If you're updating your project's Steamworks SDK, make sure to update the `[Current Version]` in your project's Steamworks path:  
+`/YourUnrealEnginePath/Engine/Source/ThirdParty/Steamworks/`**Steam[Current Version]**`/sdk`  
+You'll also need to update `Steamworks.build.cs`, which is located in the project directory.
+![](SteamworksCSFile.png)  
+In this example, the `[Current Version]` is formatted as `v139`.
+[/REGION]
 
+###Compiling UE4
 
-All games must have a valid Steam App ID.   During development, this App ID is exposed to the SDK via a file called steam_appid.txt which must reside in the same directory as the executable.  The file is generated by Unreal Engine at launch and deleted during graceful shutdown of the engine.   This makes it unnecessary to launch the game via the Steam client (although it must be running).  The file should not be included in any Steam images.
+If you're compiling UE4 from source, the following steps should be taken:
 
-In shipping builds, the engine will also check to make sure the logged in user is properly subscribed to the game and will shutdown the engine if not.  This is but one way to help secure the game.  Using Steam DRM (see the Steamworks SDK) should further protect the game from tampering.
+1.	Open `OnlineSubsystemSteamPrivatePCH.h` (being located in `..\..\Plugins\Online\OnlineSubsystemSteam\Source\Private\`), where you'll find the following line of code, defining the root location of the Steamworks SDK:
 
-## Basic Setup
+		#define STEAM_SDK_ROOT_PATH TEXT("Binaries/ThirdParty/Steamworks")
 
-The Steam subsystem requires some additional setup through Valve.   Contact Valve and use their documentation to make sure you are setup on their end before attempting to use Steam in the Unreal Engine.
+1.	You'll want to modify the line to read:
 
-### Installing the Steamworks SDK
-For legal reasons, using Steam requires downloading the latest SDK from Valve.   Currently this is *v1.29a*, but any future update should be a straightforward change to the path name.  See *Steamworks.build.cs* in the ThirdParty/Steamworks directory
+		#define STEAM_SDK_VER TEXT("Steam[Current Version]")
 
-Using Steam against the precompiled version of the engine should only require copying some of the DLLs from Valve's SDK into the appropriate places.  If you intend to recompile the engine against the source, putting the SDK in the right place is required as well.
+## Setting up the Steamworks SDK
 
-* The SDK needs to be unzipped/copied to this path
-	* /YourUnrealEnginePath/Engine/Source/ThirdParty/Steamworks/Steamv130/sdk
-* Find the following binaries from the /redistributable_bin/ directory of the SDK and copy them to their noted locations.  <br>*Note:* Some of the 64bit dlls can be found in your normal Steam client directory. Valve for some reason does not include all of them in the SDK.
-	* /YourUnrealEnginePath/Engine/Binaries/ThirdParty/Steamworks/Steamv130/Win64
+Using Steam against the precompiled version of the engine should only require copying some of the dynamically linked libraries (*.dll) from Valve's SDK into the appropriate folders. 
+If you intend to recompile the engine against the source, putting the SDK in the right place is required as well. Now, copy the relevant redistributable files from the `/redistributable_bin/` directory of the SDK to the following locations:  
+%Globals:osselect_linux%
+[OBJECT:ToggleButtonContent]
+	[PARAMLITERAL:category]
+	OS
+	[/PARAMLITERAL]
+	[PARAMLITERAL:id]
+	windows
+	[/PARAMLITERAL]
+	[PARAMLITERAL:active]
+	active_button_content
+	[/PARAMLITERAL]
+	[PARAM:content]
+		[REGION:note]
+			Some of the 64 bit libraries can be found in your Steam client directory (at the time of this writing, Valve didn't include all of the libraries in the SDK).
+		[/REGION]
+		`/YourUnrealEnginePath/Engine/Binaries/ThirdParty/Steamworks/Steam[Current Version]/Win64`
 		* steam_api64.dll
 		* steamclient64.dll
 		* tier0_s64.dll
 		* vstdlib_s64.dll
-	* /YourUnrealEnginePath/Engine/Binaries/ThirdParty/Steamworks/Steamv130/Win32
+
+		`/YourUnrealEnginePath/Engine/Binaries/ThirdParty/Steamworks/Steam[Current Version]/Win32`
 		* steam_api.dll
 		* steamclient.dll
 		* tier0_s.dll
 		* vstdlib_s.dll
-	* /YourUnrealEnginePath/EngineOrGameFolder/Binaries/Mac/YourGame.app/Contents/MacOS
-		* libsteam_api.dylib (from /redistributable_bin/osx32 - single dylib has both 32 and 64 bit support)
-		
-[REGION:note]
-If you are compiling the entire engine, you will want to modify the following line in OnlineSubsystemSteam\Private\OnlineSubsystemSteamPrivatePCH.h 
+	[/PARAM]
+[/OBJECT]
+[OBJECT:ToggleButtonContent]
+	[PARAMLITERAL:category]
+	OS
+	[/PARAMLITERAL]
+	[PARAMLITERAL:id]
+	mac
+	[/PARAMLITERAL]
+	[PARAMLITERAL:active]
+	active_button_content
+	[/PARAMLITERAL]
+	[PARAM:content]
+		`/YourUnrealEnginePath/EngineOrGameFolder/Binaries/Mac/YourGame.app/Contents/MacOS`
+		* libsteam_api.dylib (from `/redistributable_bin/osx32` - the [dynamic library](https://developer.apple.com/library/content/documentation/DeveloperTools/Conceptual/DynamicLibraries/000-Introduction/Introduction.html#//apple_ref/doc/uid/TP40001908-SW1) (*.dylib) has both 32 and 64 bit support)
+	[/PARAM]
+[/OBJECT]
+[OBJECT:ToggleButtonContent]
+	[PARAMLITERAL:category]
+	OS
+	[/PARAMLITERAL]
+	[PARAMLITERAL:id]
+	linux
+	[/PARAMLITERAL]
+	[PARAMLITERAL:active]
+	active_button_content
+	[/PARAMLITERAL]
+	[PARAM:content]
+		[REGION:note]
+			Linux users must link and ship the following file alongside the executable.
+		[/REGION]
+		* libsteam_api.so
+	[/PARAM]
+[/OBJECT]
+
+## Steam App ID
 
-define STEAM_SDK_VER TEXT("Steamv130")
+All games using the Steam Online Subsystem must have a valid application ID because the Steamworks API won't initialize if it doesn't know your application's Steam App ID. 
+Before initializing Steam, UE4 will generate `steam_appid.txt` (during a graceful shutdown of the engine, UE4 deletes this file). 
+It's important to note that `steam_appid.txt` must be located in the same directory as your application's executable file because Steam will look for the text file in the current working directory. 
+Additionally, the file should not be included in any Steam images. 
 
-to make sure it references the new SDK directory location / version
+If you open `steam_appid.txt`, you'll see a **SteamDevAppId** entry, which is a field that hints your application's ID to Steam.
+This makes it unnecessary to launch the game using the Steam client (although it must be running). 
+
+[REGION:tip]
+If you want to test your application, you can use a **SteamDevAppId** of `480`, which is a test App ID that is shared by all developers.
+Although you can test most Steam Interfaces with the aforementioned test App ID, your application will need a Steam App ID before being shipped.
 [/REGION]
 
-### INI Configuration
+###Shipping Builds
+
+In shipping builds, the engine will check to make sure that the logged-in user is properly subscribed to the game and will shutdown if the engine's test returns false (this is one way to help secure the game). Additionally, using Steam DRM (see the Steamworks SDK) should further protect the game from being tampered with.
+
+## Configuring your Application's Settings
 
-Turn on some settings in the game's DefaultEngine.ini.  
 [REGION:note]
-The SteamDevAppId of 480 is Valve's test app id, shared by everyone.   You will need your own app id eventually, but most features of Steam should work before then.
+If you're creating a new project using the **Unreal Project Browser**, **Online Subsystem** settings shouldn't exist in `DefaultEngine.ini`; however, if you're modifying one of our Sample Projects, the **Online Subsystem** settings might already exist.
 [/REGION]
 
-    [/Script/Engine.GameEngine]
-    +NetDriverDefinitions=(DefName="GameNetDriver",DriverClassName="OnlineSubsystemSteam.SteamNetDriver",DriverClassNameFallback="OnlineSubsystemUtils.IpNetDriver")
-    
-    [OnlineSubsystem]
-    DefaultPlatformService=Steam
+Now that you've set up the Steamworks SDK for your application (along with setting up the Steam App ID), you're ready to configure your application's settings to use Online Subsystem Steam. 
 
-    [OnlineSubsystemSteam]
-    bEnabled=true
-    SteamDevAppId=480
+###Steps
 
-    [/Script/OnlineSubsystemSteam.SteamNetDriver]
-    NetConnectionClassName="OnlineSubsystemSteam.SteamNetConnection"
+1.	Getting started, open you application's `DefaultEngine.ini` file.
 
-### Module Setup
+1.	Add the following setting:
 
-Make sure to include the Unreal Engine Steam module as part of your project. (see [](Programming/UnrealBuildSystem/TargetFiles) for additional help)
+	    [/Script/Engine.GameEngine]
+		+NetDriverDefinitions=(DefName="GameNetDriver",DriverClassName="OnlineSubsystemSteam.SteamNetDriver",DriverClassNameFallback="OnlineSubsystemUtils.IpNetDriver")
 
-Adding the following should be enough to make sure that the Steam module is built along with your game.  It goes inside the constructor of mygame.build.cs 
+	**NetDriverDefinitions** describes the net drivers available to UE4 with the following properties being defined:
+	* **DefName** is the unique name of this net driver definition.
+	* **DriverClassName** is the class name of the primary net driver.
+	* **DriverClassNameFallBack** is the class name of the fallback net driver if the primary net driver class fails to initialize.
 
-    DynamicallyLoadedModuleNames.Add("OnlineSubsystemSteam");
+1.	To tell UE4 to use Online Subsystem Steam, add the following setting:
 
-### Steam Overlay on Mac
+		[OnlineSubsystem]
+		DefaultPlatformService=Steam
 
-Contrary to Windows, Steam Overlay on Mac requires game to be launched using Steam client. For this you first need to add the game to your library using "Add a Non-Steam Game to My Library" option from Steam's Games menu.
+1.	Now that you've told UE4 that you want your application to use the Steam Online Subsystem, you'll need to configure the **OnlineSubsystemSteam** module by adding the following settings:
 
-<!--
-### Steam Sockets
+		[OnlineSubsystemSteam]
+		bEnabled=true
+		SteamDevAppId=480
 
-### Steam Net Driver
+1.	Finally, you'll want to specify the Steam class to the Net Driver for the application's connections: 
 
-### Steam Net Connection
+		[/Script/OnlineSubsystemSteam.SteamNetDriver]
+		NetConnectionClassName="OnlineSubsystemSteam.SteamNetConnection"
 
+###End Result
 
+At the end of this brief guide, your application's `DefaultEngine.ini` file should look like the following setting block. If you want to see how other projects set up and use the Online Subsystem, please refer to our library of Sample Projects.
 
-## Reference
+[OBJECT:Section]
+    [PARAMLITERAL:id]
+    code
+    [/PARAMLITERAL]
+    [PARAM:heading]
+    Finished Settings
+    [/PARAM]
+    [PARAM:content]
+		**DefaultEngine.ini**
+			[/Script/Engine.GameEngine]
+			+NetDriverDefinitions=(DefName="GameNetDriver",DriverClassName="OnlineSubsystemSteam.SteamNetDriver",DriverClassNameFallback="OnlineSubsystemUtils.IpNetDriver")
+			
+			[OnlineSubsystem]
+			DefaultPlatformService=Steam
 
-[OBJECT:TopicList]
-	[PARAM:icon]
-	[/PARAM]
-	[PARAM:title]
-	Interfaces
-	[/PARAM]
-	[PARAM:description]
-	[/PARAM]
-	[PARAM:links]
-	* [](Programming/Online/): %Programming/Online/:Description%
-	[/PARAM]
+			[OnlineSubsystemSteam]
+			bEnabled=true
+			SteamDevAppId=480
+
+			[/Script/OnlineSubsystemSteam.SteamNetDriver]
+			NetConnectionClassName="OnlineSubsystemSteam.SteamNetConnection"
+    [/PARAM]
 [/OBJECT]
-   
-  --> 
+
+##Additional Notes
+
+### Module Setup
+
+Make sure to include the Unreal Engine Steam module as part of your project (see [](Programming/UnrealBuildSystem/TargetFiles) for additional help).
+Specifically, adding the following line in the construtor for `mygame.build.cs` should be enough to make sure that the Steam module is built along with your game. 
+
+    DynamicallyLoadedModuleNames.Add("OnlineSubsystemSteam");
+
+### Steam Overlay on Mac
+
+Steam Overlay on Mac requires game to be launched using Steam client. For this, you first need to add the game to your library using "Add a Non-Steam Game to My Library" option from Steam's Games menu.
diff --git a/Engine/Documentation/Source/Programming/Plugins/Plugins.INT.udn b/Engine/Documentation/Source/Programming/Plugins/Plugins.INT.udn
--- a/Engine/Documentation/Source/Programming/Plugins/Plugins.INT.udn
+++ b/Engine/Documentation/Source/Programming/Plugins/Plugins.INT.udn
@@ -2,7 +2,7 @@ Availability:Public
 Title:Plugins
 Crumbs:
 Description:How to create Unreal Engine plugins.
-Version: 4.9
+Version: 4.16
 Parent:Programming/UnrealArchitecture/Reference
 
 [TOC(start:2)]
@@ -192,22 +192,30 @@ Plugin descriptors are in the Json ([JavaScript Object Notation](http://www.json
 
 ### Descriptor File Example
 
-This example plugin descriptor is from the **UObjectPlugin** [downloadable example](#PluginExamples) above.
+This example plugin descriptor is from the **UObjectPlugin** [example](#PluginExamples) above.
 
 	{
 		"FileVersion" : 3,
-
-		"FriendlyName" : "UObject Example Plugin",
 		"Version" : 1,
-		"FriendlyVersion" : "1.0",
-		"Description" : "An example of a plugin which declares its own UObject type. This can be used as a starting point when creating your own plugin.",
-		"Category" : "Programming Examples.Plugins",
-	
+		"VersionName" : "1.0",
+		"FriendlyName" : "UObject Example Plugin",
+		"Description" : "An example of a plugin which declares its own UObject type.  This can be used as a starting point when creating your own plugin.",
+		"Category" : "Examples",
+		"CreatedBy" : "Epic Games, Inc.",
+		"CreatedByURL" : "http://epicgames.com",
+		"DocsURL" : "",
+		"MarketplaceURL" : "",
+		"SupportURL" : "",
+		"EnabledByDefault" : true,
+		"CanContainContent" : false,
+		"IsBetaVersion" : false,
+		"Installed" : false,
 		"Modules" :
 		[
 			{
 				"Name" : "UObjectPlugin",
-				"Type" : "Developer"
+				"Type" : "Developer",
+				"LoadingPhase" : "Default"
 			}
 		]
 	}
@@ -226,8 +234,12 @@ This example plugin descriptor is from the **UObjectPlugin** [downloadable examp
 | **CreatedBy**						| _Optional_	| The individual or company name that created this plugin. This may be displayed in the plugin UI or in other locations. |
 | **CreatedByURL**					| _Optional_	| A web link to the individual or company that created this plugin. If specified, the editor UI may display a hyperlink that allows the user to browse to this web page. |
 | **DocsURL**						| _Optional_	| A web link to the plugin's documentation. If specified, the editor's plugin browser will display a hyperlink that allows the user to browse to this web page. |
+| **MarketplaceURL**				| _Optional_	| A link to the plugin's Marketplace page. If specified, the editor's plugin browser will display a hyperlink that allows the user to browse to this web page. |
 | **SupportURL**					| _Optional_	| Link to the support page for the plugin. Will be shown on the editor's plugin browser if specified. |
-| **CanContainContent**				| _Optional_	| When specified and set to **true**, enables Content support for this plugin. The default setting is **false**. See the section on [Content in Plugins](#ContentInPlugins) for more info. | 
+| **EnabledByDefault**				| _Optional_	| If set to true, this plugin will be enabled by default when it is present. |
+| **CanContainContent**				| _Optional_	| When specified and set to **true**, enables Content support for this plugin. The default setting is **false**. See the section on [Content in Plugins](#ContentInPlugins) for more info. |
+| **IsBetaVersion**					| _Optional_	| Marks the content as a beta version. In the list of Plugins, "Version" will be replaced with "BETA Version" for this entry. |
+| **Installed**						| _Optional_	| If set to `true`, the Plugin will be listed under the "Installed" section in the list of Plugins, instead of the "Project" or "Built-in" sections. |
 | **Modules**						| _Optional_	| For plugins that contain source code (and binaries), this is the list of modules that should be loaded at startup. See below for more info. |
 
 
@@ -238,6 +250,7 @@ For plugins that contain code, the descriptor file will contain at least one mod
 	{
 		"Name" : "UObjectPlugin",
 		"Type" : "Developer"
+		"LoadingPhase" : "Default"
 	}
 
 | Field name				| Info			| Description |
diff --git a/Engine/Documentation/Source/Programming/Rendering/MaterialExpressions/CreatingMaterialExpressions.INT.udn b/Engine/Documentation/Source/Programming/Rendering/MaterialExpressions/CreatingMaterialExpressions.INT.udn
--- a/Engine/Documentation/Source/Programming/Rendering/MaterialExpressions/CreatingMaterialExpressions.INT.udn
+++ b/Engine/Documentation/Source/Programming/Rendering/MaterialExpressions/CreatingMaterialExpressions.INT.udn
@@ -13,7 +13,7 @@ Creating a new material expression in UE4 is straightforward.  As an example, I
 
 ## Declaration
 
-First, add a new engine script file with a new type that derives from MaterialExpression.  Declare the expression's inputs and virtual functions, for example:<br> 
+First, add a new engine script file with a new type that derives from MaterialExpression.  Declare the expression's inputs and virtual functions, for example: 
 
 	MaterialExpressionSquare.h
 
diff --git a/Engine/Documentation/Source/Programming/Rendering/Overview/RenderingOverview.INT.udn b/Engine/Documentation/Source/Programming/Rendering/Overview/RenderingOverview.INT.udn
--- a/Engine/Documentation/Source/Programming/Rendering/Overview/RenderingOverview.INT.udn
+++ b/Engine/Documentation/Source/Programming/Rendering/Overview/RenderingOverview.INT.udn
@@ -206,4 +206,4 @@ Since there are so many rendering states, it is not practical to set them all ev
 * RHISetRasterizerState
 * Any dependencies of the shaders set by RHISetBoundShaderState
 
-All other states are assumed to be at their defaults (as defined by the relevant TStaticState, for example the default stencil state is set by RHISetStencilState(TStaticStencilState<>::GetRHI()).
+All other states are assumed to be at their defaults (as defined by the relevant TStaticState, for example the default stencil state is set by `RHISetStencilState(TStaticStencilState<>::GetRHI())`.
diff --git a/Engine/Documentation/Source/Programming/Rendering/ShaderInPlugin/Overview/OverviewShaderInPlugin.INT.udn b/Engine/Documentation/Source/Programming/Rendering/ShaderInPlugin/Overview/OverviewShaderInPlugin.INT.udn
new file mode 100644
--- /dev/null
+++ b/Engine/Documentation/Source/Programming/Rendering/ShaderInPlugin/Overview/OverviewShaderInPlugin.INT.udn
@@ -0,0 +1,306 @@
+Availability: Docs
+Crumbs: %ROOT%
+Title: Overview Shader In Plugins
+Description: Going over creating shaders in Plugins.
+Type: Overview
+SkillLevel: Advanced
+Version: 4.17
+Parent: Programming/Rendering/ShaderInPlugIn
+Order: 1
+tags:Rendering
+
+[TOC(start:1 end:3)]
+
+[region:warning]
+Please note that this document **is not** a guide on how to write HLSL code or GPU efficient shaders but merely to show you how to create a new shader using the Plugin system.
+[/region]
+
+Adding new shaders for use in Unreal Engine 4 (UE4) can now be achieved via the Plugin system. Creating a shader via the Plugin system allows you to quickly and easily share what you have created with the one you want. In the following document, we will take a high-level look at what needs to done to create shaders in Plugins in UE4.
+
+[region:note]
+For additional help, you can directly look at **//Engine/Plugins/Compositing/LensDistortion's** Plugin.
+[/region]
+
+## Plugin Creation Tips
+
+When creating a new Plugin, you should be aware of the following things:
+
+* Use the Plugin Wizard to quickly create all of the needed files and folders for your Plugin.
+
+* Right now, it is not possible to make drastic changes, like adding a new shader model, to the Material Editor via Plugins.
+
+* Make sure you add all of your files and folder in the required location and then regenerate your Visual Studio solution files.
+
+* In your ProjectName.uplugin, make sure that the module's LoadingPhase to PostConfigInit (only for the modules that will have shader implementations) like in the following example: 
+
+        {
+            "FileVersion" : 3,
+            "Version" : 1,
+            "VersionName" : "1.0",
+            "FriendlyName" : "Foo",
+            "Description" : "Plugin to play around with shaders.",
+            "Category" : "Sandbox",
+            "CreatedBy" : "Epic Games, Inc.",
+            "CreatedByURL" : "http://epicgames.com",
+            "DocsURL" : "",
+            "MarketplaceURL" : "",
+            "SupportURL" : "",
+            "EnabledByDefault" : false,
+            "CanContainContent" : true,
+            "IsBetaVersion" : false,
+            "Installed" : false,
+            "Modules" :
+            [
+                {
+                    "Name" : "Foo",
+                    "Type" : "Developer",
+                    "LoadingPhase" : "PostConfigInit"
+                }
+            ]
+        }
+
+<!--
+In the following section, we will take a look at how you create a new Plugin that has shader support in it. 
+
+1. First, make a copy of the **//Engine/Plugins/Developer/BlankPlugin** and rename it to **FooPlugin**.
+
+1. Inside the FooPlugin folder, create a new folder called **Shaders** and inside that folder create a new folder called **Private**.
+
+    [region:note]
+    Note that the  **"/Plugin/Foo/…"** is going to map to your Plugins **Foo Shaders/** directory and not the Engine/Shaders directory.
+    [/region]
+
+1. Inside of the newly created Private folder, create a new file called **MyShader.usf**.
+
+1. Once that is completed, re-run the **GenerateProjectFiles.bat** file to regenerate the projects files to include our new shader file.
+
+1. In **.../BlankPlugin/BlankPlugin.uplugin**, change the module's LoadingPhase to PostConfigInit (only for the modules that will have shader implementations) like in the following example.
+
+        {
+            "FileVersion" : 3,
+            "Version" : 1,
+            "VersionName" : "1.0",
+            "FriendlyName" : "Foo",
+            "Description" : "Plugin to play around with shaders.",
+            "Category" : "Sandbox",
+            "CreatedBy" : "Epic Games, Inc.",
+            "CreatedByURL" : "http://epicgames.com",
+            "DocsURL" : "",
+            "MarketplaceURL" : "",
+            "SupportURL" : "",
+            "EnabledByDefault" : false,
+            "CanContainContent" : true,
+            "IsBetaVersion" : false,
+            "Installed" : false,
+            "Modules" :
+            [
+                {
+                    "Name" : "Foo",
+                    "Type" : "Developer",
+                    "LoadingPhase" : "PostConfigInit"
+                }
+            ]
+        }
+-->
+## Render Thread
+
+As opposed to game-side API, the RHI render commands are (and should be) enqueued by a dedicated thread: the rendering thread. The rendering thread is a slave thread of the game thread because the latter enqueue FIFO (first in, first out) commands through the ENQUEUE_RENDER_COMMAND to the former. The rendering thread can be 0 or one frame behind the game thread. As a matter of CPU performance, the synchronization between them must be avoided at all cost in production runtime. To make sure your plugin's C++ function is called by the right thread, you can add multiple asserts such as check (IsInGameThread()); or check (IsInRenderingThread()); for improved threading robustness.
+
+## Unreal Shader Files
+
+There are two different shader file types that you need to be aware of when developing new shaders for use in UE4. Each file has a different purpose which you will find listed below:
+
+* **Unreal Shader Headers (.USH)**
+    * Only included by other USH or USF files
+
+* **Unreal Shader Format (.USF)**
+    * Should be private data only
+        * No backward compatibility is guaranteed in Private directories
+    * Should contain shader entry points
+
+
+## Shader File Preprocessing and Virtual File Path
+
+USF shader files, based on HLSL language, is Unreal Engine shader file format that contains the multi platform shader code. To achieve this multi platform support, the engine's shader compiler has added an extra platform-independent source file preprocessing pass before the platform-specific shader compiler (FXC, HLSLCC for GLSL cross compilation, etc.). As a result, all #define and #if are resolved at this very first preprocessing. Of course, each platform has built-in #define to know at shader preprocessing the targeted platform, such as VULKAN_PROFILE.
+
+As same as C/C++ files, you can include usf files with #include "HelloWorld.usf," that would include the file named HelloWorld.usf stored in the same directory as the USF file you have the #include written in. To avoid multiple includes of the same file, you can add the #pragma once pre processing directive at the top of your file. For instance:
+
+* FooCommon.usf 
+        // File shared across all Plugin's shaders
+        #pragma once
+
+        #include "/Engine/Public/Platform.ush"
+
+        // ...
+
+* FooBar.usf 
+        // File containing all foobar-related functions and structures.
+        #pragma once
+
+        #include "FooCommon.usf"
+
+        // ...
+
+* FooBarComputeShader.usf 
+        // Compute shader that does foobar on the GPU
+
+        #include "FooCommon.usf"
+        #include "FooBar.usf"
+
+        // ...
+You can also do this from a Plugin or project module's shader to include a USF file by doing either of the following:
+
+* In the engine with #include `/Engine/<FilePath>`, where `<FilePath>` is a file path relative to `//Engine/Shaders/` directory;
+
+* Or another plugin with #include `/Plugin/<PluginName>/<PluginFilePath>`, where `<PluginName>` is the name of an **activated** Plugin, and `<PluginFilePath>` is a file path relative to the Plugin's `Shaders/` directory. It is the responsibility of the developer to add a dependency to the correct Plugin in the .uplugin.
+
+
+## First Global Shader
+
+Global shaders inherit from the FGlobalShader in the following manner:
+
+
+	class FLensDistortionUVGenerationShader : public FGlobalShader
+	{
+	public:
+		// This function determines whether or not the shader should be compiled for a given platform.
+		// In this case, the shader will not be usable without SM4 support.
+		static bool ShouldCache(EShaderPlatform Platform)
+		{
+			return IsFeatureLevelSupported(Platform, ERHIFeatureLevel::SM4);
+		}
+
+		// Compile-time constants for shader can be defined in this function:
+		static void ModifyCompilationEnvironment(EShaderPlatform Platform, FShaderCompilerEnvironment& OutEnvironment)
+		{
+			FGlobalShader::ModifyCompilationEnvironment(Platform, OutEnvironment);
+			OutEnvironment.SetDefine(TEXT("GRID_SUBDIVISION_X"), kGridSubdivisionX);
+			OutEnvironment.SetDefine(TEXT("GRID_SUBDIVISION_Y"), kGridSubdivisionY);
+		}
+
+		// Default constructor.
+		FLensDistortionUVGenerationShader() {}
+
+		// Constructor using an initializer object. We can bind parameters here so that our C++ code
+		// can interface with USF, enabling us to set shader parameters from code.
+		FLensDistortionUVGenerationShader(const ShaderMetaType::CompiledShaderInitializerType& Initializer)
+			: FGlobalShader(Initializer)
+		{
+			PixelUVSize.Bind(Initializer.ParameterMap, TEXT("PixelUVSize"));
+			RadialDistortionCoefs.Bind(Initializer.ParameterMap, TEXT("RadialDistortionCoefs"));
+			TangentialDistortionCoefs.Bind(Initializer.ParameterMap, TEXT("TangentialDistortionCoefs"));
+			DistortedCameraMatrix.Bind(Initializer.ParameterMap, TEXT("DistortedCameraMatrix"));
+			UndistortedCameraMatrix.Bind(Initializer.ParameterMap, TEXT("UndistortedCameraMatrix"));
+			OutputMultiplyAndAdd.Bind(Initializer.ParameterMap, TEXT("OutputMultiplyAndAdd"));
+		}
+
+		// All members must be serialized here. This function is run at load and save time, and is used
+		// to put the shader into the DDC and pak files.
+		virtual bool Serialize(FArchive& Ar) override
+		{
+			bool bShaderHasOutdatedParameters = FGlobalShader::Serialize(Ar);
+			Ar << PixelUVSize << RadialDistortionCoefs << TangentialDistortionCoefs << DistortedCameraMatrix << UndistortedCameraMatrix << OutputMultiplyAndAdd;
+			return bShaderHasOutdatedParameters;
+		}
+
+		// This function is an example of how shader parameters can be precomputed based on data specific
+		// to the shader. In this case, the shader requires several matrices that can be computed from
+		// a few parameters, and that would be inefficient to compute inside the shader itself. Note that
+		// this function is not an override; it is custom to this class, and is called when this feature's
+		// specific implementation requires it.
+		template<typename TShaderRHIParamRef>
+		void SetParameters(
+			FRHICommandListImmediate& RHICmdList,
+			const TShaderRHIParamRef ShaderRHI,
+			const FCompiledCameraModel& CompiledCameraModel,
+			const FIntPoint& DisplacementMapResolution)
+		{
+			FVector2D PixelUVSizeValue(
+				1.f / float(DisplacementMapResolution.X), 1.f / float(DisplacementMapResolution.Y));
+			FVector RadialDistortionCoefsValue(
+				CompiledCameraModel.OriginalCameraModel.K1,
+				CompiledCameraModel.OriginalCameraModel.K2,
+				CompiledCameraModel.OriginalCameraModel.K3);
+			FVector2D TangentialDistortionCoefsValue(
+				CompiledCameraModel.OriginalCameraModel.P1,
+				CompiledCameraModel.OriginalCameraModel.P2);
+
+			SetShaderValue(RHICmdList, ShaderRHI, PixelUVSize, PixelUVSizeValue);
+			SetShaderValue(RHICmdList, ShaderRHI, DistortedCameraMatrix, CompiledCameraModel.DistortedCameraMatrix);
+			SetShaderValue(RHICmdList, ShaderRHI, UndistortedCameraMatrix, CompiledCameraModel.UndistortedCameraMatrix);
+			SetShaderValue(RHICmdList, ShaderRHI, RadialDistortionCoefs, RadialDistortionCoefsValue);
+			SetShaderValue(RHICmdList, ShaderRHI, TangentialDistortionCoefs, TangentialDistortionCoefsValue);
+			SetShaderValue(RHICmdList, ShaderRHI, OutputMultiplyAndAdd, CompiledCameraModel.OutputMultiplyAndAdd);
+		}
+
+	private:
+		// Shader parameters.
+		FShaderParameter PixelUVSize;
+		FShaderParameter RadialDistortionCoefs;
+		FShaderParameter TangentialDistortionCoefs;
+		FShaderParameter DistortedCameraMatrix;
+		FShaderParameter UndistortedCameraMatrix;
+		FShaderParameter OutputMultiplyAndAdd;
+	};
+
+	// This macro exposes the shader to the Engine. Note the absolute virtual source file path.
+	IMPLEMENT_SHADER_TYPE(, FLensDistortionUVGenerationVS, TEXT("/Plugin/LensDistortion/Private/UVGeneration.usf"), TEXT("MainVS"), SF_Vertex)
+
+
+### Engine/Public/Platform.usf
+
+To have your shader compiling on all UE4 platforms, you are required to include /Engine/Public/Platform.usf in all your shader files (directly or indirectly).
+
+## Shader Development Tips
+
+You can customize locally using ConsoleVariables.ini to change some settings in the renderer to accelerate iteration process when writing shaders. For example, the following Console Variables will help you get detailed debug information about what your shader is doing:
+
+* **r.ShaderDevelopmentMode = 1** To get detailed logs on shader compiles and the opportunity to retry on errors.
+
+* **r.DumpShaderDebugInfo = 1** To dump preprocessed shaders in the Saved folder. 
+
+    [region:warning]
+    Warning: leaving this on for a while will fill your hard drive with many small files and folders so make sure to disable it when you are done.
+    [/region]
+    
+
+
+## Troubleshooting 
+
+If you are having issues getting your shader to compile or show up in the UE4 editor, try the following:
+
+* If you get the error **Can't compile:** `/Plugin/<MyPluginName>/<MyFile>` **not found.** Make sure to check your Plugin's module's LoadingPhase is set to PostConfigInit, and that there are no typo's in the Plugin's Shaders/ directory name.
+
+* If you get the error **Can't #include** `"/Plugin/<ParentPluginName>/<MyFile>":` Make sure to check that the parent Plugin is activated and also check your Plugin dependency as this error means you are missing a Plugin dependency in your .uplugin or .uproject.
+
+
+## Existing Renderer Convention
+
+In the renderer, we tend to have a convention on naming shader classes and shader entry point, especially with a shader domain suffix as shown in the following table:
+
+|Shader Domain  | Suffix |
+|-------------- |--------|
+|**Vertex Shader**  |VS|
+|**Hull Shader **   |HS|
+|**Domain Shader**  |DS|
+|**Geometry Shader**|GS|
+|**Pixel Shader  ** |PS|
+|**Compute Shader **|CS|
+
+For example, in a C++ file, the call to **FLensDistortionUVGenerationVS** has VS at the end signaling that it is a Vertex Shader. Inside of a USF file the **void MainVS(...)** ends with a VS signaling that it is going to use the Vertex Shader. When dealing with Struct's in HLSL, the struct name should start with **F** like FBasePassInterpolators for example.
+
+[region:note]
+To read more about coding standards in UE4 check out the [UE4 Coding Standards document](https://docs.unrealengine.com/latest/INT/Programming/Development/CodingStandard/) for more information. 
+[/region]
+
+## Additional Links
+
+The following links contain more information about developing Global Shaders inside UE4.
+
+*  [FGlobalShader Base Class](https://docs.unrealengine.com/latest/INT/API/Runtime/ShaderCore/FGlobalShader/index.html)
+*  [Debugging the Shader Compiling Process](https://www.unrealengine.com/en-US/blog/debugging-the-shader-compiling-process)
+*  [Adding Global Sjaders to UE4](https://rcaloca.blogspot.com/2017/05/adding-global-shaders-to-ue4.html)
+
+
+
+
