Copying //UE4/Dev-Documentation to //UE4/Samples-Main (Source: //UE4/Dev-Documentation @ 3623665)

diff --git a/Engine/Documentation/Source/Programming/Rendering/ShaderInPlugin/QuickStart/1/QSSIP_01.INT.udn b/Engine/Documentation/Source/Programming/Rendering/ShaderInPlugin/QuickStart/1/QSSIP_01.INT.udn
new file mode 100644
--- /dev/null
+++ b/Engine/Documentation/Source/Programming/Rendering/ShaderInPlugin/QuickStart/1/QSSIP_01.INT.udn
@@ -0,0 +1,950 @@
+Availability: Docs 
+Crumbs: %ROOT%
+Title: Creating a New Global Shader in a Plugin - Code Setup
+Description: Setting up the plugin and code for the Gloabl Shader plugin
+Type:quick start
+Type:multi-step
+SkillLevel:Intermediate
+Version: 4.17
+Parent:Programming/Rendering/ShaderInPlugin/QuickStart
+Order:1
+Tags:Rendering
+Tags:Programming
+Checkpoint:shaderqs
+
+[Nav]
+
+Before you can start creating a new Plugin for Unreal Engine 4 (UE4) you will first need to make sure that you have Visual Studio installed. Visual Studio is required for this Quick Start as you will need to compile your Plugin's code so that it can run. If you are not sure how to go about doing this, check out the [Setting Up Visual Studio for UE4](Programming/Development/VisualStudioSetup/) documents for more help.   
+
+1. First, create a new C++ **Basic Code** project that uses the **Desktop Setting** with **Maximum Quality** and does not have any **Starter Content**. 
+
+    [REGION:lightbox]
+    [![](HT_ShadersInPlugins_00.png)(w:555)](HT_ShadersInPlugins_00.png)
+    [/REGION]
+
+    [REGION:caption]
+    Click for full image.
+    [/REGION]
+
+1. Once the project has been created, Visual Studio will open up and after it has,  right- click on the ShadersInPlugins project and select the **Build** option to compile the project.
+
+    ![](HT_ShadersInPlugins_01.png)(w:481)
+
+1. When the project has finished compiling, press the **F5** key inside Visual Studio to launch the ShadersInPlugins project in the UE4 Editor.
+
+1. Once the UE4 Editor has finished loading, open up the **Plugins** manager by going to **Edit** > **Plugins** and then click on the **New Plugin** option that is on the lower right-hand side of the Plugin window to bring up the New Plugin creation window. 
+    [REGION:lightbox]
+    [![](HT_ShadersInPlugins_02.png)(w:564)](HT_ShadersInPlugins_02.png)
+    [/REGION]
+
+    [REGION:caption]
+    Click for full image.
+    [/REGION]
+
+1. From the **New Plugin** window, select the **Blank** plugin and then input the name **Foo** for the name and keeping all of the settings as default. When all of that is completed, press the **Create Plugin** button to create all of the initial required content for the plugin.
+
+    [REGION:lightbox]
+    [![](HT_ShadersInPlugins_03.png)(w:564)](HT_ShadersInPlugins_03.png)
+    [/REGION]
+	[REGION:caption] 
+	Click for full image. 
+	[/REGION]
+
+1. When that is completed, shut down UE4 and Visual Studio and then go to the **Plugins** > **Foo** plugin folder that was created inside your project folder.
+
+1. Inside the Foo plugins folder, add a new folder called **Shaders** and then inside of that folder, create a new folder called **Private**.
+
+    [REGION:lightbox]
+    [![](HT_ShadersInPlugins_04.png)(w:564)](HT_ShadersInPlugins_04.png)
+    [/REGION]
+
+    [REGION:caption]
+    Click for full image.
+    [/REGION]
+
+1. Inside of the Private folder, create a new text file and call it **MyShader.USF** and then copy and paste the following HLSL code into this file, saving the file when completed.
+
+    [REGION:note]
+    Make sure that you change the file extension to **.USF** or UE4 will not understand that this is a shader file.
+    [/REGION]
+
+        // Copyright 1998-2017 Epic Games, Inc. All Rights Reserved.
+
+        /*=============================================================================
+            LensDistortionUVGeneration.usf: Generate lens distortion and undistortion
+            UV displacement map into a render target.
+
+            The pixel shader directly compute the distort viewport UV to undistort
+            viewport UV displacement using Sv_Position and the reference equations and
+            store them into the red and green channels.
+
+            However to avoid resolving with a ferrari method, or doing a newton method
+            on the GPU to compute the undistort viewport UV to distort viewport UV
+            displacement, this couple of shaders works as follow: The vertex shader
+            undistort the grid's vertices, and pass down to the pixel shader the viewport
+            UV of where they should have been on screen without undistortion. The pixel
+            shader can then generate the undistort viewport UV to distort viewport UV
+            displacement by just subtracting the pixel's viewport UV.
+        =============================================================================*/
+
+        #include "/Engine/Public/Platform.ush"
+
+        // Size of the pixels in the viewport UV coordinates.
+        float2 PixelUVSize;
+
+        // K1, K2, K3
+        float3 RadialDistortionCoefs;
+
+        // P1, P2
+        float2 TangentialDistortionCoefs;
+
+        // Camera matrix of the undistorted viewport.
+        float4 UndistortedCameraMatrix;
+
+        // Camera matrix of the distorted viewport.
+        float4 DistortedCameraMatrix;
+
+        // Output multiply and add to the render target.
+        float2 OutputMultiplyAndAdd;
+
+        // Undistort a view position at V.z=1.
+        float2 UndistortNormalizedViewPosition(float2 V)
+        {
+            float2 V2 = V * V;
+            float R2 = V2.x + V2.y;
+
+            // Radial distortion (extra parenthesis to match MF_Undistortion.uasset).
+            float2 UndistortedV = V * (1.0 + R2 * (RadialDistortionCoefs.x + R2 * (RadialDistortionCoefs.y + R2 * RadialDistortionCoefs.z)));
+
+            // Tangential distortion.
+            UndistortedV.x += TangentialDistortionCoefs.y * (R2 + 2 * V2.x) + 2 * TangentialDistortionCoefs.x * V.x * V.y;
+            UndistortedV.y += TangentialDistortionCoefs.x * (R2 + 2 * V2.y) + 2 * TangentialDistortionCoefs.y * V.x * V.y;
+
+            return UndistortedV;
+        }
+
+        // Returns the undistorted viewport UV of the distorted's viewport UV.
+        //
+        // Notes:
+        //        UVs are bottom left originated.
+        float2 UndistortViewportUV(float2 ViewportUV)
+        {
+            // Distorted viewport UV -> Distorted view position (z=1)
+            float2 DistortedViewPosition = (ViewportUV - DistortedCameraMatrix.zw) / DistortedCameraMatrix.xy;
+
+            // Compute undistorted view position (z=1)
+            float2 UndistortedViewPosition = UndistortNormalizedViewPosition(DistortedViewPosition);
+
+            // Undistorted view position (z=1) -> Undistorted viewport UV.
+            return UndistortedCameraMatrix.xy * UndistortedViewPosition + UndistortedCameraMatrix.zw;
+        }
+
+        // Flip UV's y component.
+        float2 FlipUV(float2 UV)
+        {
+            return float2(UV.x, 1 - UV.y);
+        }
+
+        void MainVS(
+            in uint GlobalVertexId : SV_VertexID,
+            out float2 OutVertexDistortedViewportUV : TEXCOORD0,
+            out float4 OutPosition : SV_POSITION
+            )
+        {
+            // Compute the cell index.
+            uint GridCellIndex = GlobalVertexId / 6;
+
+            // Compute row and column id of the cell within the grid.
+            uint GridColumnId = GridCellIndex / GRID_SUBDIVISION_Y;
+            uint GridRowId = GridCellIndex - GridColumnId * GRID_SUBDIVISION_Y;
+
+            // Compute the vertex id within a 2 triangles grid cell.
+            uint VertexId = GlobalVertexId - GridCellIndex * 6;
+
+            // Compute the bottom left originated UV coordinate of the triangle's vertex within the cell.
+            float2 CellVertexUV = float2(0x1 & ((VertexId + 1) / 3), VertexId & 0x1);
+
+            // Compute the top left originated UV of the vertex within the grid.
+            float2 GridInvSize = 1.f / float2(GRID_SUBDIVISION_X, GRID_SUBDIVISION_Y);
+            float2 GridVertexUV = FlipUV(
+                GridInvSize * (CellVertexUV + float2(GridColumnId, GridRowId)));
+
+            // The standard doesn't have half pixel shift.
+            GridVertexUV -= PixelUVSize * 0.5;
+
+            // Output vertex position.
+            OutPosition = float4(FlipUV(
+                UndistortViewportUV(GridVertexUV) + PixelUVSize * 0.5) * 2 - 1, 0, 1);
+
+            // Output top left originated UV of the vertex.
+            OutVertexDistortedViewportUV = GridVertexUV;
+        }
+
+        void MainPS(
+            in noperspective float2 VertexDistortedViewportUV : TEXCOORD0,
+            in float4 SvPosition : SV_POSITION,
+            out float4 OutColor : SV_Target0
+            )
+        {
+            // Compute the pixel's top left originated UV.
+            float2 ViewportUV = SvPosition.xy * PixelUVSize;
+
+            // The standard doesn't have half pixel shift.
+            ViewportUV -= PixelUVSize * 0.5;
+
+            float2 DistortUVtoUndistortUV = (UndistortViewportUV((ViewportUV))) - ViewportUV;
+            float2 UndistortUVtoDistortUV = VertexDistortedViewportUV - ViewportUV;
+
+            // Output displacement channels.
+            OutColor = OutputMultiplyAndAdd.y + OutputMultiplyAndAdd.x * float4(
+                DistortUVtoUndistortUV, UndistortUVtoDistortUV);
+        }
+
+1. Now, locate the **Foo.uplugin** file and open it up inside of a text editor and replace the information in that file with the following text, saving the file when done.
+
+        {
+            "FileVersion" : 3,
+            "Version" : 1,
+            "VersionName" : "1.0",
+            "FriendlyName" : "Foo",
+            "Description" : "Plugin to play around with shaders.",
+            "Category" : "Sandbox",
+            "CreatedBy" : "Epic Games, Inc.",
+            "CreatedByURL" : "http://epicgames.com",
+            "DocsURL" : "",
+            "MarketplaceURL" : "",
+            "SupportURL" : "",
+            "EnabledByDefault" : false,
+            "CanContainContent" : true,
+            "IsBetaVersion" : false,
+            "Installed" : false,
+            "Modules" :
+            [
+                {
+                    "Name" : "Foo",
+                    "Type" : "Developer",
+                    "LoadingPhase" : "PostConfigInit"
+                }
+            ]
+        }
+
+1. Next, go to **Plugins\Foo\Source\Foo** and create a new folder called, **Classes** and then copy both the **LensDistortionAPI.h** and **LensDistortionBlueprintLibrary.h** files from  **Engine\Plugins\Compositing\LensDistortion** to this newly created folder.
+
+    [region:note]
+    The files that we will be copying can be found in **Engine\Plugins\Compositing\LensDistortion**.
+    [/region]
+
+    * Classes - Create New Folder
+
+        * Copy - LensDistortionAPI.h
+        * Copy - LensDistortionBlueprintLibrary.h
+
+        ![](HT_ShadersInPlugins_05.png)
+
+1. Then, go to the **Private** folder and copy both the **LensDistortionBlueprintLibrary.cpp** and the **LensDistortionRendering.cpp** files to this Private folder.
+
+    * Private - Existing Folder
+
+        * Copy - LensDistortionBlueprintLibrary.cpp
+        * Copy - LensDistortionRendering.cpp
+
+        ![](HT_ShadersInPlugins_06.png)
+
+1. Now close both the UE4 Editor and Visual Studio and then locate your projects .U project file. When located, right-click on it and select the **Generate Visual Studio project files** option.
+
+    ![](HT_ShadersInPlugins_07.png)
+
+1. Reopen your Visual Studio solution and then go to Foo > Classes and open up the **LensDistortionAPI.h** file. Inside this file replace **FLensDistortionCameraModel** with **FFooCameraModel**.
+
+    [region:note]
+    You should replace FLensDistortionCameraModel with FFooCameraModel **four** times in this file.
+    [/region]
+
+            // Copyright 1998-2017 Epic Games, Inc. All Rights Reserved.
+
+            #pragma once
+
+            #include "CoreMinimal.h"
+            #include "LensDistortionAPI.generated.h"
+
+            /** Mathematic camera model for lens distortion/undistortion.
+             *
+             * Camera matrix =
+             *  | F.X  0  C.x |
+             *  |  0  F.Y C.Y |
+             *  |  0   0   1  |
+             */
+            USTRUCT(BlueprintType)
+            struct FFooCameraModel
+            {
+                GENERATED_USTRUCT_BODY()
+                FFooCameraModel()
+                {
+                    K1 = K2 = K3 = P1 = P2 = 0.f;
+                    F = FVector2D(1.f, 1.f);
+                    C = FVector2D(0.5f, 0.5f);
+                }
+
+                /** Radial parameter #1. */
+                UPROPERTY(Interp, EditAnywhere, BlueprintReadWrite, Category = "Lens Distortion|Camera Model")
+                float K1;
+
+                /** Radial parameter #2. */
+                UPROPERTY(Interp, EditAnywhere, BlueprintReadWrite, Category = "Lens Distortion|Camera Model")
+                float K2;
+
+                /** Radial parameter #3. */
+                UPROPERTY(Interp, EditAnywhere, BlueprintReadWrite, Category = "Lens Distortion|Camera Model")
+                float K3;
+
+                /** Tangential parameter #1. */
+                UPROPERTY(Interp, EditAnywhere, BlueprintReadWrite, Category = "Lens Distortion|Camera Model")
+                float P1;
+
+                /** Tangential parameter #2. */
+                UPROPERTY(Interp, EditAnywhere, BlueprintReadWrite, Category = "Lens Distortion|Camera Model")
+                float P2;
+
+                /** Camera matrix's Fx and Fy. */
+                UPROPERTY(Interp, EditAnywhere, BlueprintReadWrite, Category = "Lens Distortion|Camera Model")
+                FVector2D F;
+
+                /** Camera matrix's Cx and Cy. */
+                UPROPERTY(Interp, EditAnywhere, BlueprintReadWrite, Category = "Lens Distortion|Camera Model")
+                FVector2D C;
+    
+                /** Undistorts 3d vector (x, y, z=1.f) in the view space and returns (x', y', z'=1.f). */
+                FVector2D UndistortNormalizedViewPosition(FVector2D V) const;
+
+                /** Returns the overscan factor required for the undistort rendering to avoid unrendered distorted pixels. */
+                float GetUndistortOverscanFactor(
+                    float DistortedHorizontalFOV,
+                    float DistortedAspectRatio) const;
+
+                /** Draws UV displacement map within the output render target.
+                 * - Red & green channels hold the distortion displacement;
+                 * - Blue & alpha channels hold the undistortion displacement.
+                 * @param World Current world to get the rendering settings from (such as feature level).
+                 * @param DistortedHorizontalFOV The desired horizontal FOV in the distorted render.
+                 * @param DistortedAspectRatio The desired aspect ratio of the distorted render.
+                 * @param UndistortOverscanFactor The factor of the overscan for the undistorted render.
+                 * @param OutputRenderTarget The render target to draw to. Don't necessarily need to have same resolution or aspect ratio as distorted render.
+                 * @param OutputMultiply The multiplication factor applied on the displacement.
+                 * @param OutputAdd Value added to the multiplied displacement before storing the output render target.
+                 */
+                void DrawUVDisplacementToRenderTarget(
+                    class UWorld* World,
+                    float DistortedHorizontalFOV,
+                    float DistortedAspectRatio,
+                    float UndistortOverscanFactor,
+                    class UTextureRenderTarget2D* OutputRenderTarget,
+                    float OutputMultiply,
+                    float OutputAdd) const;
+
+                /** Compare two lens distortion models and return whether they are equal. */
+                bool operator == (const FFooCameraModel& Other) const
+                {
+                    return (
+                        K1 == Other.K1 &&
+                        K2 == Other.K2 &&
+                        K3 == Other.K3 &&
+                        P1 == Other.P1 &&
+                        P2 == Other.P2 &&
+                        F == Other.F &&
+                        C == Other.C);
+                }
+
+                /** Compare two lens distortion models and return whether they are different. */
+                bool operator != (const FFooCameraModel& Other) const
+                {
+                    return !(*this == Other);
+                }
+            };
+
+1. Next, open up the **LensDistortionBlueprintLibrary.h** file. This file controls how this node will show up in Blueprints so not only do we need to replace **FLensDistortionCameraModel** with **FFooCameraModel**, we also need to change the **Category = "Lens Distortion" ** to **Category = "Foo | Lens Distortion"**.
+
+    [region:note]
+    You should replace FLensDistortionCameraModel with FFooCameraModel **six** times in this file.
+    [/region]
+
+        // Copyright 1998-2017 Epic Games, Inc. All Rights Reserved.
+
+        #pragma once
+
+        #include "CoreMinimal.h"
+        #include "UObject/ObjectMacros.h"
+        #include "Classes/Kismet/BlueprintFunctionLibrary.h"
+        #include "LensDistortionAPI.h"
+        #include "LensDistortionBlueprintLibrary.generated.h"
+
+        UCLASS(MinimalAPI)
+        class ULensDistortionBlueprintLibrary : public UBlueprintFunctionLibrary
+        {
+            GENERATED_UCLASS_BODY()
+
+            /** Returns the overscan factor required for the undistort rendering to avoid unrendered distorted pixels. */
+            UFUNCTION(BlueprintPure, Category = "Foo | Lens Distortion")
+            static void GetUndistortOverscanFactor(
+                const FFooCameraModel& CameraModel,
+                float DistortedHorizontalFOV,
+                float DistortedAspectRatio,
+                float& UndistortOverscanFactor);
+        
+            /** Draws UV displacement map within the output render target.
+             * - Red & green channels hold the distortion displacement;
+             * - Blue & alpha channels hold the undistortion displacement.
+             * @param DistortedHorizontalFOV The desired horizontal FOV in the distorted render.
+             * @param DistortedAspectRatio The desired aspect ratio of the distorted render.
+             * @param UndistortOverscanFactor The factor of the overscan for the undistorted render.
+             * @param OutputRenderTarget The render target to draw to. Don't necessarily need to have same resolution or aspect ratio as distorted render.
+             * @param OutputMultiply The multiplication factor applied on the displacement.
+             * @param OutputAdd Value added to the multiplied displacement before storing into the output render target.
+             */
+            UFUNCTION(BlueprintCallable, Category = "Foo | Lens Distortion", meta = (WorldContext = "WorldContextObject"))
+            static void DrawUVDisplacementToRenderTarget(
+                const UObject* WorldContextObject,
+                const FFooCameraModel& CameraModel,
+                float DistortedHorizontalFOV,
+                float DistortedAspectRatio,
+                float UndistortOverscanFactor,
+                class UTextureRenderTarget2D* OutputRenderTarget,
+                float OutputMultiply = 0.5,
+                float OutputAdd = 0.5
+                );
+
+            /* Returns true if A is equal to B (A == B) */
+            UFUNCTION(BlueprintPure, meta=(DisplayName = "Equal (LensDistortionCameraModel)", CompactNodeTitle = "==", Keywords = "== equal"), Category = "Foo | Lens Distortion")
+            static bool EqualEqual_CompareLensDistortionModels(
+                const FFooCameraModel& A,
+                const FFooCameraModel& B)
+            {
+                return A == B;
+            }
+
+            /* Returns true if A is not equal to B (A != B) */
+            UFUNCTION(BlueprintPure, meta = (DisplayName = "NotEqual (LensDistortionCameraModel)", CompactNodeTitle = "!=", Keywords = "!= not equal"), Category = "Foo | Lens Distortion")
+            static bool NotEqual_CompareLensDistortionModels(
+                const FFooCameraModel& A,
+                const FFooCameraModel& B)
+            {
+                return A != B;
+            }
+        };
+
+1. Now go to the **Private** folder and open up the **LensDistortionBlueprintLibrary.cpp** file and replace the following:
+
+    * **FLensDistortionCameraModel** with **FFooCameraModel**
+    * **ULensDistortionBlueprintLibrary** with **UFooBlueprintLibrary**
+
+    [region:note]
+    You should replace FLensDistortionCameraModel with FFooCameraModel **two** times and ULensDistortionBlueprintLibrary with UFooBlueprintLibrary **four** times.
+    [/region]
+
+            // Copyright 1998-2017 Epic Games, Inc. All Rights Reserved.
+
+            #include "LensDistortionBlueprintLibrary.h"
+
+            ULensDistortionBlueprintLibrary::ULensDistortionBlueprintLibrary(const FObjectInitializer& ObjectInitializer)
+                : Super(ObjectInitializer)
+            { }
+
+            // static
+            void ULensDistortionBlueprintLibrary::GetUndistortOverscanFactor(
+                const FFooCameraModel& CameraModel,
+                float DistortedHorizontalFOV,
+                float DistortedAspectRatio,
+                float& UndistortOverscanFactor)
+            {
+                UndistortOverscanFactor = CameraModel.GetUndistortOverscanFactor(DistortedHorizontalFOV, DistortedAspectRatio);
+            }
+
+            // static
+            void ULensDistortionBlueprintLibrary::DrawUVDisplacementToRenderTarget(
+                const UObject* WorldContextObject,
+                const FFooCameraModel& CameraModel,
+                float DistortedHorizontalFOV,
+                float DistortedAspectRatio,
+                float UndistortOverscanFactor,
+                class UTextureRenderTarget2D* OutputRenderTarget,
+                float OutputMultiply,
+                float OutputAdd)
+            {
+                CameraModel.DrawUVDisplacementToRenderTarget(
+                    WorldContextObject->GetWorld(),
+                    DistortedHorizontalFOV, DistortedAspectRatio,
+                    UndistortOverscanFactor, OutputRenderTarget,
+                    OutputMultiply, OutputAdd);
+            }
+
+1. Next, in the **Private** folder open up the **LensDistortionRendering.cpp** file and replace **FLensDistortionCameraModel** with **FFooCameraModel**
+
+    [region:note]
+    You should replace FLensDistortionCameraModel with FFooCameraModel **six** times in this file.
+    [/region]
+            // Copyright 1998-2017 Epic Games, Inc. All Rights Reserved.
+
+            #include "LensDistortionAPI.h"
+            #include "Classes/Engine/TextureRenderTarget2D.h"
+            #include "Classes/Engine/World.h"
+            #include "Public/GlobalShader.h"
+            #include "Public/PipelineStateCache.h"
+            #include "Public/RHIStaticStates.h"
+            #include "Public/SceneUtils.h"
+            #include "Public/SceneInterface.h"
+            #include "Public/ShaderParameterUtils.h"
+
+            static const uint32 kGridSubdivisionX = 32;
+            static const uint32 kGridSubdivisionY = 16;
+
+            /**
+             * Internal intermediary structure derived from FFooCameraModel by the game thread
+             * to hand to the render thread.
+             */
+            struct FCompiledCameraModel
+            {
+                /** Orignal camera model that has generated this compiled model. */
+                FFooCameraModel OriginalCameraModel;
+
+                /** Camera matrices of the lens distortion for the undistorted and distorted render.
+                 *  XY holds the scales factors, ZW holds the translates.
+                 */
+                FVector4 DistortedCameraMatrix;
+                FVector4 UndistortedCameraMatrix;
+
+                /** Output multiply and add of the channel to the render target. */
+                FVector2D OutputMultiplyAndAdd;
+            };
+
+            /** Undistorts top left originated viewport UV into the view space (x', y', z'=1.f) */
+            static FVector2D LensUndistortViewportUVIntoViewSpace(
+                const FFooCameraModel& CameraModel,
+                float TanHalfDistortedHorizontalFOV, float DistortedAspectRatio,
+                FVector2D DistortedViewportUV)
+            {
+                FVector2D AspectRatioAwareF = CameraModel.F * FVector2D(1, -DistortedAspectRatio);
+                return CameraModel.UndistortNormalizedViewPosition((DistortedViewportUV - CameraModel.C) / AspectRatioAwareF);
+            }
+
+            class FLensDistortionUVGenerationShader : public FGlobalShader
+            {
+            public:
+                static bool ShouldCache(EShaderPlatform Platform)
+                {
+                    return IsFeatureLevelSupported(Platform, ERHIFeatureLevel::SM4);
+                }
+
+                static void ModifyCompilationEnvironment(EShaderPlatform Platform, FShaderCompilerEnvironment& OutEnvironment)
+                {
+                    FGlobalShader::ModifyCompilationEnvironment(Platform, OutEnvironment);
+                    OutEnvironment.SetDefine(TEXT("GRID_SUBDIVISION_X"), kGridSubdivisionX);
+                    OutEnvironment.SetDefine(TEXT("GRID_SUBDIVISION_Y"), kGridSubdivisionY);
+                }
+
+                FLensDistortionUVGenerationShader() {}
+
+                FLensDistortionUVGenerationShader(const ShaderMetaType::CompiledShaderInitializerType& Initializer)
+                    : FGlobalShader(Initializer)
+                {
+                    PixelUVSize.Bind(Initializer.ParameterMap, TEXT("PixelUVSize"));
+                    RadialDistortionCoefs.Bind(Initializer.ParameterMap, TEXT("RadialDistortionCoefs"));
+                    TangentialDistortionCoefs.Bind(Initializer.ParameterMap, TEXT("TangentialDistortionCoefs"));
+                    DistortedCameraMatrix.Bind(Initializer.ParameterMap, TEXT("DistortedCameraMatrix"));
+                    UndistortedCameraMatrix.Bind(Initializer.ParameterMap, TEXT("UndistortedCameraMatrix"));
+                    OutputMultiplyAndAdd.Bind(Initializer.ParameterMap, TEXT("OutputMultiplyAndAdd"));
+                }
+
+                template<typename TShaderRHIParamRef>
+                void SetParameters(
+                    FRHICommandListImmediate& RHICmdList,
+                    const TShaderRHIParamRef ShaderRHI,
+                    const FCompiledCameraModel& CompiledCameraModel,
+                    const FIntPoint& DisplacementMapResolution)
+                {
+                    FVector2D PixelUVSizeValue(
+                        1.f / float(DisplacementMapResolution.X), 1.f / float(DisplacementMapResolution.Y));
+                    FVector RadialDistortionCoefsValue(
+                        CompiledCameraModel.OriginalCameraModel.K1,
+                        CompiledCameraModel.OriginalCameraModel.K2,
+                        CompiledCameraModel.OriginalCameraModel.K3);
+                    FVector2D TangentialDistortionCoefsValue(
+                        CompiledCameraModel.OriginalCameraModel.P1,
+                        CompiledCameraModel.OriginalCameraModel.P2);
+
+                    SetShaderValue(RHICmdList, ShaderRHI, PixelUVSize, PixelUVSizeValue);
+                    SetShaderValue(RHICmdList, ShaderRHI, DistortedCameraMatrix, CompiledCameraModel.DistortedCameraMatrix);
+                    SetShaderValue(RHICmdList, ShaderRHI, UndistortedCameraMatrix, CompiledCameraModel.UndistortedCameraMatrix);
+                    SetShaderValue(RHICmdList, ShaderRHI, RadialDistortionCoefs, RadialDistortionCoefsValue);
+                    SetShaderValue(RHICmdList, ShaderRHI, TangentialDistortionCoefs, TangentialDistortionCoefsValue);
+                    SetShaderValue(RHICmdList, ShaderRHI, OutputMultiplyAndAdd, CompiledCameraModel.OutputMultiplyAndAdd);
+                }
+
+                virtual bool Serialize(FArchive& Ar) override
+                {
+                    bool bShaderHasOutdatedParameters = FGlobalShader::Serialize(Ar);
+                    Ar << PixelUVSize << RadialDistortionCoefs << TangentialDistortionCoefs << DistortedCameraMatrix << UndistortedCameraMatrix << OutputMultiplyAndAdd;
+                    return bShaderHasOutdatedParameters;
+                }
+
+            private:
+                FShaderParameter PixelUVSize;
+                FShaderParameter RadialDistortionCoefs;
+                FShaderParameter TangentialDistortionCoefs;
+                FShaderParameter DistortedCameraMatrix;
+                FShaderParameter UndistortedCameraMatrix;
+                FShaderParameter OutputMultiplyAndAdd;
+
+            };
+
+            class FLensDistortionUVGenerationVS : public FLensDistortionUVGenerationShader
+            {
+                DECLARE_SHADER_TYPE(FLensDistortionUVGenerationVS, Global);
+
+            public:
+
+                /** Default constructor. */
+                FLensDistortionUVGenerationVS() {}
+
+                /** Initialization constructor. */
+                FLensDistortionUVGenerationVS(const ShaderMetaType::CompiledShaderInitializerType& Initializer)
+                    : FLensDistortionUVGenerationShader(Initializer)
+                {
+                }
+            };
+
+            class FLensDistortionUVGenerationPS : public FLensDistortionUVGenerationShader
+            {
+                DECLARE_SHADER_TYPE(FLensDistortionUVGenerationPS, Global);
+
+            public:
+
+                /** Default constructor. */
+                FLensDistortionUVGenerationPS() {}
+
+                /** Initialization constructor. */
+                FLensDistortionUVGenerationPS(const ShaderMetaType::CompiledShaderInitializerType& Initializer)
+                    : FLensDistortionUVGenerationShader(Initializer)
+                { }
+            };
+
+            IMPLEMENT_SHADER_TYPE(, FLensDistortionUVGenerationVS, TEXT("/Plugin/Foo/Private/MyShader.usf"), TEXT("MainVS"), SF_Vertex)
+            IMPLEMENT_SHADER_TYPE(, FLensDistortionUVGenerationPS, TEXT("/Plugin/Foo/Private/MyShader.usf"), TEXT("MainPS"), SF_Pixel)
+
+            static void DrawUVDisplacementToRenderTarget_RenderThread(
+                FRHICommandListImmediate& RHICmdList,
+                const FCompiledCameraModel& CompiledCameraModel,
+                const FName& TextureRenderTargetName,
+                FTextureRenderTargetResource* OutTextureRenderTargetResource,
+                ERHIFeatureLevel::Type FeatureLevel)
+            {
+                check(IsInRenderingThread());
+
+            #if WANTS_DRAW_MESH_EVENTS
+                FString EventName;
+                TextureRenderTargetName.ToString(EventName);
+                SCOPED_DRAW_EVENTF(RHICmdList, SceneCapture, TEXT("LensDistortionDisplacementGeneration %s"), *EventName);
+            #else
+                SCOPED_DRAW_EVENT(RHICmdList, DrawUVDisplacementToRenderTarget_RenderThread);
+            #endif
+
+                // Set render target.
+                SetRenderTarget(
+                    RHICmdList,
+                    OutTextureRenderTargetResource->GetRenderTargetTexture(),
+                    FTextureRHIRef(),
+                    ESimpleRenderTargetMode::EUninitializedColorAndDepth,
+                    FExclusiveDepthStencil::DepthNop_StencilNop);
+
+                FIntPoint DisplacementMapResolution(OutTextureRenderTargetResource->GetSizeX(), OutTextureRenderTargetResource->GetSizeY());
+
+                // Update viewport.
+                RHICmdList.SetViewport(
+                    0, 0, 0.f,
+                    DisplacementMapResolution.X, DisplacementMapResolution.Y, 1.f);
+    
+                // Get shaders.
+                TShaderMap<FGlobalShaderType>* GlobalShaderMap = GetGlobalShaderMap(FeatureLevel);
+                TShaderMapRef< FLensDistortionUVGenerationVS > VertexShader(GlobalShaderMap);
+                TShaderMapRef< FLensDistortionUVGenerationPS > PixelShader(GlobalShaderMap);
+
+                // Set the graphic pipeline state.
+                FGraphicsPipelineStateInitializer GraphicsPSOInit;
+                RHICmdList.ApplyCachedRenderTargets(GraphicsPSOInit);
+                GraphicsPSOInit.DepthStencilState = TStaticDepthStencilState<false, CF_Always>::GetRHI();
+                GraphicsPSOInit.BlendState = TStaticBlendState<>::GetRHI();
+                GraphicsPSOInit.RasterizerState = TStaticRasterizerState<>::GetRHI();
+                GraphicsPSOInit.PrimitiveType = PT_TriangleList;
+                GraphicsPSOInit.BoundShaderState.VertexDeclarationRHI = GetVertexDeclarationFVector4();
+                GraphicsPSOInit.BoundShaderState.VertexShaderRHI = GETSAFERHISHADER_VERTEX(*VertexShader);
+                GraphicsPSOInit.BoundShaderState.PixelShaderRHI = GETSAFERHISHADER_PIXEL(*PixelShader);
+                SetGraphicsPipelineState(RHICmdList, GraphicsPSOInit);
+
+                // Update viewport.
+                RHICmdList.SetViewport(
+                    0, 0, 0.f,
+                    OutTextureRenderTargetResource->GetSizeX(), OutTextureRenderTargetResource->GetSizeY(), 1.f);
+
+                // Update shader uniform parameters.
+                VertexShader->SetParameters(RHICmdList, VertexShader->GetVertexShader(), CompiledCameraModel, DisplacementMapResolution);
+                PixelShader->SetParameters(RHICmdList, PixelShader->GetPixelShader(), CompiledCameraModel, DisplacementMapResolution);
+
+                // Draw grid.
+                uint32 PrimitiveCount = kGridSubdivisionX * kGridSubdivisionY * 2;
+                RHICmdList.DrawPrimitive(PT_TriangleList, 0, PrimitiveCount, 1);
+
+                // Resolve render target.
+                RHICmdList.CopyToResolveTarget(
+                    OutTextureRenderTargetResource->GetRenderTargetTexture(),
+                    OutTextureRenderTargetResource->TextureRHI,
+                    false, FResolveParams());
+            }
+
+            FVector2D FFooCameraModel::UndistortNormalizedViewPosition(FVector2D EngineV) const
+            {
+                // Engine view space -> standard view space.
+                FVector2D V = FVector2D(1, -1) * EngineV;
+
+                FVector2D V2 = V * V;
+                float R2 = V2.X + V2.Y;
+
+                // Radial distortion (extra parenthesis to match MF_Undistortion.uasset).
+                FVector2D UndistortedV = V * (1.0 + (R2 * K1 + (R2 * R2) * K2 + (R2 * R2 * R2) * K3));
+
+                // Tangential distortion.
+                UndistortedV.X += P2 * (R2 + 2 * V2.X) + 2 * P1 * V.X * V.Y;
+                UndistortedV.Y += P1 * (R2 + 2 * V2.Y) + 2 * P2 * V.X * V.Y;
+
+                // Returns engine V.
+                return UndistortedV * FVector2D(1, -1);
+            }
+
+            /** Compiles the camera model. */
+            float FFooCameraModel::GetUndistortOverscanFactor(
+                float DistortedHorizontalFOV, float DistortedAspectRatio) const
+            {
+                // If the lens distortion model is identity, then early return 1.
+                if (*this == FFooCameraModel())
+                {
+                    return 1.0f;
+                }
+
+                float TanHalfDistortedHorizontalFOV = FMath::Tan(DistortedHorizontalFOV * 0.5f);
+
+                // Get the position in the view space at z'=1 of different key point in the distorted Viewport UV coordinate system.
+                // This very approximative to know the required overscan scale factor of the undistorted viewport, but works really well in practice.
+                //
+                //  Undistorted UV position in the view space:
+                //                 ^ View space's Y
+                //                 |
+                //        0        1        2
+                //     
+                //        7        0        3 --> View space's X
+                //     
+                //        6        5        4
+                FVector2D UndistortCornerPos0 = LensUndistortViewportUVIntoViewSpace(
+                    *this, TanHalfDistortedHorizontalFOV, DistortedAspectRatio, FVector2D(0.0f, 0.0f));
+                FVector2D UndistortCornerPos1 = LensUndistortViewportUVIntoViewSpace(
+                    *this, TanHalfDistortedHorizontalFOV, DistortedAspectRatio, FVector2D(0.5f, 0.0f));
+                FVector2D UndistortCornerPos2 = LensUndistortViewportUVIntoViewSpace(
+                    *this, TanHalfDistortedHorizontalFOV, DistortedAspectRatio, FVector2D(1.0f, 0.0f));
+                FVector2D UndistortCornerPos3 = LensUndistortViewportUVIntoViewSpace(
+                    *this, TanHalfDistortedHorizontalFOV, DistortedAspectRatio, FVector2D(1.0f, 0.5f));
+                FVector2D UndistortCornerPos4 = LensUndistortViewportUVIntoViewSpace(
+                    *this, TanHalfDistortedHorizontalFOV, DistortedAspectRatio, FVector2D(1.0f, 1.0f));
+                FVector2D UndistortCornerPos5 = LensUndistortViewportUVIntoViewSpace(
+                    *this, TanHalfDistortedHorizontalFOV, DistortedAspectRatio, FVector2D(0.5f, 1.0f));
+                FVector2D UndistortCornerPos6 = LensUndistortViewportUVIntoViewSpace(
+                    *this, TanHalfDistortedHorizontalFOV, DistortedAspectRatio, FVector2D(0.0f, 1.0f));
+                FVector2D UndistortCornerPos7 = LensUndistortViewportUVIntoViewSpace(
+                    *this, TanHalfDistortedHorizontalFOV, DistortedAspectRatio, FVector2D(0.0f, 0.5f));
+
+                // Find min and max of the inner square of undistorted Viewport in the view space at z'=1.
+                FVector2D MinInnerViewportRect;
+                FVector2D MaxInnerViewportRect;
+                MinInnerViewportRect.X = FMath::Max3(UndistortCornerPos0.X, UndistortCornerPos6.X, UndistortCornerPos7.X);
+                MinInnerViewportRect.Y = FMath::Max3(UndistortCornerPos4.Y, UndistortCornerPos5.Y, UndistortCornerPos6.Y);
+                MaxInnerViewportRect.X = FMath::Min3(UndistortCornerPos2.X, UndistortCornerPos3.X, UndistortCornerPos4.X);
+                MaxInnerViewportRect.Y = FMath::Min3(UndistortCornerPos0.Y, UndistortCornerPos1.Y, UndistortCornerPos2.Y);
+
+                check(MinInnerViewportRect.X < 0.f);
+                check(MinInnerViewportRect.Y < 0.f);
+                check(MaxInnerViewportRect.X > 0.f);
+                check(MaxInnerViewportRect.Y > 0.f);
+
+                // Compute tan(VerticalFOV * 0.5)
+                float TanHalfDistortedVerticalFOV = TanHalfDistortedHorizontalFOV / DistortedAspectRatio;
+
+                // Compute the required un-distorted viewport scale on each axis.
+                FVector2D ViewportScaleUpFactorPerViewAxis = 0.5 * FVector2D(
+                    TanHalfDistortedHorizontalFOV / FMath::Max(-MinInnerViewportRect.X, MaxInnerViewportRect.X),
+                    TanHalfDistortedVerticalFOV / FMath::Max(-MinInnerViewportRect.Y, MaxInnerViewportRect.Y));
+
+                // Scale up by 2% more the undistorted viewport size in the view space to work
+                // around the fact that odd undistorted positions might not exactly be at the minimal
+                // in case of a tangential distorted barrel lens distortion.
+                const float ViewportScaleUpConstMultiplier = 1.02f;
+                return FMath::Max(ViewportScaleUpFactorPerViewAxis.X, ViewportScaleUpFactorPerViewAxis.Y) * ViewportScaleUpConstMultiplier;
+            }
+
+            void FFooCameraModel::DrawUVDisplacementToRenderTarget(
+                UWorld* World,
+                float DistortedHorizontalFOV,
+                float DistortedAspectRatio,
+                float UndistortOverscanFactor,
+                UTextureRenderTarget2D* OutputRenderTarget,
+                float OutputMultiply,
+                float OutputAdd) const
+            {
+                check(IsInGameThread());
+
+                // Compiles the camera model to know the overscan scale factor.
+                float TanHalfUndistortedHorizontalFOV = FMath::Tan(DistortedHorizontalFOV * 0.5f) * UndistortOverscanFactor;
+                float TanHalfUndistortedVerticalFOV = TanHalfUndistortedHorizontalFOV / DistortedAspectRatio;
+
+                // Output.
+                FCompiledCameraModel CompiledCameraModel;
+                CompiledCameraModel.OriginalCameraModel = *this;
+
+                CompiledCameraModel.DistortedCameraMatrix.X = 1.0f / TanHalfUndistortedHorizontalFOV;
+                CompiledCameraModel.DistortedCameraMatrix.Y = 1.0f / TanHalfUndistortedVerticalFOV;
+                CompiledCameraModel.DistortedCameraMatrix.Z = 0.5f;
+                CompiledCameraModel.DistortedCameraMatrix.W = 0.5f;
+
+                CompiledCameraModel.UndistortedCameraMatrix.X = F.X;
+                CompiledCameraModel.UndistortedCameraMatrix.Y = F.Y * DistortedAspectRatio;
+                CompiledCameraModel.UndistortedCameraMatrix.Z = C.X;
+                CompiledCameraModel.UndistortedCameraMatrix.W = C.Y;
+
+                CompiledCameraModel.OutputMultiplyAndAdd.X = OutputMultiply;
+                CompiledCameraModel.OutputMultiplyAndAdd.Y = OutputAdd;
+
+                const FName TextureRenderTargetName = OutputRenderTarget->GetFName();
+                FTextureRenderTargetResource* TextureRenderTargetResource = OutputRenderTarget->GameThread_GetRenderTargetResource();
+
+                ERHIFeatureLevel::Type FeatureLevel = World->Scene->GetFeatureLevel();
+
+                ENQUEUE_RENDER_COMMAND(CaptureCommand)(
+                    [CompiledCameraModel, TextureRenderTargetResource, TextureRenderTargetName, FeatureLevel](FRHICommandListImmediate& RHICmdList)
+                    {
+                        DrawUVDisplacementToRenderTarget_RenderThread(
+                            RHICmdList,
+                            CompiledCameraModel,
+                            TextureRenderTargetName,
+                            TextureRenderTargetResource,
+                            FeatureLevel);
+                    }
+                );
+            }
+
+1. Finally, in the **LensDistortionRendering.cpp** file, at lines 155 and 156, change the following two lines of code to point to the new MyShader.USF file that was created earlier.
+    
+    Change:
+
+    * IMPLEMENT_SHADER_TYPE(, FLensDistortionUVGenerationVS, TEXT("/Plugin/LensDistortion/Private/UVGeneration.usf"), TEXT("MainVS"), SF_Vertex)
+
+    To: 
+
+    * IMPLEMENT_SHADER_TYPE(, FLensDistortionUVGenerationVS, TEXT("/Plugin/Foo/Private/MyShader.usf"), TEXT("MainVS"), SF_Vertex)
+
+    Change: 
+
+     * IMPLEMENT_SHADER_TYPE(, FLensDistortionUVGenerationPS, TEXT("/Plugin/LensDistortion/Private/UVGeneration.usf"), TEXT("MainPS"), SF_Pixel)
+
+     To: 
+
+     * IMPLEMENT_SHADER_TYPE(, FLensDistortionUVGenerationPS, TEXT("/Plugin/Foo/Private/MyShader.usf"), TEXT("MainPS"), SF_Pixel)
+
+1. Now, go to the **Foo/Source** folder and open up the **Foo.Build.cs** file and in the **Foo.Build.cs** under the **PublicDependencyModuleNames.AddRange** section add the following lines of code:
+
+            PublicDependencyModuleNames.AddRange(
+            new string[]
+            {
+                "Core",
+                "RenderCore",
+                "ShaderCore",
+                "RHI",
+                // ... add other public dependencies that you statically link with here ...
+            }
+            );
+
+1. Then, in the **PrivateDependencyModuleNames.AddRange** section of **Foo.Build.cs** remove **Slate** and **SlateCore** and when completed your Foo.Build.cs file should look like this:
+
+        // Copyright 1998-2017 Epic Games, Inc. All Rights Reserved.
+
+        using UnrealBuildTool;
+
+        public class Foo : ModuleRules
+        {
+            public Foo(ReadOnlyTargetRules Target) : base(Target)
+            {
+                PCHUsage = ModuleRules.PCHUsageMode.UseExplicitOrSharedPCHs;
+        
+                PublicIncludePaths.AddRange(
+                    new string[] {
+                        "Foo/Public"
+                        // ... add public include paths required here ...
+                    }
+                    );    
+                PrivateIncludePaths.AddRange(
+                    new string[] {
+                        "Foo/Private",
+                        // ... add other private include paths required here ...
+                    }
+                    );
+    
+                PublicDependencyModuleNames.AddRange(
+                    new string[]
+                    {
+                        "Core",
+                        "RenderCore",
+                        "ShaderCore",
+                        "RHI",
+                        // ... add other public dependencies that you statically link with here ...
+                    }
+                    );
+                PrivateDependencyModuleNames.AddRange(
+                    new string[]
+                    {
+                        "CoreUObject",
+                        "Engine",
+                        // ... add private dependencies that you statically link with here ...    
+                    }
+                    );
+                
+                DynamicallyLoadedModuleNames.AddRange(
+                    new string[]
+                    {
+                        // ... add any modules that your module loads dynamically here ...
+                    }
+                    );
+            }
+        }
+
+1. Re-launch your projects Visual Studio solution file and then recompile your project by pressing **CRTL + 5** and when the compile has completed, press **F5** to launch the UE4 Editor. 
+
+1. When the UE4 Editor finishes loading, go to the **Plugins** manager by going to **Edit** > **Plugins**.
+
+    ![](HT_ShadersInPlugins_08.png)(w:250)
+	
+1. Inside of the Plugins manager, scroll down to the bottom and look for the **Project** section and in this section, you should see your Plugin. 
+    
+    ![](HT_ShadersInPlugins_09.png)
+
+    [region:note]
+    If your plug is not enabled, enable it by clicking on the checkmark box next to its name and re-start the UE4 editor. 
+    [/region]
+
+1. To check to make sure everything is there, open up the Level Blueprint and right - click in the Event Graph inputting **Foo** into the search box. Once that is complete, you should now see all of the items that were added to the Foo Camera category. 
+
+    [REGION:lightbox]
+    [![](HT_ShadersInPlugins_10.png)(w:950)](HT_ShadersInPlugins_10.png)
+    [/REGION]
+
+    [REGION:caption]
+    Click for full image.
+    [/REGION]
+
+## End Result
+
+Congratulations, you have now successfully recreated and compiled a new version of the Lense Distortion UE4 plugin. In the next step, you will learn how to invoke this new global shader from Blueprints and also see how it distorts Render Targets.
+
+[Nav]
\ No newline at end of file
