Copying //UE4/Dev-Documentation to //UE4/Samples-Main (Source: //UE4/Dev-Documentation @ 3623665)

diff --git a/Engine/Documentation/Source/Programming/Rendering/ShaderInPlugin/QuickStart/2/QSSIP_02.INT.udn b/Engine/Documentation/Source/Programming/Rendering/ShaderInPlugin/QuickStart/2/QSSIP_02.INT.udn
new file mode 100644
--- /dev/null
+++ b/Engine/Documentation/Source/Programming/Rendering/ShaderInPlugin/QuickStart/2/QSSIP_02.INT.udn
@@ -0,0 +1,333 @@
+Availability: Docs 
+Crumbs: %ROOT%
+Title: Creating a New Global Shader in a Plugin - Blueprint Setup
+Description: Setting up a Blueprint to test the Shader in Plugin code.
+Type:quick start
+Type:multi-step
+SkillLevel:Intermediate
+Version: 4.17
+Parent:Programming/Rendering/ShaderInPlugin/QuickStart
+Order:2
+Tags:Rendering
+Tags:Programming
+Checkpoint:shaderqs
+
+[NAV]
+
+With the Plugin now verified to be working, in the following section you will be shown how to call the global shader using Blueprints.
+
+1. To see how this works, right-click in the **Content Browser** and create a new **Blueprint Class** that has **Actor** as its parent class with the name,  **DrawUVDisToRenderTarget**. 
+
+    ![](HT_ShadersInPlugins_11.png)(w:600)
+
+1. Next, create a new **Render Target** called **RT_00** and a new Material called **MAT_RT**.
+
+    ![](HT_ShadersInPlugins_12.png)(w:600)
+
+1. Open up MAT_RT Material and then add the RT_00 Render Target to the Material Graph, connecting its output to the Base Color input on the Main Material Node, pressing the **Apply** and **Save** button when done.
+
+    [REGION:lightbox]
+    [![](HT_ShadersInPlugins_13.png)(w:600)](HT_ShadersInPlugins_13.png)
+    [/REGION]
+
+    [REGION:caption]
+    Click for full image.
+    [/REGION]
+
+1. Next, open up the DrawUVDisToRenderTarget Blueprint and create/add the following variables and nodes to the Event Graph.
+
+    |Variables | Default Value|
+    |----------|--------------|
+    |K1  | 1.0 |
+    |K2  | 1.0 |
+    |inc  | 1.0 |
+
+    |Nodes | Default Value|
+    |----------|--------------|
+    |W  | N/A |
+    |S  | N/A |
+    |Event Tick  | N/A |
+    |Draw UVDisplacement to Render Target | Input **1** into FOV, Aspect Ratio, and Overscan Factor |
+    |Make FooCameraModel | N/A |
+    |Clear Render Target 2D | N/A |
+
+    ![](HT_ShadersInPlugins_14.png)(w:600)
+
+1. To easily demonstrate how the Draw UVDisplacement to Render Target will warp the supplied render target, you will need to setup your up Blueprint so that the values input into the **K1** and **K2** inputs on the Draw UVDisplacement to Render Target node can be increase or decrease by the press of a button. To accomplish this, setup the nodes in the Event Graph so that they match the following image: 
+    
+    [Include:#DrawDIS]
+
+    [REGION:caption]
+    You can copy and paste this Blueprint code directly into your Blueprint.
+    [/REGION]
+    [region:note]
+    Do not forget to input your Render Target into the **Output Render Target** or you will not see anything happen.
+    [/region]
+
+1. With the required Blueprint logic now setup, make sure to **Compile** and **Save** your Blueprint and then drag your Blueprint from the **Content Browser** into a **level**. When completed, select the Blueprint and under **Input** change the **Auto Receive Input** from **Disabled** to **Player 0**.
+
+    [REGION:lightbox]
+    [![](HT_ShadersInPlugins_16.png)(w:600)](HT_ShadersInPlugins_16.png)
+    [/REGION]
+
+    [REGION:caption]
+    Click for full image.
+    [/REGION]
+
+1. Next, select the floor and duplicate it by pressing **CTRL + W** and then rotate it **-90 degrees** in the **X** axis then drag the Material **MAT_RT** onto it so you can see the effect the of the Draw UVDisplacement to Render Target.
+
+    [REGION:lightbox]
+    [![](HT_ShadersInPlugins_17.png)(w:600)](HT_ShadersInPlugins_17.png)
+    [/REGION]
+
+    [REGION:caption]
+    Click for full image.
+    [/REGION]
+
+1. With everything now setup, press the **Play** button and then press the **W** and **S** keys to warp the Render Target that was input into the Draw UVDisplacement to Render Target node like in the following video.
+
+	[OBJECT:EmbeddedVideo]
+		[PARAMLITERAL:videoid]
+		VCMGiSTM3-s
+		[/PARAMLITERAL]
+		[PARAMLITERAL:width]
+		100
+		[/PARAMLITERAL]
+		[PARAMLITERAL:height]
+		56.25
+		[/PARAMLITERAL]
+		[PARAMLITERAL:units]
+		%
+		[/PARAMLITERAL]
+		[PARAMLITERAL:type]
+		simple
+		[/PARAMLITERAL]
+	[/OBJECT]
+
+## End Result
+
+At this point you can now distort the content inside a render target from Blueprint to create all kinds of interesting shapes. Next, try and create a new Plugin or edit the existing one and see what kind of interesting images and effects you can create.
+
+[NAV]
+
+	<!--
+[EXCERPT:DrawDIS]	
+[OBJECT:Snippet]
+	[PARAMLITERAL:id]
+		DrawDIS
+	[/PARAMLITERAL]
+	[PARAMLITERAL:type]
+		blueprint
+	[/PARAMLITERAL]
+	[PARAM:image]
+		[REGION:lightbox]
+		[![](HT_ShadersInPlugins_15.png)(w:750)](HT_ShadersInPlugins_15.png)
+		[/REGION]
+	[/PARAM]
+	[PARAM:snippet]
+			BEGIN OBJECT
+				Begin Object Class=/Script/BlueprintGraph.K2Node_CallFunction Name="K2Node_CallFunction_415"
+			   FunctionReference=(MemberParent=Class'/Script/Foo.FooBlueprintLibrary',MemberName="DrawUVDisplacementToRenderTarget")
+			   NodePosX=-320
+			   NodePosY=320
+			   NodeGuid=6A47664B4ECC988FD0BC3083A247DEE4
+			   CustomProperties Pin (PinId=37A50945406550EE210602B9EDA46A00,PinName="execute",PinToolTip="\nExec",PinType.PinCategory="exec",PinType.PinSubCategory="",PinType.PinSubCategoryObject=None,PinType.PinSubCategoryMemberReference=(),PinType.PinValueType=(),PinType.ContainerType=None,PinType.bIsMap=False,PinType.bIsSet=False,PinType.bIsArray=False,PinType.bIsReference=False,PinType.bIsConst=False,PinType.bIsWeakPointer=False,LinkedTo=(K2Node_CallFunction_416 90ED15F94550FFA4BBC36CA3442B9E34,),PersistentGuid=00000000000000000000000000000000,bHidden=False,bNotConnectable=False,bDefaultValueIsReadOnly=False,bDefaultValueIsIgnored=False,bAdvancedView=False,bOrphanedPin=False,)
+			   CustomProperties Pin (PinId=A229BA4C4B3AB1E96CD8FCB7337B7E81,PinName="then",PinToolTip="\nExec",Direction="EGPD_Output",PinType.PinCategory="exec",PinType.PinSubCategory="",PinType.PinSubCategoryObject=None,PinType.PinSubCategoryMemberReference=(),PinType.PinValueType=(),PinType.ContainerType=None,PinType.bIsMap=False,PinType.bIsSet=False,PinType.bIsArray=False,PinType.bIsReference=False,PinType.bIsConst=False,PinType.bIsWeakPointer=False,PersistentGuid=00000000000000000000000000000000,bHidden=False,bNotConnectable=False,bDefaultValueIsReadOnly=False,bDefaultValueIsIgnored=False,bAdvancedView=False,bOrphanedPin=False,)
+			   CustomProperties Pin (PinId=074FA9FC4524DDA34DE5FD9182A32C96,PinName="self",PinFriendlyName=NSLOCTEXT("K2Node", "Target", "Target"),PinToolTip="Target\nFoo Blueprint Library Object Reference",PinType.PinCategory="object",PinType.PinSubCategory="",PinType.PinSubCategoryObject=Class'/Script/Foo.FooBlueprintLibrary',PinType.PinSubCategoryMemberReference=(),PinType.PinValueType=(),PinType.ContainerType=None,PinType.bIsMap=False,PinType.bIsSet=False,PinType.bIsArray=False,PinType.bIsReference=False,PinType.bIsConst=False,PinType.bIsWeakPointer=False,DefaultObject="/Script/Foo.Default__FooBlueprintLibrary",PersistentGuid=00000000000000000000000000000000,bHidden=True,bNotConnectable=False,bDefaultValueIsReadOnly=False,bDefaultValueIsIgnored=False,bAdvancedView=False,bOrphanedPin=False,)
+			   CustomProperties Pin (PinId=E7AD9BC443D9AE707AB792B6E1B4C883,PinName="WorldContextObject",PinToolTip="World Context Object\nObject Object Reference",PinType.PinCategory="object",PinType.PinSubCategory="",PinType.PinSubCategoryObject=Class'/Script/CoreUObject.Object',PinType.PinSubCategoryMemberReference=(),PinType.PinValueType=(),PinType.ContainerType=None,PinType.bIsMap=False,PinType.bIsSet=False,PinType.bIsArray=False,PinType.bIsReference=False,PinType.bIsConst=False,PinType.bIsWeakPointer=False,PersistentGuid=00000000000000000000000000000000,bHidden=True,bNotConnectable=False,bDefaultValueIsReadOnly=False,bDefaultValueIsIgnored=False,bAdvancedView=False,bOrphanedPin=False,)
+			   CustomProperties Pin (PinId=F2AB8E6441774533823DD2A413DEF544,PinName="CameraModel",PinToolTip="Camera Model\nFoo Camera Model Structure (by ref)",PinType.PinCategory="struct",PinType.PinSubCategory="",PinType.PinSubCategoryObject=ScriptStruct'/Script/Foo.FooCameraModel',PinType.PinSubCategoryMemberReference=(),PinType.PinValueType=(),PinType.ContainerType=None,PinType.bIsMap=False,PinType.bIsSet=False,PinType.bIsArray=False,PinType.bIsReference=True,PinType.bIsConst=False,PinType.bIsWeakPointer=False,LinkedTo=(K2Node_MakeStruct_96 9EB96E474F8D2725864D30B6ECDFF41F,),PersistentGuid=00000000000000000000000000000000,bHidden=False,bNotConnectable=False,bDefaultValueIsReadOnly=False,bDefaultValueIsIgnored=True,bAdvancedView=False,bOrphanedPin=False,)
+			   CustomProperties Pin (PinId=997B44ED47F8AFC32E1E3D8D076096B4,PinName="DistortedHorizontalFOV",PinToolTip="Distorted Horizontal FOV\nFloat\n\nThe desired horizontal FOV in the distorted render.",PinType.PinCategory="float",PinType.PinSubCategory="",PinType.PinSubCategoryObject=None,PinType.PinSubCategoryMemberReference=(),PinType.PinValueType=(),PinType.ContainerType=None,PinType.bIsMap=False,PinType.bIsSet=False,PinType.bIsArray=False,PinType.bIsReference=False,PinType.bIsConst=False,PinType.bIsWeakPointer=False,DefaultValue="1",AutogeneratedDefaultValue="0.0",PersistentGuid=00000000000000000000000000000000,bHidden=False,bNotConnectable=False,bDefaultValueIsReadOnly=False,bDefaultValueIsIgnored=False,bAdvancedView=False,bOrphanedPin=False,)
+			   CustomProperties Pin (PinId=7D39B8B24ECD165F87E3AE8A7B13FA35,PinName="DistortedAspectRatio",PinToolTip="Distorted Aspect Ratio\nFloat\n\nThe desired aspect ratio of the distorted render.",PinType.PinCategory="float",PinType.PinSubCategory="",PinType.PinSubCategoryObject=None,PinType.PinSubCategoryMemberReference=(),PinType.PinValueType=(),PinType.ContainerType=None,PinType.bIsMap=False,PinType.bIsSet=False,PinType.bIsArray=False,PinType.bIsReference=False,PinType.bIsConst=False,PinType.bIsWeakPointer=False,DefaultValue="1",AutogeneratedDefaultValue="0.0",PersistentGuid=00000000000000000000000000000000,bHidden=False,bNotConnectable=False,bDefaultValueIsReadOnly=False,bDefaultValueIsIgnored=False,bAdvancedView=False,bOrphanedPin=False,)
+			   CustomProperties Pin (PinId=FBC9F81D4E3E519A67B8DB9F8DF3EE82,PinName="UndistortOverscanFactor",PinToolTip="Undistort Overscan Factor\nFloat\n\nThe factor of the overscan for the undistorted render.",PinType.PinCategory="float",PinType.PinSubCategory="",PinType.PinSubCategoryObject=None,PinType.PinSubCategoryMemberReference=(),PinType.PinValueType=(),PinType.ContainerType=None,PinType.bIsMap=False,PinType.bIsSet=False,PinType.bIsArray=False,PinType.bIsReference=False,PinType.bIsConst=False,PinType.bIsWeakPointer=False,DefaultValue="1",AutogeneratedDefaultValue="0.0",PersistentGuid=00000000000000000000000000000000,bHidden=False,bNotConnectable=False,bDefaultValueIsReadOnly=False,bDefaultValueIsIgnored=False,bAdvancedView=False,bOrphanedPin=False,)
+			   CustomProperties Pin (PinId=F7F38CFD47D34A5FE4DE3A9FEE2F2B82,PinName="OutputRenderTarget",PinToolTip="Output Render Target\nTexture Render Target 2D Object Reference\n\nThe render target to draw to. Don\'t necessarily need to have same resolution or aspect ratio as distorted render.",PinType.PinCategory="object",PinType.PinSubCategory="",PinType.PinSubCategoryObject=Class'/Script/Engine.TextureRenderTarget2D',PinType.PinSubCategoryMemberReference=(),PinType.PinValueType=(),PinType.ContainerType=None,PinType.bIsMap=False,PinType.bIsSet=False,PinType.bIsArray=False,PinType.bIsReference=False,PinType.bIsConst=False,PinType.bIsWeakPointer=False,DefaultObject="/Game/RT_00.RT_00",PersistentGuid=00000000000000000000000000000000,bHidden=False,bNotConnectable=False,bDefaultValueIsReadOnly=False,bDefaultValueIsIgnored=False,bAdvancedView=False,bOrphanedPin=False,)
+			   CustomProperties Pin (PinId=3EA265BD4D6E66170D25E4B760B8C4B9,PinName="OutputMultiply",PinToolTip="Output Multiply\nFloat\n\nThe multiplication factor applied on the displacement.",PinType.PinCategory="float",PinType.PinSubCategory="",PinType.PinSubCategoryObject=None,PinType.PinSubCategoryMemberReference=(),PinType.PinValueType=(),PinType.ContainerType=None,PinType.bIsMap=False,PinType.bIsSet=False,PinType.bIsArray=False,PinType.bIsReference=False,PinType.bIsConst=False,PinType.bIsWeakPointer=False,DefaultValue="0.500000",AutogeneratedDefaultValue="0.500000",PersistentGuid=00000000000000000000000000000000,bHidden=False,bNotConnectable=False,bDefaultValueIsReadOnly=False,bDefaultValueIsIgnored=False,bAdvancedView=False,bOrphanedPin=False,)
+			   CustomProperties Pin (PinId=41D4A4834F5BF73382730A915976A8AA,PinName="OutputAdd",PinToolTip="Output Add\nFloat\n\nValue added to the multiplied displacement before storing into the output render target.",PinType.PinCategory="float",PinType.PinSubCategory="",PinType.PinSubCategoryObject=None,PinType.PinSubCategoryMemberReference=(),PinType.PinValueType=(),PinType.ContainerType=None,PinType.bIsMap=False,PinType.bIsSet=False,PinType.bIsArray=False,PinType.bIsReference=False,PinType.bIsConst=False,PinType.bIsWeakPointer=False,DefaultValue="0.500000",AutogeneratedDefaultValue="0.500000",PersistentGuid=00000000000000000000000000000000,bHidden=False,bNotConnectable=False,bDefaultValueIsReadOnly=False,bDefaultValueIsIgnored=False,bAdvancedView=False,bOrphanedPin=False,)
+			End Object
+			Begin Object Class=/Script/BlueprintGraph.K2Node_MakeStruct Name="K2Node_MakeStruct_96"
+			   bMadeAfterOverridePinRemoval=True
+			   ShowPinForProperties(0)=(PropertyName="K1",PropertyFriendlyName="K1",PropertyTooltip=NSLOCTEXT("", "AAE0E1744CF690D4533F5B8735A92401", "Radial parameter #1."),CategoryName="Lens Distortion|Camera Model",bShowPin=True,bCanToggleVisibility=True,bPropertyIsCustomized=True)
+			   ShowPinForProperties(1)=(PropertyName="K2",PropertyFriendlyName="K2",PropertyTooltip=NSLOCTEXT("", "21EF07AD4C81C729E14C25902E1EC1D4", "Radial parameter #2."),CategoryName="Lens Distortion|Camera Model",bShowPin=True,bCanToggleVisibility=True,bPropertyIsCustomized=True)
+			   ShowPinForProperties(2)=(PropertyName="K3",PropertyFriendlyName="K3",PropertyTooltip=NSLOCTEXT("", "51EE62C949CD22B416EAF4932244E40C", "Radial parameter #3."),CategoryName="Lens Distortion|Camera Model",bShowPin=True,bCanToggleVisibility=True,bPropertyIsCustomized=True)
+			   ShowPinForProperties(3)=(PropertyName="P1",PropertyFriendlyName="P1",PropertyTooltip=NSLOCTEXT("", "739015D3419D714685C9AB810122ECC7", "Tangential parameter #1."),CategoryName="Lens Distortion|Camera Model",bShowPin=True,bCanToggleVisibility=True,bPropertyIsCustomized=True)
+			   ShowPinForProperties(4)=(PropertyName="P2",PropertyFriendlyName="P2",PropertyTooltip=NSLOCTEXT("", "294AD7ED4093BCA7AFBF01B8EC0C9A93", "Tangential parameter #2."),CategoryName="Lens Distortion|Camera Model",bShowPin=True,bCanToggleVisibility=True,bPropertyIsCustomized=True)
+			   ShowPinForProperties(5)=(PropertyName="F",PropertyFriendlyName="F",PropertyTooltip=NSLOCTEXT("", "A99C20144CE2EBB9DD0A62B672BE148D", "Camera matrix\'s Fx and Fy."),CategoryName="Lens Distortion|Camera Model",bShowPin=True,bCanToggleVisibility=True,bPropertyIsCustomized=True)
+			   ShowPinForProperties(6)=(PropertyName="C",PropertyFriendlyName="C",PropertyTooltip=NSLOCTEXT("", "68542CDC4C0F33F2FA893981BFBFAF16", "Camera matrix\'s Cx and Cy."),CategoryName="Lens Distortion|Camera Model",bShowPin=True,bCanToggleVisibility=True,bPropertyIsCustomized=True)
+			   StructType=ScriptStruct'/Script/Foo.FooCameraModel'
+			   NodePosX=-656
+			   NodePosY=384
+			   AdvancedPinDisplay=Shown
+			   NodeGuid=891CB1A843FED8C70EA3B28B103CB264
+			   CustomProperties Pin (PinId=9EB96E474F8D2725864D30B6ECDFF41F,PinName="FooCameraModel",Direction="EGPD_Output",PinType.PinCategory="struct",PinType.PinSubCategory="",PinType.PinSubCategoryObject=ScriptStruct'/Script/Foo.FooCameraModel',PinType.PinSubCategoryMemberReference=(),PinType.PinValueType=(),PinType.ContainerType=None,PinType.bIsMap=False,PinType.bIsSet=False,PinType.bIsArray=False,PinType.bIsReference=False,PinType.bIsConst=False,PinType.bIsWeakPointer=False,LinkedTo=(K2Node_CallFunction_415 F2AB8E6441774533823DD2A413DEF544,),PersistentGuid=00000000000000000000000000000000,bHidden=False,bNotConnectable=False,bDefaultValueIsReadOnly=False,bDefaultValueIsIgnored=False,bAdvancedView=False,bOrphanedPin=False,)
+			   CustomProperties Pin (PinId=2B2D8DA343819F2653E8FD8EDB5CCF2E,PinName="K1",PinFriendlyName=NSLOCTEXT("", "51FD07904E508F7CB13C93BD688F8053", "K1"),PinToolTip="K1\nFloat\n\nRadial parameter #1.",PinType.PinCategory="float",PinType.PinSubCategory="",PinType.PinSubCategoryObject=None,PinType.PinSubCategoryMemberReference=(),PinType.PinValueType=(),PinType.ContainerType=None,PinType.bIsMap=False,PinType.bIsSet=False,PinType.bIsArray=False,PinType.bIsReference=False,PinType.bIsConst=False,PinType.bIsWeakPointer=False,DefaultValue="0.000000",AutogeneratedDefaultValue="0.000000",LinkedTo=(K2Node_VariableGet_576 A5D1350D4684134EB5EDE0B43F3983D6,),PersistentGuid=00000000000000000000000000000000,bHidden=False,bNotConnectable=False,bDefaultValueIsReadOnly=False,bDefaultValueIsIgnored=False,bAdvancedView=False,bOrphanedPin=False,)
+			   CustomProperties Pin (PinId=3704C64244726B58EC95F2A67CFF69A6,PinName="K2",PinFriendlyName=NSLOCTEXT("", "E801FDE14963B939270257AD8B050D63", "K2"),PinToolTip="K2\nFloat\n\nRadial parameter #2.",PinType.PinCategory="float",PinType.PinSubCategory="",PinType.PinSubCategoryObject=None,PinType.PinSubCategoryMemberReference=(),PinType.PinValueType=(),PinType.ContainerType=None,PinType.bIsMap=False,PinType.bIsSet=False,PinType.bIsArray=False,PinType.bIsReference=False,PinType.bIsConst=False,PinType.bIsWeakPointer=False,DefaultValue="0.000000",AutogeneratedDefaultValue="0.000000",LinkedTo=(K2Node_VariableGet_577 F9AAFCFC4510CDA989300C82B66FDD0A,),PersistentGuid=00000000000000000000000000000000,bHidden=False,bNotConnectable=False,bDefaultValueIsReadOnly=False,bDefaultValueIsIgnored=False,bAdvancedView=False,bOrphanedPin=False,)
+			   CustomProperties Pin (PinId=6BBE5F394D044BCFCB0DE791007C3C29,PinName="K3",PinFriendlyName=NSLOCTEXT("", "23B742A84529B786D655BAA474DF093F", "K3"),PinToolTip="K3\nFloat\n\nRadial parameter #3.",PinType.PinCategory="float",PinType.PinSubCategory="",PinType.PinSubCategoryObject=None,PinType.PinSubCategoryMemberReference=(),PinType.PinValueType=(),PinType.ContainerType=None,PinType.bIsMap=False,PinType.bIsSet=False,PinType.bIsArray=False,PinType.bIsReference=False,PinType.bIsConst=False,PinType.bIsWeakPointer=False,DefaultValue="0.000000",AutogeneratedDefaultValue="0.000000",PersistentGuid=00000000000000000000000000000000,bHidden=False,bNotConnectable=False,bDefaultValueIsReadOnly=False,bDefaultValueIsIgnored=False,bAdvancedView=True,bOrphanedPin=False,)
+			   CustomProperties Pin (PinId=750F705D4ABB4FEC5A0E2F94FAFC71D2,PinName="P1",PinFriendlyName=NSLOCTEXT("", "12DC4DF74BEABC4D838FC1B6B6AE31D9", "P1"),PinToolTip="P1\nFloat\n\nTangential parameter #1.",PinType.PinCategory="float",PinType.PinSubCategory="",PinType.PinSubCategoryObject=None,PinType.PinSubCategoryMemberReference=(),PinType.PinValueType=(),PinType.ContainerType=None,PinType.bIsMap=False,PinType.bIsSet=False,PinType.bIsArray=False,PinType.bIsReference=False,PinType.bIsConst=False,PinType.bIsWeakPointer=False,DefaultValue="0.000000",AutogeneratedDefaultValue="0.000000",PersistentGuid=00000000000000000000000000000000,bHidden=False,bNotConnectable=False,bDefaultValueIsReadOnly=False,bDefaultValueIsIgnored=False,bAdvancedView=True,bOrphanedPin=False,)
+			   CustomProperties Pin (PinId=0CC9318C424EDB349C2708B0FC278312,PinName="P2",PinFriendlyName=NSLOCTEXT("", "F697F2D046F8BCE75C362E907893C2CE", "P2"),PinToolTip="P2\nFloat\n\nTangential parameter #2.",PinType.PinCategory="float",PinType.PinSubCategory="",PinType.PinSubCategoryObject=None,PinType.PinSubCategoryMemberReference=(),PinType.PinValueType=(),PinType.ContainerType=None,PinType.bIsMap=False,PinType.bIsSet=False,PinType.bIsArray=False,PinType.bIsReference=False,PinType.bIsConst=False,PinType.bIsWeakPointer=False,DefaultValue="0.000000",AutogeneratedDefaultValue="0.000000",PersistentGuid=00000000000000000000000000000000,bHidden=False,bNotConnectable=False,bDefaultValueIsReadOnly=False,bDefaultValueIsIgnored=False,bAdvancedView=True,bOrphanedPin=False,)
+			   CustomProperties Pin (PinId=C1ED9647469C1C9EEB589195C808CA4F,PinName="F",PinFriendlyName=NSLOCTEXT("", "87C4A2A842A30F6626475C9B1A20B3C1", "F"),PinToolTip="F\nVector 2D Structure\n\nCamera matrix\'s Fx and Fy.",PinType.PinCategory="struct",PinType.PinSubCategory="",PinType.PinSubCategoryObject=ScriptStruct'/Script/CoreUObject.Vector2D',PinType.PinSubCategoryMemberReference=(),PinType.PinValueType=(),PinType.ContainerType=None,PinType.bIsMap=False,PinType.bIsSet=False,PinType.bIsArray=False,PinType.bIsReference=False,PinType.bIsConst=False,PinType.bIsWeakPointer=False,DefaultValue="(X=1.000000,Y=1.000000)",AutogeneratedDefaultValue="(X=1.000000,Y=1.000000)",PersistentGuid=00000000000000000000000000000000,bHidden=False,bNotConnectable=False,bDefaultValueIsReadOnly=False,bDefaultValueIsIgnored=False,bAdvancedView=True,bOrphanedPin=False,)
+			   CustomProperties Pin (PinId=D3F2C462415328548B3A7DB01661C58B,PinName="C",PinFriendlyName=NSLOCTEXT("", "0D0E5404467E73DD95C63F894611758E", "C"),PinToolTip="C\nVector 2D Structure\n\nCamera matrix\'s Cx and Cy.",PinType.PinCategory="struct",PinType.PinSubCategory="",PinType.PinSubCategoryObject=ScriptStruct'/Script/CoreUObject.Vector2D',PinType.PinSubCategoryMemberReference=(),PinType.PinValueType=(),PinType.ContainerType=None,PinType.bIsMap=False,PinType.bIsSet=False,PinType.bIsArray=False,PinType.bIsReference=False,PinType.bIsConst=False,PinType.bIsWeakPointer=False,DefaultValue="(X=0.500000,Y=0.500000)",AutogeneratedDefaultValue="(X=0.500000,Y=0.500000)",PersistentGuid=00000000000000000000000000000000,bHidden=False,bNotConnectable=False,bDefaultValueIsReadOnly=False,bDefaultValueIsIgnored=False,bAdvancedView=True,bOrphanedPin=False,)
+			End Object
+			Begin Object Class=/Script/BlueprintGraph.K2Node_VariableGet Name="K2Node_VariableGet_576"
+			   VariableReference=(MemberName="K1",MemberGuid=4839CC3849FA034E4200C5930E5527BE,bSelfContext=True)
+			   NodePosX=-880
+			   NodePosY=528
+			   NodeGuid=9D923D8B485BDBC67AE6C49857387A19
+			   CustomProperties Pin (PinId=A5D1350D4684134EB5EDE0B43F3983D6,PinName="K1",Direction="EGPD_Output",PinType.PinCategory="float",PinType.PinSubCategory="",PinType.PinSubCategoryObject=None,PinType.PinSubCategoryMemberReference=(),PinType.PinValueType=(),PinType.ContainerType=None,PinType.bIsMap=False,PinType.bIsSet=False,PinType.bIsArray=False,PinType.bIsReference=False,PinType.bIsConst=False,PinType.bIsWeakPointer=False,DefaultValue="0.0",AutogeneratedDefaultValue="0.0",LinkedTo=(K2Node_MakeStruct_96 2B2D8DA343819F2653E8FD8EDB5CCF2E,),PersistentGuid=00000000000000000000000000000000,bHidden=False,bNotConnectable=False,bDefaultValueIsReadOnly=False,bDefaultValueIsIgnored=False,bAdvancedView=False,bOrphanedPin=False,)
+			   CustomProperties Pin (PinId=855757D34C3ACB48D139749625D425A1,PinName="self",PinFriendlyName=NSLOCTEXT("K2Node", "Target", "Target"),PinType.PinCategory="object",PinType.PinSubCategory="",PinType.PinSubCategoryObject=BlueprintGeneratedClass'/Game/DrawUVDisToRenderTarget.DrawUVDisToRenderTarget_C',PinType.PinSubCategoryMemberReference=(),PinType.PinValueType=(),PinType.ContainerType=None,PinType.bIsMap=False,PinType.bIsSet=False,PinType.bIsArray=False,PinType.bIsReference=False,PinType.bIsConst=False,PinType.bIsWeakPointer=False,PersistentGuid=00000000000000000000000000000000,bHidden=True,bNotConnectable=False,bDefaultValueIsReadOnly=False,bDefaultValueIsIgnored=False,bAdvancedView=False,bOrphanedPin=False,)
+			End Object
+			Begin Object Class=/Script/BlueprintGraph.K2Node_VariableGet Name="K2Node_VariableGet_577"
+			   VariableReference=(MemberName="K2",MemberGuid=D895BC2A4F39BE80B8EAD096E8B7700E,bSelfContext=True)
+			   NodePosX=-880
+			   NodePosY=576
+			   NodeGuid=EF3CEEFE4A1744607E5ADFAAFC05FD12
+			   CustomProperties Pin (PinId=F9AAFCFC4510CDA989300C82B66FDD0A,PinName="K2",Direction="EGPD_Output",PinType.PinCategory="float",PinType.PinSubCategory="",PinType.PinSubCategoryObject=None,PinType.PinSubCategoryMemberReference=(),PinType.PinValueType=(),PinType.ContainerType=None,PinType.bIsMap=False,PinType.bIsSet=False,PinType.bIsArray=False,PinType.bIsReference=False,PinType.bIsConst=False,PinType.bIsWeakPointer=False,DefaultValue="0.0",AutogeneratedDefaultValue="0.0",LinkedTo=(K2Node_MakeStruct_96 3704C64244726B58EC95F2A67CFF69A6,),PersistentGuid=00000000000000000000000000000000,bHidden=False,bNotConnectable=False,bDefaultValueIsReadOnly=False,bDefaultValueIsIgnored=False,bAdvancedView=False,bOrphanedPin=False,)
+			   CustomProperties Pin (PinId=21CD4CBD4F198D2481CD28BDE8611989,PinName="self",PinFriendlyName=NSLOCTEXT("K2Node", "Target", "Target"),PinType.PinCategory="object",PinType.PinSubCategory="",PinType.PinSubCategoryObject=BlueprintGeneratedClass'/Game/DrawUVDisToRenderTarget.DrawUVDisToRenderTarget_C',PinType.PinSubCategoryMemberReference=(),PinType.PinValueType=(),PinType.ContainerType=None,PinType.bIsMap=False,PinType.bIsSet=False,PinType.bIsArray=False,PinType.bIsReference=False,PinType.bIsConst=False,PinType.bIsWeakPointer=False,PersistentGuid=00000000000000000000000000000000,bHidden=True,bNotConnectable=False,bDefaultValueIsReadOnly=False,bDefaultValueIsIgnored=False,bAdvancedView=False,bOrphanedPin=False,)
+			End Object
+			Begin Object Class=/Script/BlueprintGraph.K2Node_Event Name="K2Node_Event_192"
+			   EventReference=(MemberParent=Class'/Script/Engine.Actor',MemberName="ReceiveTick")
+			   bOverrideFunction=True
+			   NodePosX=-1024
+			   NodePosY=320
+			   NodeGuid=B7AA8C13412B9A161A352B92CA6A517B
+			   CustomProperties Pin (PinId=0946FFD548ACE53983F881B5777ABC2F,PinName="OutputDelegate",Direction="EGPD_Output",PinType.PinCategory="delegate",PinType.PinSubCategory="",PinType.PinSubCategoryObject=None,PinType.PinSubCategoryMemberReference=(MemberParent=Class'/Script/Engine.Actor',MemberName="ReceiveTick"),PinType.PinValueType=(),PinType.ContainerType=None,PinType.bIsMap=False,PinType.bIsSet=False,PinType.bIsArray=False,PinType.bIsReference=False,PinType.bIsConst=False,PinType.bIsWeakPointer=False,PersistentGuid=00000000000000000000000000000000,bHidden=False,bNotConnectable=False,bDefaultValueIsReadOnly=False,bDefaultValueIsIgnored=False,bAdvancedView=False,bOrphanedPin=False,)
+			   CustomProperties Pin (PinId=CF3109B747079A44B53B3DB5B51020FC,PinName="then",Direction="EGPD_Output",PinType.PinCategory="exec",PinType.PinSubCategory="",PinType.PinSubCategoryObject=None,PinType.PinSubCategoryMemberReference=(),PinType.PinValueType=(),PinType.ContainerType=None,PinType.bIsMap=False,PinType.bIsSet=False,PinType.bIsArray=False,PinType.bIsReference=False,PinType.bIsConst=False,PinType.bIsWeakPointer=False,LinkedTo=(K2Node_CallFunction_416 79BFE6EF42D85E5A60B44D92A288F50C,),PersistentGuid=00000000000000000000000000000000,bHidden=False,bNotConnectable=False,bDefaultValueIsReadOnly=False,bDefaultValueIsIgnored=False,bAdvancedView=False,bOrphanedPin=False,)
+			   CustomProperties Pin (PinId=42C21E7D47A3E8C35F8585A663FFBFB5,PinName="DeltaSeconds",PinToolTip="Delta Seconds\nFloat",Direction="EGPD_Output",PinType.PinCategory="float",PinType.PinSubCategory="",PinType.PinSubCategoryObject=None,PinType.PinSubCategoryMemberReference=(),PinType.PinValueType=(),PinType.ContainerType=None,PinType.bIsMap=False,PinType.bIsSet=False,PinType.bIsArray=False,PinType.bIsReference=False,PinType.bIsConst=False,PinType.bIsWeakPointer=False,DefaultValue="0.0",AutogeneratedDefaultValue="0.0",PersistentGuid=00000000000000000000000000000000,bHidden=False,bNotConnectable=False,bDefaultValueIsReadOnly=False,bDefaultValueIsIgnored=False,bAdvancedView=False,bOrphanedPin=False,)
+			End Object
+			Begin Object Class=/Script/BlueprintGraph.K2Node_InputKey Name="K2Node_InputKey_192"
+			   InputKey=W
+			   NodePosX=-976
+			   NodePosY=-161
+			   NodeGuid=4FBA7A1249F1F885C74D75AA91AC55F2
+			   CustomProperties Pin (PinId=0429EAAA4F3A1525226DAE8C05168835,PinName="Pressed",Direction="EGPD_Output",PinType.PinCategory="exec",PinType.PinSubCategory="",PinType.PinSubCategoryObject=None,PinType.PinSubCategoryMemberReference=(),PinType.PinValueType=(),PinType.ContainerType=None,PinType.bIsMap=False,PinType.bIsSet=False,PinType.bIsArray=False,PinType.bIsReference=False,PinType.bIsConst=False,PinType.bIsWeakPointer=False,LinkedTo=(K2Node_VariableSet_384 34FD5FE0407CEFD46D3AE59501650E22,),PersistentGuid=00000000000000000000000000000000,bHidden=False,bNotConnectable=False,bDefaultValueIsReadOnly=False,bDefaultValueIsIgnored=False,bAdvancedView=False,bOrphanedPin=False,)
+			   CustomProperties Pin (PinId=1AAED7B94E1848D93AC4D1A2C1FF372D,PinName="Released",Direction="EGPD_Output",PinType.PinCategory="exec",PinType.PinSubCategory="",PinType.PinSubCategoryObject=None,PinType.PinSubCategoryMemberReference=(),PinType.PinValueType=(),PinType.ContainerType=None,PinType.bIsMap=False,PinType.bIsSet=False,PinType.bIsArray=False,PinType.bIsReference=False,PinType.bIsConst=False,PinType.bIsWeakPointer=False,PersistentGuid=00000000000000000000000000000000,bHidden=False,bNotConnectable=False,bDefaultValueIsReadOnly=False,bDefaultValueIsIgnored=False,bAdvancedView=False,bOrphanedPin=False,)
+			   CustomProperties Pin (PinId=5E9345974B5A3DDD591CC19CB2AAB597,PinName="Key",Direction="EGPD_Output",PinType.PinCategory="struct",PinType.PinSubCategory="",PinType.PinSubCategoryObject=ScriptStruct'/Script/InputCore.Key',PinType.PinSubCategoryMemberReference=(),PinType.PinValueType=(),PinType.ContainerType=None,PinType.bIsMap=False,PinType.bIsSet=False,PinType.bIsArray=False,PinType.bIsReference=False,PinType.bIsConst=False,PinType.bIsWeakPointer=False,DefaultValue="AnyKey",PersistentGuid=00000000000000000000000000000000,bHidden=False,bNotConnectable=False,bDefaultValueIsReadOnly=False,bDefaultValueIsIgnored=False,bAdvancedView=False,bOrphanedPin=False,)
+			End Object
+			Begin Object Class=/Script/BlueprintGraph.K2Node_VariableGet Name="K2Node_VariableGet_578"
+			   VariableReference=(MemberName="K1",MemberGuid=4839CC3849FA034E4200C5930E5527BE,bSelfContext=True)
+			   NodePosX=-880
+			   NodePosY=-65
+			   NodeGuid=B4638D624BC2F7C9D6FD4588B8F3E398
+			   CustomProperties Pin (PinId=4B8548F545AA42E5DEAA30BB788BDFD5,PinName="K1",Direction="EGPD_Output",PinType.PinCategory="float",PinType.PinSubCategory="",PinType.PinSubCategoryObject=None,PinType.PinSubCategoryMemberReference=(),PinType.PinValueType=(),PinType.ContainerType=None,PinType.bIsMap=False,PinType.bIsSet=False,PinType.bIsArray=False,PinType.bIsReference=False,PinType.bIsConst=False,PinType.bIsWeakPointer=False,DefaultValue="0.0",AutogeneratedDefaultValue="0.0",LinkedTo=(K2Node_CommutativeAssociativeBinaryOperator_288 2881016344EEA036B25198B1A2901731,),PersistentGuid=00000000000000000000000000000000,bHidden=False,bNotConnectable=False,bDefaultValueIsReadOnly=False,bDefaultValueIsIgnored=False,bAdvancedView=False,bOrphanedPin=False,)
+			   CustomProperties Pin (PinId=ED52FD8B45FABD0DB36912A73E4B4A42,PinName="self",PinFriendlyName=NSLOCTEXT("K2Node", "Target", "Target"),PinType.PinCategory="object",PinType.PinSubCategory="",PinType.PinSubCategoryObject=BlueprintGeneratedClass'/Game/DrawUVDisToRenderTarget.DrawUVDisToRenderTarget_C',PinType.PinSubCategoryMemberReference=(),PinType.PinValueType=(),PinType.ContainerType=None,PinType.bIsMap=False,PinType.bIsSet=False,PinType.bIsArray=False,PinType.bIsReference=False,PinType.bIsConst=False,PinType.bIsWeakPointer=False,PersistentGuid=00000000000000000000000000000000,bHidden=True,bNotConnectable=False,bDefaultValueIsReadOnly=False,bDefaultValueIsIgnored=False,bAdvancedView=False,bOrphanedPin=False,)
+			End Object
+			Begin Object Class=/Script/BlueprintGraph.K2Node_CommutativeAssociativeBinaryOperator Name="K2Node_CommutativeAssociativeBinaryOperator_288"
+			   bIsPureFunc=True
+			   FunctionReference=(MemberParent=Class'/Script/Engine.KismetMathLibrary',MemberName="Add_FloatFloat")
+			   NodePosX=-736
+			   NodePosY=-43
+			   NodeGuid=9C35063A4FDEA57DACE141B8111FD85F
+			   CustomProperties Pin (PinId=31DA9D4E432BCD1808FA6E85CB5FAC16,PinName="self",PinFriendlyName=NSLOCTEXT("K2Node", "Target", "Target"),PinType.PinCategory="object",PinType.PinSubCategory="",PinType.PinSubCategoryObject=Class'/Script/Engine.KismetMathLibrary',PinType.PinSubCategoryMemberReference=(),PinType.PinValueType=(),PinType.ContainerType=None,PinType.bIsMap=False,PinType.bIsSet=False,PinType.bIsArray=False,PinType.bIsReference=False,PinType.bIsConst=False,PinType.bIsWeakPointer=False,DefaultObject="/Script/Engine.Default__KismetMathLibrary",PersistentGuid=00000000000000000000000000000000,bHidden=True,bNotConnectable=False,bDefaultValueIsReadOnly=False,bDefaultValueIsIgnored=False,bAdvancedView=False,bOrphanedPin=False,)
+			   CustomProperties Pin (PinId=2881016344EEA036B25198B1A2901731,PinName="A",PinType.PinCategory="float",PinType.PinSubCategory="",PinType.PinSubCategoryObject=None,PinType.PinSubCategoryMemberReference=(),PinType.PinValueType=(),PinType.ContainerType=None,PinType.bIsMap=False,PinType.bIsSet=False,PinType.bIsArray=False,PinType.bIsReference=False,PinType.bIsConst=False,PinType.bIsWeakPointer=False,DefaultValue="0.0",AutogeneratedDefaultValue="0.0",LinkedTo=(K2Node_VariableGet_578 4B8548F545AA42E5DEAA30BB788BDFD5,),PersistentGuid=00000000000000000000000000000000,bHidden=False,bNotConnectable=False,bDefaultValueIsReadOnly=False,bDefaultValueIsIgnored=False,bAdvancedView=False,bOrphanedPin=False,)
+			   CustomProperties Pin (PinId=0F92D08C4BA9B44D30CBD388DA7B164C,PinName="B",PinType.PinCategory="float",PinType.PinSubCategory="",PinType.PinSubCategoryObject=None,PinType.PinSubCategoryMemberReference=(),PinType.PinValueType=(),PinType.ContainerType=None,PinType.bIsMap=False,PinType.bIsSet=False,PinType.bIsArray=False,PinType.bIsReference=False,PinType.bIsConst=False,PinType.bIsWeakPointer=False,DefaultValue="1.000000",AutogeneratedDefaultValue="1.000000",LinkedTo=(K2Node_VariableGet_580 2980C241419284FB7AE273AF1E29C508,),PersistentGuid=00000000000000000000000000000000,bHidden=False,bNotConnectable=False,bDefaultValueIsReadOnly=False,bDefaultValueIsIgnored=False,bAdvancedView=False,bOrphanedPin=False,)
+			   CustomProperties Pin (PinId=16B924E244A4C93A9D11E5B6D0644337,PinName="ReturnValue",Direction="EGPD_Output",PinType.PinCategory="float",PinType.PinSubCategory="",PinType.PinSubCategoryObject=None,PinType.PinSubCategoryMemberReference=(),PinType.PinValueType=(),PinType.ContainerType=None,PinType.bIsMap=False,PinType.bIsSet=False,PinType.bIsArray=False,PinType.bIsReference=False,PinType.bIsConst=False,PinType.bIsWeakPointer=False,DefaultValue="0.0",AutogeneratedDefaultValue="0.0",LinkedTo=(K2Node_VariableSet_384 5F4C474943E3F540C19AB698F5128FD4,),PersistentGuid=00000000000000000000000000000000,bHidden=False,bNotConnectable=False,bDefaultValueIsReadOnly=False,bDefaultValueIsIgnored=False,bAdvancedView=False,bOrphanedPin=False,)
+			End Object
+			Begin Object Class=/Script/BlueprintGraph.K2Node_VariableSet Name="K2Node_VariableSet_384"
+			   VariableReference=(MemberName="K1",MemberGuid=4839CC3849FA034E4200C5930E5527BE,bSelfContext=True)
+			   NodePosX=-560
+			   NodePosY=-80
+			   NodeGuid=193B539C4A92386BB27C14ACC1591AFE
+			   CustomProperties Pin (PinId=34FD5FE0407CEFD46D3AE59501650E22,PinName="execute",PinType.PinCategory="exec",PinType.PinSubCategory="",PinType.PinSubCategoryObject=None,PinType.PinSubCategoryMemberReference=(),PinType.PinValueType=(),PinType.ContainerType=None,PinType.bIsMap=False,PinType.bIsSet=False,PinType.bIsArray=False,PinType.bIsReference=False,PinType.bIsConst=False,PinType.bIsWeakPointer=False,LinkedTo=(K2Node_InputKey_192 0429EAAA4F3A1525226DAE8C05168835,),PersistentGuid=00000000000000000000000000000000,bHidden=False,bNotConnectable=False,bDefaultValueIsReadOnly=False,bDefaultValueIsIgnored=False,bAdvancedView=False,bOrphanedPin=False,)
+			   CustomProperties Pin (PinId=4855AEC64D016016F1DBA0B7F06F6E34,PinName="then",Direction="EGPD_Output",PinType.PinCategory="exec",PinType.PinSubCategory="",PinType.PinSubCategoryObject=None,PinType.PinSubCategoryMemberReference=(),PinType.PinValueType=(),PinType.ContainerType=None,PinType.bIsMap=False,PinType.bIsSet=False,PinType.bIsArray=False,PinType.bIsReference=False,PinType.bIsConst=False,PinType.bIsWeakPointer=False,PersistentGuid=00000000000000000000000000000000,bHidden=False,bNotConnectable=False,bDefaultValueIsReadOnly=False,bDefaultValueIsIgnored=False,bAdvancedView=False,bOrphanedPin=False,)
+			   CustomProperties Pin (PinId=5F4C474943E3F540C19AB698F5128FD4,PinName="K1",PinType.PinCategory="float",PinType.PinSubCategory="",PinType.PinSubCategoryObject=None,PinType.PinSubCategoryMemberReference=(),PinType.PinValueType=(),PinType.ContainerType=None,PinType.bIsMap=False,PinType.bIsSet=False,PinType.bIsArray=False,PinType.bIsReference=False,PinType.bIsConst=False,PinType.bIsWeakPointer=False,DefaultValue="0.0",AutogeneratedDefaultValue="0.0",LinkedTo=(K2Node_CommutativeAssociativeBinaryOperator_288 16B924E244A4C93A9D11E5B6D0644337,),PersistentGuid=00000000000000000000000000000000,bHidden=False,bNotConnectable=False,bDefaultValueIsReadOnly=False,bDefaultValueIsIgnored=False,bAdvancedView=False,bOrphanedPin=False,)
+			   CustomProperties Pin (PinId=B369DA294B0EFE460125B89D37B36016,PinName="Output_Get",PinToolTip="Retrieves the value of the variable, can use instead of a separate Get node",Direction="EGPD_Output",PinType.PinCategory="float",PinType.PinSubCategory="",PinType.PinSubCategoryObject=None,PinType.PinSubCategoryMemberReference=(),PinType.PinValueType=(),PinType.ContainerType=None,PinType.bIsMap=False,PinType.bIsSet=False,PinType.bIsArray=False,PinType.bIsReference=False,PinType.bIsConst=False,PinType.bIsWeakPointer=False,DefaultValue="0.0",AutogeneratedDefaultValue="0.0",PersistentGuid=00000000000000000000000000000000,bHidden=False,bNotConnectable=False,bDefaultValueIsReadOnly=False,bDefaultValueIsIgnored=False,bAdvancedView=False,bOrphanedPin=False,)
+			   CustomProperties Pin (PinId=375B1D804A328D8760AD28AE350A5683,PinName="self",PinFriendlyName=NSLOCTEXT("K2Node", "Target", "Target"),PinType.PinCategory="object",PinType.PinSubCategory="",PinType.PinSubCategoryObject=BlueprintGeneratedClass'/Game/DrawUVDisToRenderTarget.DrawUVDisToRenderTarget_C',PinType.PinSubCategoryMemberReference=(),PinType.PinValueType=(),PinType.ContainerType=None,PinType.bIsMap=False,PinType.bIsSet=False,PinType.bIsArray=False,PinType.bIsReference=False,PinType.bIsConst=False,PinType.bIsWeakPointer=False,PersistentGuid=00000000000000000000000000000000,bHidden=True,bNotConnectable=False,bDefaultValueIsReadOnly=False,bDefaultValueIsIgnored=False,bAdvancedView=False,bOrphanedPin=False,)
+			End Object
+			Begin Object Class=/Script/BlueprintGraph.K2Node_InputKey Name="K2Node_InputKey_193"
+			   InputKey=S
+			   NodePosX=-976
+			   NodePosY=32
+			   NodeGuid=15E31E514BEC9996A60ABB8835A76ECC
+			   CustomProperties Pin (PinId=8E7ECA29433D7EDD8DCDA9AEF8A6CAB3,PinName="Pressed",Direction="EGPD_Output",PinType.PinCategory="exec",PinType.PinSubCategory="",PinType.PinSubCategoryObject=None,PinType.PinSubCategoryMemberReference=(),PinType.PinValueType=(),PinType.ContainerType=None,PinType.bIsMap=False,PinType.bIsSet=False,PinType.bIsArray=False,PinType.bIsReference=False,PinType.bIsConst=False,PinType.bIsWeakPointer=False,LinkedTo=(K2Node_VariableSet_385 34FD5FE0407CEFD46D3AE59501650E22,),PersistentGuid=00000000000000000000000000000000,bHidden=False,bNotConnectable=False,bDefaultValueIsReadOnly=False,bDefaultValueIsIgnored=False,bAdvancedView=False,bOrphanedPin=False,)
+			   CustomProperties Pin (PinId=68CE8604440CA4A0CEACDFAB4483450D,PinName="Released",Direction="EGPD_Output",PinType.PinCategory="exec",PinType.PinSubCategory="",PinType.PinSubCategoryObject=None,PinType.PinSubCategoryMemberReference=(),PinType.PinValueType=(),PinType.ContainerType=None,PinType.bIsMap=False,PinType.bIsSet=False,PinType.bIsArray=False,PinType.bIsReference=False,PinType.bIsConst=False,PinType.bIsWeakPointer=False,PersistentGuid=00000000000000000000000000000000,bHidden=False,bNotConnectable=False,bDefaultValueIsReadOnly=False,bDefaultValueIsIgnored=False,bAdvancedView=False,bOrphanedPin=False,)
+			   CustomProperties Pin (PinId=811D73054ED84FD4651E0D94C848DB85,PinName="Key",Direction="EGPD_Output",PinType.PinCategory="struct",PinType.PinSubCategory="",PinType.PinSubCategoryObject=ScriptStruct'/Script/InputCore.Key',PinType.PinSubCategoryMemberReference=(),PinType.PinValueType=(),PinType.ContainerType=None,PinType.bIsMap=False,PinType.bIsSet=False,PinType.bIsArray=False,PinType.bIsReference=False,PinType.bIsConst=False,PinType.bIsWeakPointer=False,DefaultValue="AnyKey",PersistentGuid=00000000000000000000000000000000,bHidden=False,bNotConnectable=False,bDefaultValueIsReadOnly=False,bDefaultValueIsIgnored=False,bAdvancedView=False,bOrphanedPin=False,)
+			End Object
+			Begin Object Class=/Script/BlueprintGraph.K2Node_VariableGet Name="K2Node_VariableGet_579"
+			   VariableReference=(MemberName="K1",MemberGuid=4839CC3849FA034E4200C5930E5527BE,bSelfContext=True)
+			   NodePosX=-848
+			   NodePosY=117
+			   NodeGuid=FF6F5BDF4D0BFAF03F5310BAEDD12754
+			   CustomProperties Pin (PinId=4B8548F545AA42E5DEAA30BB788BDFD5,PinName="K1",Direction="EGPD_Output",PinType.PinCategory="float",PinType.PinSubCategory="",PinType.PinSubCategoryObject=None,PinType.PinSubCategoryMemberReference=(),PinType.PinValueType=(),PinType.ContainerType=None,PinType.bIsMap=False,PinType.bIsSet=False,PinType.bIsArray=False,PinType.bIsReference=False,PinType.bIsConst=False,PinType.bIsWeakPointer=False,DefaultValue="0.0",AutogeneratedDefaultValue="0.0",LinkedTo=(K2Node_CommutativeAssociativeBinaryOperator_289 2881016344EEA036B25198B1A2901731,),PersistentGuid=00000000000000000000000000000000,bHidden=False,bNotConnectable=False,bDefaultValueIsReadOnly=False,bDefaultValueIsIgnored=False,bAdvancedView=False,bOrphanedPin=False,)
+			   CustomProperties Pin (PinId=ED52FD8B45FABD0DB36912A73E4B4A42,PinName="self",PinFriendlyName=NSLOCTEXT("K2Node", "Target", "Target"),PinType.PinCategory="object",PinType.PinSubCategory="",PinType.PinSubCategoryObject=BlueprintGeneratedClass'/Game/DrawUVDisToRenderTarget.DrawUVDisToRenderTarget_C',PinType.PinSubCategoryMemberReference=(),PinType.PinValueType=(),PinType.ContainerType=None,PinType.bIsMap=False,PinType.bIsSet=False,PinType.bIsArray=False,PinType.bIsReference=False,PinType.bIsConst=False,PinType.bIsWeakPointer=False,PersistentGuid=00000000000000000000000000000000,bHidden=True,bNotConnectable=False,bDefaultValueIsReadOnly=False,bDefaultValueIsIgnored=False,bAdvancedView=False,bOrphanedPin=False,)
+			End Object
+			Begin Object Class=/Script/BlueprintGraph.K2Node_CommutativeAssociativeBinaryOperator Name="K2Node_CommutativeAssociativeBinaryOperator_289"
+			   bIsPureFunc=True
+			   FunctionReference=(MemberParent=Class'/Script/Engine.KismetMathLibrary',MemberName="Add_FloatFloat")
+			   NodePosX=-688
+			   NodePosY=133
+			   NodeGuid=B6EF6E4D40028BE66750DFBFC28D5026
+			   CustomProperties Pin (PinId=31DA9D4E432BCD1808FA6E85CB5FAC16,PinName="self",PinFriendlyName=NSLOCTEXT("K2Node", "Target", "Target"),PinToolTip="Target\nKismet Math Library Object Reference",PinType.PinCategory="object",PinType.PinSubCategory="",PinType.PinSubCategoryObject=Class'/Script/Engine.KismetMathLibrary',PinType.PinSubCategoryMemberReference=(),PinType.PinValueType=(),PinType.ContainerType=None,PinType.bIsMap=False,PinType.bIsSet=False,PinType.bIsArray=False,PinType.bIsReference=False,PinType.bIsConst=False,PinType.bIsWeakPointer=False,DefaultObject="/Script/Engine.Default__KismetMathLibrary",PersistentGuid=00000000000000000000000000000000,bHidden=True,bNotConnectable=False,bDefaultValueIsReadOnly=False,bDefaultValueIsIgnored=False,bAdvancedView=False,bOrphanedPin=False,)
+			   CustomProperties Pin (PinId=2881016344EEA036B25198B1A2901731,PinName="A",PinToolTip="A\nFloat",PinType.PinCategory="float",PinType.PinSubCategory="",PinType.PinSubCategoryObject=None,PinType.PinSubCategoryMemberReference=(),PinType.PinValueType=(),PinType.ContainerType=None,PinType.bIsMap=False,PinType.bIsSet=False,PinType.bIsArray=False,PinType.bIsReference=False,PinType.bIsConst=False,PinType.bIsWeakPointer=False,DefaultValue="0.0",AutogeneratedDefaultValue="0.0",LinkedTo=(K2Node_VariableGet_579 4B8548F545AA42E5DEAA30BB788BDFD5,),PersistentGuid=00000000000000000000000000000000,bHidden=False,bNotConnectable=False,bDefaultValueIsReadOnly=False,bDefaultValueIsIgnored=False,bAdvancedView=False,bOrphanedPin=False,)
+			   CustomProperties Pin (PinId=0F92D08C4BA9B44D30CBD388DA7B164C,PinName="B",PinToolTip="B\nFloat",PinType.PinCategory="float",PinType.PinSubCategory="",PinType.PinSubCategoryObject=None,PinType.PinSubCategoryMemberReference=(),PinType.PinValueType=(),PinType.ContainerType=None,PinType.bIsMap=False,PinType.bIsSet=False,PinType.bIsArray=False,PinType.bIsReference=False,PinType.bIsConst=False,PinType.bIsWeakPointer=False,DefaultValue="1.000000",AutogeneratedDefaultValue="1.000000",LinkedTo=(K2Node_CommutativeAssociativeBinaryOperator_290 783E07954DAB142BA4459BA53C394344,),PersistentGuid=00000000000000000000000000000000,bHidden=False,bNotConnectable=False,bDefaultValueIsReadOnly=False,bDefaultValueIsIgnored=False,bAdvancedView=False,bOrphanedPin=False,)
+			   CustomProperties Pin (PinId=16B924E244A4C93A9D11E5B6D0644337,PinName="ReturnValue",PinToolTip="Return Value\nFloat",Direction="EGPD_Output",PinType.PinCategory="float",PinType.PinSubCategory="",PinType.PinSubCategoryObject=None,PinType.PinSubCategoryMemberReference=(),PinType.PinValueType=(),PinType.ContainerType=None,PinType.bIsMap=False,PinType.bIsSet=False,PinType.bIsArray=False,PinType.bIsReference=False,PinType.bIsConst=False,PinType.bIsWeakPointer=False,DefaultValue="0.0",AutogeneratedDefaultValue="0.0",LinkedTo=(K2Node_VariableSet_385 5F4C474943E3F540C19AB698F5128FD4,),PersistentGuid=00000000000000000000000000000000,bHidden=False,bNotConnectable=False,bDefaultValueIsReadOnly=False,bDefaultValueIsIgnored=False,bAdvancedView=False,bOrphanedPin=False,)
+			End Object
+			Begin Object Class=/Script/BlueprintGraph.K2Node_VariableSet Name="K2Node_VariableSet_385"
+			   VariableReference=(MemberName="K1",MemberGuid=4839CC3849FA034E4200C5930E5527BE,bSelfContext=True)
+			   NodePosX=-512
+			   NodePosY=96
+			   NodeGuid=2D13E4BF4CB4A7FC1523E1AE69AA90A3
+			   CustomProperties Pin (PinId=34FD5FE0407CEFD46D3AE59501650E22,PinName="execute",PinType.PinCategory="exec",PinType.PinSubCategory="",PinType.PinSubCategoryObject=None,PinType.PinSubCategoryMemberReference=(),PinType.PinValueType=(),PinType.ContainerType=None,PinType.bIsMap=False,PinType.bIsSet=False,PinType.bIsArray=False,PinType.bIsReference=False,PinType.bIsConst=False,PinType.bIsWeakPointer=False,LinkedTo=(K2Node_InputKey_193 8E7ECA29433D7EDD8DCDA9AEF8A6CAB3,),PersistentGuid=00000000000000000000000000000000,bHidden=False,bNotConnectable=False,bDefaultValueIsReadOnly=False,bDefaultValueIsIgnored=False,bAdvancedView=False,bOrphanedPin=False,)
+			   CustomProperties Pin (PinId=4855AEC64D016016F1DBA0B7F06F6E34,PinName="then",Direction="EGPD_Output",PinType.PinCategory="exec",PinType.PinSubCategory="",PinType.PinSubCategoryObject=None,PinType.PinSubCategoryMemberReference=(),PinType.PinValueType=(),PinType.ContainerType=None,PinType.bIsMap=False,PinType.bIsSet=False,PinType.bIsArray=False,PinType.bIsReference=False,PinType.bIsConst=False,PinType.bIsWeakPointer=False,PersistentGuid=00000000000000000000000000000000,bHidden=False,bNotConnectable=False,bDefaultValueIsReadOnly=False,bDefaultValueIsIgnored=False,bAdvancedView=False,bOrphanedPin=False,)
+			   CustomProperties Pin (PinId=5F4C474943E3F540C19AB698F5128FD4,PinName="K1",PinType.PinCategory="float",PinType.PinSubCategory="",PinType.PinSubCategoryObject=None,PinType.PinSubCategoryMemberReference=(),PinType.PinValueType=(),PinType.ContainerType=None,PinType.bIsMap=False,PinType.bIsSet=False,PinType.bIsArray=False,PinType.bIsReference=False,PinType.bIsConst=False,PinType.bIsWeakPointer=False,DefaultValue="0.0",AutogeneratedDefaultValue="0.0",LinkedTo=(K2Node_CommutativeAssociativeBinaryOperator_289 16B924E244A4C93A9D11E5B6D0644337,),PersistentGuid=00000000000000000000000000000000,bHidden=False,bNotConnectable=False,bDefaultValueIsReadOnly=False,bDefaultValueIsIgnored=False,bAdvancedView=False,bOrphanedPin=False,)
+			   CustomProperties Pin (PinId=B369DA294B0EFE460125B89D37B36016,PinName="Output_Get",PinToolTip="Retrieves the value of the variable, can use instead of a separate Get node",Direction="EGPD_Output",PinType.PinCategory="float",PinType.PinSubCategory="",PinType.PinSubCategoryObject=None,PinType.PinSubCategoryMemberReference=(),PinType.PinValueType=(),PinType.ContainerType=None,PinType.bIsMap=False,PinType.bIsSet=False,PinType.bIsArray=False,PinType.bIsReference=False,PinType.bIsConst=False,PinType.bIsWeakPointer=False,DefaultValue="0.0",AutogeneratedDefaultValue="0.0",PersistentGuid=00000000000000000000000000000000,bHidden=False,bNotConnectable=False,bDefaultValueIsReadOnly=False,bDefaultValueIsIgnored=False,bAdvancedView=False,bOrphanedPin=False,)
+			   CustomProperties Pin (PinId=375B1D804A328D8760AD28AE350A5683,PinName="self",PinFriendlyName=NSLOCTEXT("K2Node", "Target", "Target"),PinType.PinCategory="object",PinType.PinSubCategory="",PinType.PinSubCategoryObject=BlueprintGeneratedClass'/Game/DrawUVDisToRenderTarget.DrawUVDisToRenderTarget_C',PinType.PinSubCategoryMemberReference=(),PinType.PinValueType=(),PinType.ContainerType=None,PinType.bIsMap=False,PinType.bIsSet=False,PinType.bIsArray=False,PinType.bIsReference=False,PinType.bIsConst=False,PinType.bIsWeakPointer=False,PersistentGuid=00000000000000000000000000000000,bHidden=True,bNotConnectable=False,bDefaultValueIsReadOnly=False,bDefaultValueIsIgnored=False,bAdvancedView=False,bOrphanedPin=False,)
+			End Object
+			Begin Object Class=/Script/BlueprintGraph.K2Node_VariableGet Name="K2Node_VariableGet_580"
+			   VariableReference=(MemberName="inc",MemberGuid=4884B3C94FC8DF3C61FE879F885C0513,bSelfContext=True)
+			   NodePosX=-880
+			   NodePosY=-1
+			   NodeGuid=04CFCC974D09F5C39905CC92FE6A1713
+			   CustomProperties Pin (PinId=2980C241419284FB7AE273AF1E29C508,PinName="inc",Direction="EGPD_Output",PinType.PinCategory="float",PinType.PinSubCategory="",PinType.PinSubCategoryObject=None,PinType.PinSubCategoryMemberReference=(),PinType.PinValueType=(),PinType.ContainerType=None,PinType.bIsMap=False,PinType.bIsSet=False,PinType.bIsArray=False,PinType.bIsReference=False,PinType.bIsConst=False,PinType.bIsWeakPointer=False,DefaultValue="0.0",AutogeneratedDefaultValue="0.0",LinkedTo=(K2Node_CommutativeAssociativeBinaryOperator_288 0F92D08C4BA9B44D30CBD388DA7B164C,),PersistentGuid=00000000000000000000000000000000,bHidden=False,bNotConnectable=False,bDefaultValueIsReadOnly=False,bDefaultValueIsIgnored=False,bAdvancedView=False,bOrphanedPin=False,)
+			   CustomProperties Pin (PinId=F03E9DF14235EE0A66DA7A8AE9CCEC0A,PinName="self",PinFriendlyName=NSLOCTEXT("K2Node", "Target", "Target"),PinType.PinCategory="object",PinType.PinSubCategory="",PinType.PinSubCategoryObject=BlueprintGeneratedClass'/Game/DrawUVDisToRenderTarget.DrawUVDisToRenderTarget_C',PinType.PinSubCategoryMemberReference=(),PinType.PinValueType=(),PinType.ContainerType=None,PinType.bIsMap=False,PinType.bIsSet=False,PinType.bIsArray=False,PinType.bIsReference=False,PinType.bIsConst=False,PinType.bIsWeakPointer=False,PersistentGuid=00000000000000000000000000000000,bHidden=True,bNotConnectable=False,bDefaultValueIsReadOnly=False,bDefaultValueIsIgnored=False,bAdvancedView=False,bOrphanedPin=False,)
+			End Object
+			Begin Object Class=/Script/BlueprintGraph.K2Node_VariableGet Name="K2Node_VariableGet_581"
+			   VariableReference=(MemberName="inc",MemberGuid=4884B3C94FC8DF3C61FE879F885C0513,bSelfContext=True)
+			   NodePosX=-1024
+			   NodePosY=165
+			   NodeGuid=0EB1B27D4A04A1C09B751CA82FBE119E
+			   CustomProperties Pin (PinId=3A16050E455FC2A7B8753280F3B3A49F,PinName="inc",Direction="EGPD_Output",PinType.PinCategory="float",PinType.PinSubCategory="",PinType.PinSubCategoryObject=None,PinType.PinSubCategoryMemberReference=(),PinType.PinValueType=(),PinType.ContainerType=None,PinType.bIsMap=False,PinType.bIsSet=False,PinType.bIsArray=False,PinType.bIsReference=False,PinType.bIsConst=False,PinType.bIsWeakPointer=False,DefaultValue="0.0",AutogeneratedDefaultValue="0.0",LinkedTo=(K2Node_CommutativeAssociativeBinaryOperator_290 A361B639411834B96FE764A2D4906668,),PersistentGuid=00000000000000000000000000000000,bHidden=False,bNotConnectable=False,bDefaultValueIsReadOnly=False,bDefaultValueIsIgnored=False,bAdvancedView=False,bOrphanedPin=False,)
+			   CustomProperties Pin (PinId=B918EB1D428E95863AB5D7B814A3BB9E,PinName="self",PinFriendlyName=NSLOCTEXT("K2Node", "Target", "Target"),PinType.PinCategory="object",PinType.PinSubCategory="",PinType.PinSubCategoryObject=BlueprintGeneratedClass'/Game/DrawUVDisToRenderTarget.DrawUVDisToRenderTarget_C',PinType.PinSubCategoryMemberReference=(),PinType.PinValueType=(),PinType.ContainerType=None,PinType.bIsMap=False,PinType.bIsSet=False,PinType.bIsArray=False,PinType.bIsReference=False,PinType.bIsConst=False,PinType.bIsWeakPointer=False,PersistentGuid=00000000000000000000000000000000,bHidden=True,bNotConnectable=False,bDefaultValueIsReadOnly=False,bDefaultValueIsIgnored=False,bAdvancedView=False,bOrphanedPin=False,)
+			End Object
+			Begin Object Class=/Script/BlueprintGraph.K2Node_CommutativeAssociativeBinaryOperator Name="K2Node_CommutativeAssociativeBinaryOperator_290"
+			   bIsPureFunc=True
+			   FunctionReference=(MemberParent=Class'/Script/Engine.KismetMathLibrary',MemberName="Multiply_FloatFloat")
+			   NodePosX=-848
+			   NodePosY=168
+			   NodeGuid=BDAB4F194B50AE675B4401BB9D6C8C39
+			   CustomProperties Pin (PinId=3FDEE1EC4FC5966307A607902689EAEB,PinName="self",PinFriendlyName=NSLOCTEXT("K2Node", "Target", "Target"),PinType.PinCategory="object",PinType.PinSubCategory="",PinType.PinSubCategoryObject=Class'/Script/Engine.KismetMathLibrary',PinType.PinSubCategoryMemberReference=(),PinType.PinValueType=(),PinType.ContainerType=None,PinType.bIsMap=False,PinType.bIsSet=False,PinType.bIsArray=False,PinType.bIsReference=False,PinType.bIsConst=False,PinType.bIsWeakPointer=False,DefaultObject="/Script/Engine.Default__KismetMathLibrary",PersistentGuid=00000000000000000000000000000000,bHidden=True,bNotConnectable=False,bDefaultValueIsReadOnly=False,bDefaultValueIsIgnored=False,bAdvancedView=False,bOrphanedPin=False,)
+			   CustomProperties Pin (PinId=A361B639411834B96FE764A2D4906668,PinName="A",PinType.PinCategory="float",PinType.PinSubCategory="",PinType.PinSubCategoryObject=None,PinType.PinSubCategoryMemberReference=(),PinType.PinValueType=(),PinType.ContainerType=None,PinType.bIsMap=False,PinType.bIsSet=False,PinType.bIsArray=False,PinType.bIsReference=False,PinType.bIsConst=False,PinType.bIsWeakPointer=False,DefaultValue="0.0",AutogeneratedDefaultValue="0.0",LinkedTo=(K2Node_VariableGet_581 3A16050E455FC2A7B8753280F3B3A49F,),PersistentGuid=00000000000000000000000000000000,bHidden=False,bNotConnectable=False,bDefaultValueIsReadOnly=False,bDefaultValueIsIgnored=False,bAdvancedView=False,bOrphanedPin=False,)
+			   CustomProperties Pin (PinId=BC525A6F442150954DE5E5BEE36BC24A,PinName="B",PinType.PinCategory="float",PinType.PinSubCategory="",PinType.PinSubCategoryObject=None,PinType.PinSubCategoryMemberReference=(),PinType.PinValueType=(),PinType.ContainerType=None,PinType.bIsMap=False,PinType.bIsSet=False,PinType.bIsArray=False,PinType.bIsReference=False,PinType.bIsConst=False,PinType.bIsWeakPointer=False,DefaultValue="-1",AutogeneratedDefaultValue="0.0",PersistentGuid=00000000000000000000000000000000,bHidden=False,bNotConnectable=False,bDefaultValueIsReadOnly=False,bDefaultValueIsIgnored=False,bAdvancedView=False,bOrphanedPin=False,)
+			   CustomProperties Pin (PinId=783E07954DAB142BA4459BA53C394344,PinName="ReturnValue",Direction="EGPD_Output",PinType.PinCategory="float",PinType.PinSubCategory="",PinType.PinSubCategoryObject=None,PinType.PinSubCategoryMemberReference=(),PinType.PinValueType=(),PinType.ContainerType=None,PinType.bIsMap=False,PinType.bIsSet=False,PinType.bIsArray=False,PinType.bIsReference=False,PinType.bIsConst=False,PinType.bIsWeakPointer=False,DefaultValue="0.0",AutogeneratedDefaultValue="0.0",LinkedTo=(K2Node_CommutativeAssociativeBinaryOperator_289 0F92D08C4BA9B44D30CBD388DA7B164C,),PersistentGuid=00000000000000000000000000000000,bHidden=False,bNotConnectable=False,bDefaultValueIsReadOnly=False,bDefaultValueIsIgnored=False,bAdvancedView=False,bOrphanedPin=False,)
+			End Object
+			Begin Object Class=/Script/BlueprintGraph.K2Node_CallFunction Name="K2Node_CallFunction_416"
+			   FunctionReference=(MemberParent=Class'/Script/Engine.KismetRenderingLibrary',MemberName="ClearRenderTarget2D")
+			   NodePosX=-880
+			   NodePosY=320
+			   NodeGuid=FE0330D94D3366BD6864BDB61EE7936C
+			   CustomProperties Pin (PinId=79BFE6EF42D85E5A60B44D92A288F50C,PinName="execute",PinToolTip="\nExec",PinType.PinCategory="exec",PinType.PinSubCategory="",PinType.PinSubCategoryObject=None,PinType.PinSubCategoryMemberReference=(),PinType.PinValueType=(),PinType.ContainerType=None,PinType.bIsMap=False,PinType.bIsSet=False,PinType.bIsArray=False,PinType.bIsReference=False,PinType.bIsConst=False,PinType.bIsWeakPointer=False,LinkedTo=(K2Node_Event_192 CF3109B747079A44B53B3DB5B51020FC,),PersistentGuid=00000000000000000000000000000000,bHidden=False,bNotConnectable=False,bDefaultValueIsReadOnly=False,bDefaultValueIsIgnored=False,bAdvancedView=False,bOrphanedPin=False,)
+			   CustomProperties Pin (PinId=90ED15F94550FFA4BBC36CA3442B9E34,PinName="then",PinToolTip="\nExec",Direction="EGPD_Output",PinType.PinCategory="exec",PinType.PinSubCategory="",PinType.PinSubCategoryObject=None,PinType.PinSubCategoryMemberReference=(),PinType.PinValueType=(),PinType.ContainerType=None,PinType.bIsMap=False,PinType.bIsSet=False,PinType.bIsArray=False,PinType.bIsReference=False,PinType.bIsConst=False,PinType.bIsWeakPointer=False,LinkedTo=(K2Node_CallFunction_415 37A50945406550EE210602B9EDA46A00,),PersistentGuid=00000000000000000000000000000000,bHidden=False,bNotConnectable=False,bDefaultValueIsReadOnly=False,bDefaultValueIsIgnored=False,bAdvancedView=False,bOrphanedPin=False,)
+			   CustomProperties Pin (PinId=701455F543533B934DAD6D9724A33137,PinName="self",PinFriendlyName=NSLOCTEXT("K2Node", "Target", "Target"),PinToolTip="Target\nKismet Rendering Library Object Reference",PinType.PinCategory="object",PinType.PinSubCategory="",PinType.PinSubCategoryObject=Class'/Script/Engine.KismetRenderingLibrary',PinType.PinSubCategoryMemberReference=(),PinType.PinValueType=(),PinType.ContainerType=None,PinType.bIsMap=False,PinType.bIsSet=False,PinType.bIsArray=False,PinType.bIsReference=False,PinType.bIsConst=False,PinType.bIsWeakPointer=False,DefaultObject="/Script/Engine.Default__KismetRenderingLibrary",PersistentGuid=00000000000000000000000000000000,bHidden=True,bNotConnectable=False,bDefaultValueIsReadOnly=False,bDefaultValueIsIgnored=False,bAdvancedView=False,bOrphanedPin=False,)
+			   CustomProperties Pin (PinId=0410B7A34146C336929C98820F4C4F9A,PinName="WorldContextObject",PinToolTip="World Context Object\nObject Object Reference",PinType.PinCategory="object",PinType.PinSubCategory="",PinType.PinSubCategoryObject=Class'/Script/CoreUObject.Object',PinType.PinSubCategoryMemberReference=(),PinType.PinValueType=(),PinType.ContainerType=None,PinType.bIsMap=False,PinType.bIsSet=False,PinType.bIsArray=False,PinType.bIsReference=False,PinType.bIsConst=False,PinType.bIsWeakPointer=False,PersistentGuid=00000000000000000000000000000000,bHidden=True,bNotConnectable=False,bDefaultValueIsReadOnly=False,bDefaultValueIsIgnored=False,bAdvancedView=False,bOrphanedPin=False,)
+			   CustomProperties Pin (PinId=BC869BDB4035291D58A4F5BCF647E7BC,PinName="TextureRenderTarget",PinToolTip="Texture Render Target\nTexture Render Target 2D Object Reference",PinType.PinCategory="object",PinType.PinSubCategory="",PinType.PinSubCategoryObject=Class'/Script/Engine.TextureRenderTarget2D',PinType.PinSubCategoryMemberReference=(),PinType.PinValueType=(),PinType.ContainerType=None,PinType.bIsMap=False,PinType.bIsSet=False,PinType.bIsArray=False,PinType.bIsReference=False,PinType.bIsConst=False,PinType.bIsWeakPointer=False,DefaultObject="/Game/RT_00.RT_00",PersistentGuid=00000000000000000000000000000000,bHidden=False,bNotConnectable=False,bDefaultValueIsReadOnly=False,bDefaultValueIsIgnored=False,bAdvancedView=False,bOrphanedPin=False,)
+			   CustomProperties Pin (PinId=B131BC83400BC3963FBF45A829987EA1,PinName="ClearColor",PinToolTip="Clear Color\nLinear Color Structure",PinType.PinCategory="struct",PinType.PinSubCategory="",PinType.PinSubCategoryObject=ScriptStruct'/Script/CoreUObject.LinearColor',PinType.PinSubCategoryMemberReference=(),PinType.PinValueType=(),PinType.ContainerType=None,PinType.bIsMap=False,PinType.bIsSet=False,PinType.bIsArray=False,PinType.bIsReference=False,PinType.bIsConst=False,PinType.bIsWeakPointer=False,DefaultValue="(R=0.000000,G=0.000000,B=0.000000,A=1.000000)",AutogeneratedDefaultValue="(R=0.000000,G=0.000000,B=0.000000,A=1.000000)",PersistentGuid=00000000000000000000000000000000,bHidden=False,bNotConnectable=False,bDefaultValueIsReadOnly=False,bDefaultValueIsIgnored=False,bAdvancedView=False,bOrphanedPin=False,)
+			End Object
+			Begin Object Class=/Script/UnrealEd.EdGraphNode_Comment Name="EdGraphNode_Comment_5"
+			   NodePosX=-1048
+			   NodePosY=-222
+			   NodeWidth=688
+			   NodeHeight=480
+			   NodeComment="Increase (W) or Decrease (S) the UV Displacement"
+			   NodeGuid=F9DBA6D24A88AC009B27B0A979C22582
+			End Object
+			Begin Object Class=/Script/UnrealEd.EdGraphNode_Comment Name="EdGraphNode_Comment_3"
+			   NodePosX=-1056
+			   NodePosY=272
+			   NodeWidth=1088
+			   NodeHeight=400
+			   NodeComment="Make sure to only draw the most current K1 & K2 input values"
+			   NodeGuid=8A0F32874A5E6677357A34A938594544
+			End Object
+			END OBJECT
+	[/PARAM]
+[/OBJECT]
+[/EXCERPT:DrawDIS]
+
+-->
\ No newline at end of file
diff --git a/Engine/Documentation/Source/Programming/Rendering/ShaderInPlugin/QuickStart/QSShaderInPlugin.INT.udn b/Engine/Documentation/Source/Programming/Rendering/ShaderInPlugin/QuickStart/QSShaderInPlugin.INT.udn
new file mode 100644
--- /dev/null
+++ b/Engine/Documentation/Source/Programming/Rendering/ShaderInPlugin/QuickStart/QSShaderInPlugin.INT.udn
@@ -0,0 +1,49 @@
+Availability: Docs
+Crumbs: %ROOT%
+Title: Create a New Global Shader as a Plugin
+Description: Creating and setting up a new Global shader via a Plugin.
+Type: how-to
+SkillLevel: Advanced
+Version: 4.17
+Parent: Programming/Rendering/ShaderInPlugin
+Order:1
+Tags:Rendering
+Tags:Render Targets
+Tags:Plugin
+Checkpoint:shaderqs
+
+
+The Unreal Engine 4 (UE4) Plugin system enables anyone to add a wide range of new functionality without the need to recompile and distribute the entire engine. In the following How - To we will recreate the Lens Distortion global shader plugin to demonstrate how you would go about implementing a global shader that is controllable via Blueprints.
+
+[REGION:lightbox]
+[![](HT_Hero_Image.png)(w:600)](HT_Hero_Image.png)
+[/REGION]
+
+[REGION:caption]
+Click for full image.
+[/REGION]
+
+[region:note]
+For this example, we will be re-creating the existing Lens Distortion plugin that can be found in **Engine\Plugins\Compositing\LensDistortion** to a new Plugin called **Foo** to demonstrate how you would implement this workflow.
+[/region]
+
+The Shaders in Plugins Quick Start will walk you through creating and using global shader as a Plugin. 
+
+[OBJECT:TopicButtonList]
+	[PARAM:icon]
+		![](%ROOT%/build_icon.png)(convert:false)
+	[/PARAM]
+	[PARAM:title]
+		Steps
+	[/PARAM]
+	[PARAM:description]
+	[/PARAM]
+	[PARAM:links]
+	[DIR(output:"checkpoint" parent:"Programming\Rendering\ShaderInPlugin\QuickStart" org:"hierarchy")]
+	[/PARAM]
+[/OBJECT]
+
+
+[REGION:call_to_action]
+	[Click to Start](Programming\Rendering\ShaderInPlugin\QuickStart\1 "%Programming\Rendering\ShaderInPlugin\QuickStart\1:title%")
+[/REGION]
\ No newline at end of file
diff --git a/Engine/Documentation/Source/Programming/Rendering/ShaderInPlugin/ShaderInPlugIns.INT.udn b/Engine/Documentation/Source/Programming/Rendering/ShaderInPlugin/ShaderInPlugIns.INT.udn
new file mode 100644
--- /dev/null
+++ b/Engine/Documentation/Source/Programming/Rendering/ShaderInPlugin/ShaderInPlugIns.INT.udn
@@ -0,0 +1,22 @@
+Availability: Docs
+Crumbs: %ROOT%
+Title: Shaders In Plugins
+Description: Information on creating and useing shaders in plugins in UE4.
+Type: Landing
+Version: 4.17
+Parent: Programming\Rendering\
+Order: 1
+Tags: Rendering
+
+Add Banner Image
+[REGION:note]
+The Banner Image should show the user what the system, feature, or concept looks like when it's in use.
+[/REGION]
+
+
+The following documents will go over all of the various items that you need to be aware of when trying to create and use global shaders via a Plugin. 
+
+## Topic List
+
+[DIR(end:"1" output:"topic" parent:"Programming\Rendering\ShaderInPlugin" org:"Hierarchy" end:"1")]
+
diff --git a/Engine/Documentation/Source/Programming/Slate/Architecture/SlateArchitecture.INT.udn b/Engine/Documentation/Source/Programming/Slate/Architecture/SlateArchitecture.INT.udn
--- a/Engine/Documentation/Source/Programming/Slate/Architecture/SlateArchitecture.INT.udn
+++ b/Engine/Documentation/Source/Programming/Slate/Architecture/SlateArchitecture.INT.udn
@@ -118,7 +118,7 @@ During the next frame, the SEditableText will read from the Model's data. In the
 ### Attributes and Arguments
 
 Using a delegate is not always desirable. Depending on the use case, the arguments to Slate widgets may need to be constant values or functions.
-We encapsulate this notion via the **TAttribute< T >** class. An attribute can be set to a constant or to a delegate.
+We encapsulate this notion via the `TAttribute< T >` class. An attribute can be set to a constant or to a delegate.
 
 ## Performance Considerations
 
diff --git a/Engine/Documentation/Source/Programming/Slate/ClippingSystem/SlateClippingSystem.INT.udn b/Engine/Documentation/Source/Programming/Slate/ClippingSystem/SlateClippingSystem.INT.udn
new file mode 100644
--- /dev/null
+++ b/Engine/Documentation/Source/Programming/Slate/ClippingSystem/SlateClippingSystem.INT.udn
@@ -0,0 +1,283 @@
+Availability: Public
+Crumbs: %ROOT%
+Title: Slate Clipping System
+Description: An overview of the Clipping System for Slate in Unreal Engine 4.
+Type: Overview
+SkillLevel:Advanced
+Version:4.17
+Parent:Programming/Slate
+Related:Engine/UMG/UserGuide/Clipping
+Order:1
+Tags:Programming
+Tags:Slate
+Tags:UMG UI Designer
+
+[VAR:spacer]![](%ROOT%/spacer.png)(w:150 convert:false)[/VAR]
+[VAR:spacer2]![](%ROOT%/spacer.png)(w:40 convert:false)[/VAR]
+
+
+
+[TOC(start:2 end:3)]
+
+[REGION:warning]
+With the release of Unreal Engine 4.17, the new Slate Clipping System will break backward compatibility for versions 4.16 and earlier that update. See the [Conversion Guide](#conversionguidefrom4stop16andearlierto4stop17andlater) section of this 
+page to update your project. 
+[/REGION]
+
+
+
+
+The Unreal Engine 4 (UE4) uses Slate's clipping system to restrict graphics or text to a region of your screen. This is done by creating a clipping rectangle (or clipping rect) around the 
+objects and panels that are created. This system is used throughout the Editor and the UMG UI Designer tools. Slate enables users to move panels and resize them as needed in the Editor, 
+and the clipping system enables the text and graphics to be confined to their clipping rect. With UMG, this enables better render transforms that clip appropriately for your UI. 
+
+## What's New
+
+Unreal Engine 4.17 overhauled the clipping system to be much more flexible than it has in previous versions. This creates a potential to add more rendering effects for tool systems like [UMG](Engine/UMG/UserGuide)
+to use for your project's UI. The clipping system used in prior versions was limited and only axis-aligned for layout-space which lead to noticeable difficulties with render transforms, generally around edges, like in the 
+example below:
+
+[REGION:imagetable]
+| --- | --- | --- | --- | --- | --- |
+| ![](OldClippingIssue1.png)(w:350) | | | | | ![](OldClippingIssue2.png)(w:350) |
+[/REGION]
+
+This example shows that while the Widget is rotated the clipping rect (white dots) is not, which causes the graphic not to render even though it is within the Canvas panel bounds. 
+
+Below is comprised list of differences between version 1.0 and 2.0 of the Clipping System:
+
+[REGION:simpletable]
+| Clipping 1.0 (4.16 and earlier) | Clipping 2.0 (4.17 and later) |
+| ------------ | ------------ |
+| [INCLUDE:#clipping1] | [INCLUDE:#clipping2] |  
+[/REGION]
+
+Because the new system supports clipping in the render space of the widgets instead of requiring axis-alignment, the system no longer needs to be concerned with clipping issues (left), and instead clips 
+the outermost widget. Now, any transform that is applied to the inner widgets will be correctly clipped at the edges (right).
+
+[REGION:imagetable]
+| ![](ClippingOld.png)(w:350) | | | | | ![](ClippingNew.png)(w:350) |
+| --- | --- | --- | --- | --- | --- |
+| 4.16 and earlier clipping and transforms. ||||| 4.17 and later, clipping applied at edges |
+|  ||||| along with transforms to the children. |
+[/REGION]
+
+For the most part, you shouldn't need to worry about changing the clipping state of a widget. If you were to consider game UI as an example, you wouldn't really need to adjust the widget in UMG and mostly would change this for something like a scroll panel, or in cases, like the Editor, where you can't control the length of the text and need to clip it for resizing the windows or parts of the panel.
+
+![](ClippingExamples.png)(w:500)
+
+Here are a few examples of what you might see with various clipping setups:
+
+* Left - No clipping is enabled on the button or the text.
+* Center - Clipping is enabled on the text
+* Right - Clipping is enabled on the button.
+
+Take note of the difference it makes when clipping is assigned to different objects. Since the system clips to bounds, things like padding do not factor in here, so if the button is set to clip (right), the 
+contents will run all the way to the edge before it is cut off, whereas if the text is responsible for handling the clipping (center), then the text clips based on things like the button content padding that 
+reduces the space the text has to fit into.
+
+In addition to the simple axis-aligned examples of clipping, the system also supports arbitrary clipping quads that are stacked on top of each other. 
+
+![](TransformClipping.png)(w:400)
+[REGION:caption]  [/REGION]
+
+While this example is extreme, being able to combine multiple arbitrary clipping zones is a prerequisite to being able to introduce 3D transforms and projections, since any widget being
+clipped with a projection would need to perform clipping in projected space. 
+
+<!--- 
+
+[EXCERPT:clipping1]
+* Axis-aligned clipping - Only functions in layout space
+* Every Widget gets clipped
+* All clipping is done in the pixel shader (requires six floats per vertex)
+* Allowed batching across clipping rects, since the rects are part of the batch data
+* Render Transforms are "special" in that they allow rendering outside the clipping rect provided by the parent, with side effects when the layout rect is clipped
+[/EXCERPT:clipping1]
+
+
+[EXCERPT:clipping2]
+* A clipping that is an arbitrary quad
+	* Axis-aligned clipping zone using Scissor Rects to clip - which avoids unnecessary pixel shader operations
+	* Complex clipping zones use the stencil buffer to compose an arbitrary stack of clipping zones into the stencil buffer that is then used to clip draw calls
+* Very few widgets clip now (by default), consisting mostly of scroll panels and editable text fields
+	* Render Transforms are no longer "special," so any transform can allow you to render outside of another widget, provided it doesn't clip.
+* Slate can no longer batch across clipping zones. This frees up six floats in the Slate vertex format, and no longer need to clip in the pixel shader
+[/EXCERPT:clipping2]
+
+---->
+
+
+## Conversion Guide from 4.16 and Earlier to 4.17 and Later
+
+The changes to the Slate clipping system in 4.17 will break compatibility with previous versions, so you'll want to use the sections below to update your project to avoid any conversion issues and to 
+learn about the new methods you can use to debug Widgets in your project.
+
+
+### Enabling Clipping of Widgets
+
+For all your UMG Widgets, you can adjust the **Clipping** property located in the **Details** panel of the selected Widget.
+
+![](UMGClippingProperty.png)(w:600)
+
+You can learn more about [Clipping with UMG UI Designer](Engine/UMG/UserGuide/Clipping) here.
+
+To enable clipping in your code, you'll set the clipping property for `EWidgetClipping` to one of the following states:
+
+* Inherit
+* ClipToBounds
+* ClipToBoundsWithoutIntersecting
+* ClipToBoundsAlways
+* OnDemand
+
+Here's a code example: 
+
+	SNew( SBorder )
+	.Clipping(EWidgetClipping::ClipToBounds)
+	[
+		...
+	]
+
+
+### Deprecated APIs
+
+The following APIs have been deprecated as of 4.16. If your project was using these, consider if you would even need to use them now. Not every Widget clips, so you may no longer need any code that was possibly
+tricking the clipping system to allow you to move the clipping zone.
+
+* `FSlateDrawElement::Make___(...)` - there is no longer need to pass a clipping rect into every draw call, so you only need to remove the function call, and you'll automatically be using the new version.
+* `SScissorRectBox` - This is no longer necessary since every axis-aligned clipping rect is now a scissor rect. Make sure to delete these and replace the job it was doing by making its most direct child widget have clipping enabled.
+
+
+### Custom Clipping
+
+There are cases where you may need to clip internally in your widget. For example, a `SProgressBar` may need to clip the progress bar drawing at arbitrary locations based on its progress. To add your clipping, you would do the following in `OnPaint`:
+
+	OutDrawElements.PushClip(FSlateClippingZone(AllottedGeometry));
+
+	//TODO Do your drawing here, or child widget paint calls.
+
+	OutDrawElements.PopClip();
+
+The `FSlateClippingZone` is the arbitrary clipping area in window space that can be initialized using several methods that should allow for a straightforward conversion of your existing code.
+
+If you also need your custom clipping to affect Hit Testing, you will need to push your clipping zone onto the Hit Test Grid by setting up the following:
+
+	OutDrawElements.PushClip(FSlateClippingZone(AllottedGeometry));
+	Args.GetGrid().PushClip(FSlateClippingZone(GetCachedGeometry()));
+
+	//TODO Do your drawing here, or child widget paint calls.
+
+	OutDrawElements.PopClip();
+	Args.GetGrid().PopClip();
+
+[REGION:note]
+The clipping rect for the Hit Test Grid is using cached geometry, instead of the AllotedGeometry. In `OnPaint`, `AllotedGeometry` is in window space, while the Hit Test Grid is in desktop space, so you have
+to use geometry that you'd get in Tick.
+[/REGION]
+
+### Wrapped Clipping States
+
+In some instances, a widget that has a clipping state isn't going to be responsible for clipping. For example, the `SScrollBox` lets you publically change the clipping state like any other widget, 
+however, when the `SScrollBox` was constructed, it set `bClippingProxy` to `true`. This makes it so that when Slate goes to render that widget, it will ignore its clipping state.
+
+Internally, the `SScrollBox` informs another nested widget that it needs to clip or perform whatever clipping state it was told to do. Additionally, when users change the clipping state, the 
+`SScrollBox` overrides the `SWidget` function called `OnClippingChanged` so that it knows when to mirror the new clipping state to the nested private widgets. 
+
+### Culling Changes
+
+Although clipping is performed in render space, the engine still performs culling using a simple bounding box. However, that box is based on the bounding box of the render-transformed clipping zone. Additionally, clipping and culling might become more nuanced over time, so if you have a custom panel and were doing `MyClippingRect.IntersectionWith` to identify and cull widgets that can't possibly be drawn, you should now use the `SWidget` function `IsChildWidgetCulled`. The following is an example of this:
+
+	for (int32 ChildIndex = 0; ChildIndex < ArrangedChildren.Num(); ++ChildIndex)
+	{
+		FArrangedWidget& CurWidget = ArrangedChildren[ChildIndex];
+
+		if (!IsChildWidgetCulled(MyCullingRect, CurWidget))
+		{
+			// Paint this widget.
+		}
+	}
+
+[REGION:warning]
+Any place where you're taking in `MyClippingRect` in your OnPaint call should be renamed to `MyCullingRect`.
+[/REGION]
+
+Even with the changes to the clipping system, Slate's culling approach hasn't changed. In the example below, you'll need to keep in mind that if the parent widget (blue) is culled when outside the clipping 
+zone (green), the child widget (yellow) will also be culled, even if it has a transform that allows it to render entirely outside of the bounds of the parent. This is true even when the clipping 
+property **Clip to Bounds - Without Intersecting** is enabled since the flag would only be checked by the direct parent widget.
+
+
+[OBJECT:ComparisonSlider]
+	[PARAM:before]
+	![Widgets viewed outside | the clipping zone.](Slate_CullingEd.png)(w:600)
+	[/PARAM]
+	[PARAM:after]
+	![Widgets culled outside | the clipping zone.](Slate_CullingEd2.png)(w:600)
+	[/PARAM]
+[/OBJECT]
+
+
+## Debugging Clipping
+
+The **Widget Reflector** indicates the clipping state for selected Widgets inside the Editor.
+
+[REGION:fullwidth]
+![](WidgetReflector.png)
+[/REGION]
+[REGION:caption]A text block has been selected using the Widget Reflector that indicates the Clipping state to be set to On Demand.[/REGION]
+
+To open the **Widget Reflector**, navigate to the file menu and select **Window** > **Developer Tools** > **Widget Reflector**.
+
+Below are the available console variables that can be used to debug clipping and culling states for any Widget in the Editor:
+
+[REGION:simpletable]
+| Console Variable | Description |
+| --- | --- |
+| **Slate.ShowClipping** | When enabled, this shows yellow outlines for all axis-aligned clipping rects (Scissor Rects) and red outlines for all stencil clipping quads. [INCLUDE:#showclipping] |
+| **Slate.DebugCulling** | When enabled, this helps you to better understand how the culling in a panel is working or why it isn't working correctly. It disables clipping on the GPU, but everything continues as normal, so it enables you to see all the things rendering outside the bounds of clip zones and see if they are being culled when you expect them to be. |
+| **Slate.EnableDrawEvents** | When enabled, draw events will be enabled so that when you're debugging with RenderDoc or something similar that it's easier to understand the batches or clipping state transitions. This is enabled by default when using a Debug build. |
+[/REGION]
+
+<!---
+
+[EXCERPT:showclipping]
+[REGION:lightbox]
+[![](ShowClipping.png)(w:200)](ShowClipping.png)
+[/REGION]
+[REGION:caption] Click image for full size view. [/REGION]
+[/EXCERPT:showclipping]
+
+[EXCERPT:debugculling]
+[REGION:lightbox]
+[![](Slate_Culling.png)(w:250)](Slate_Culling.png)
+[/REGION]
+[REGION:caption] Click image for full size view. [/REGION]
+[/EXCERPT:debugculling]
+
+[EXCERPT:drawevents]
+[REGION:lightbox]
+[![](Slate_Culling.png)(w:250)](Slate_Culling.png)
+[/REGION]
+[REGION:caption] Click image for full size view. [/REGION]
+[/EXCERPT:drawevents]--->
+
+
+
+
+
+
+
+
+
+
+
+
+
+
+
+
+
+
+
+
+
+
+
diff --git a/Engine/Documentation/Source/Programming/Slate/Widgets/SlateWidgets.INT.udn b/Engine/Documentation/Source/Programming/Slate/Widgets/SlateWidgets.INT.udn
--- a/Engine/Documentation/Source/Programming/Slate/Widgets/SlateWidgets.INT.udn
+++ b/Engine/Documentation/Source/Programming/Slate/Widgets/SlateWidgets.INT.udn
@@ -264,7 +264,7 @@ These views take a template argument of a shared pointer to some sort of data. T
 
 ### List Views
 
-List views are widgets which store a list of child widgets. They are made as **SListView<...>**.
+List views are widgets which store a list of child widgets. They are made as `SListView<...>`.
 
 	SNew( SListView< TSharedPtr<FTestData> > )
 	.ItemHeight(24)
@@ -293,7 +293,7 @@ List views are widgets which store a list of child widgets. They are made as **S
 
 ### Tree Views
 
-Tree views, made as **STreeView<...>**, are similar to list views, except that they support parenting widgets to other widgets in the list. To determine which widgets are parented to which, the **OnGetChildren** delegate is used to pass back the children of the passed in item.
+Tree views, made as `STreeView<...>`, are similar to list views, except that they support parenting widgets to other widgets in the list. To determine which widgets are parented to which, the **OnGetChildren** delegate is used to pass back the children of the passed in item.
 
 	SNew( STreeView< TSharedPtr<FTestData> > )
 	.ItemHeight(24)
@@ -314,7 +314,7 @@ Tree views, made as **STreeView<...>**, are similar to list views, except that t
 
 ### Tile Views
 
-Tile views, made as **STileView<...>**, are similar to list views, except the child widgets are laid out in a grid rather than a list. Because of this, they do not support columns or column headers.
+Tile views, made as `STileView<...>`, are similar to list views, except the child widgets are laid out in a grid rather than a list. Because of this, they do not support columns or column headers.
 
 	SNew( STileView< TSharedPtr<FTestData> > )
 	.ItemWidth(128)
diff --git a/Engine/Documentation/Source/Programming/Tutorials/Components/3/Components_3.INT.udn b/Engine/Documentation/Source/Programming/Tutorials/Components/3/Components_3.INT.udn
--- a/Engine/Documentation/Source/Programming/Tutorials/Components/3/Components_3.INT.udn
+++ b/Engine/Documentation/Source/Programming/Tutorials/Components/3/Components_3.INT.udn
@@ -55,11 +55,17 @@ version:4.9
 
 	This code will move our **Pawn** smoothly around the world, sliding off of surfaces where appropriate. There is no gravity applied to our Pawn, and its maximum speed is hard-coded to 150 **Unreal Units** per second.
 
-	[REGION:note]This TickComponent function makes use of a few of the powerful features offered by the **UPawnMovementComponent** class.
-	<br>**ConsumeInputVector** reports and clears the value of a built-in variable that we will use to store our movement inputs.
-	<br>**SafeMoveUpdatedComponent** uses **Unreal Engine** physics to move our Pawn Movement Component while respecting solid barriers.
-	<br>**SlideAlongSurface** handles the calculations and physics involved with sliding smoothly along collision surfaces like walls and ramps when a move results in a collision, rather than simply stopping in place and sticking to the wall or ramp.
-	<br>There are more features included in Pawn Movement Components that are worthy of examination, but they are not needed for the scope of this tutorial. Looking at other classes, such as **Floating Pawn Movement**, **Spectator Pawn Movement**, or **Character Movement Component**, could provide additional usage examples and ideas.[/REGION]
+	[REGION:note]
+	This TickComponent function makes use of a few of the powerful features offered by the **UPawnMovementComponent** class.
+	
+	**ConsumeInputVector** reports and clears the value of a built-in variable that we will use to store our movement inputs.
+	
+	**SafeMoveUpdatedComponent** uses **Unreal Engine** physics to move our Pawn Movement Component while respecting solid barriers.
+	
+	**SlideAlongSurface** handles the calculations and physics involved with sliding smoothly along collision surfaces like walls and ramps when a move results in a collision, rather than simply stopping in place and sticking to the wall or ramp.
+	
+	There are more features included in Pawn Movement Components that are worthy of examination, but they are not needed for the scope of this tutorial. Looking at other classes, such as **Floating Pawn Movement**, **Spectator Pawn Movement**, or **Character Movement Component**, could provide additional usage examples and ideas.
+	[/REGION]
 
 ----
 
diff --git a/Engine/Documentation/Source/Programming/Tutorials/FirstPersonShooter/4/4/4/FirstPersonShooter_4_4_4.INT.udn b/Engine/Documentation/Source/Programming/Tutorials/FirstPersonShooter/4/4/4/FirstPersonShooter_4_4_4.INT.udn
--- a/Engine/Documentation/Source/Programming/Tutorials/FirstPersonShooter/4/4/4/FirstPersonShooter_4_4_4.INT.udn
+++ b/Engine/Documentation/Source/Programming/Tutorials/FirstPersonShooter/4/4/4/FirstPersonShooter_4_4_4.INT.udn
@@ -33,9 +33,9 @@ During this step, we're going to set your jump start to jump loop transition.
 
 1.	Right-click in the graph, search for and select the **TimeRemaining for 'FPP_JumpStart'** node.
 
-1.	Drag off the **Time Remaining** output pin and add a **<= (float)** node using the **Context Menu**.
+1.	Drag off the **Time Remaining** output pin and add a `<= (float)` node using the **Context Menu**.
 
-1.	Enter 0.1 in the other output field on the **<= (float)** node and then wire the Boolean outputpin from that node to the **Can Enter Transition** input pin of the **Result** node.
+1.	Enter 0.1 in the other output field on the `<= (float)` node and then wire the Boolean outputpin from that node to the **Can Enter Transition** input pin of the **Result** node.
 
 	![](TimeRemaining.png)
 
diff --git a/Engine/Documentation/Source/Programming/Tutorials/FirstPersonShooter/4/4/6/FirstPersonShooter_4_4_6.INT.udn b/Engine/Documentation/Source/Programming/Tutorials/FirstPersonShooter/4/4/6/FirstPersonShooter_4_4_6.INT.udn
--- a/Engine/Documentation/Source/Programming/Tutorials/FirstPersonShooter/4/4/6/FirstPersonShooter_4_4_6.INT.udn
+++ b/Engine/Documentation/Source/Programming/Tutorials/FirstPersonShooter/4/4/6/FirstPersonShooter_4_4_6.INT.udn
@@ -33,9 +33,9 @@ During this step, we're going to set your jump end to idle transition.
 
 1.	Right-click in the graph, search for and select the **TimeRemaining for 'FPP_JumpEnd'** node.
 
-1.	Drag off the **Time Remaining** output pin and add a **<= (float)** node using the **Context Menu**.
+1.	Drag off the **Time Remaining** output pin and add a `<= (float)` node using the **Context Menu**.
 
-1.	Enter 0.1 in the other output field on the **<= (float)** node and then wire the Boolean outputpin from that node to the **Can Enter Transition** input pin of the **Result** node.
+1.	Enter 0.1 in the other output field on the `<= (float)` node and then wire the Boolean outputpin from that node to the **Can Enter Transition** input pin of the **Result** node.
 
 	![](TimeRemaining.png)
 
diff --git a/Engine/Documentation/Source/Programming/Tutorials/PlayerInput/2/PlayerInput_2.INT.udn b/Engine/Documentation/Source/Programming/Tutorials/PlayerInput/2/PlayerInput_2.INT.udn
--- a/Engine/Documentation/Source/Programming/Tutorials/PlayerInput/2/PlayerInput_2.INT.udn
+++ b/Engine/Documentation/Source/Programming/Tutorials/PlayerInput/2/PlayerInput_2.INT.udn
@@ -23,9 +23,11 @@ version:4.9
 %Steps%
 
 There are two types of input mappings: Action and Axis.
-<br>**Action Mappings** are useful to think of as "yes or no" inputs, like the buttons on a mouse or joystick. They report when they are pressed, released, double-clicked, or held down for a short time. Discrete actions like jumping, shooting, or interacting with objects are good candidates for this type of mapping.
-<br>**Axis Mappings** are continuous - think of them as "how much" inputs, like the stick on a joystick, or the position of a mouse cursor. They report their value every frame, even if they're not moving. Things with magnitude or direction, such as walking, looking around, and steering a vehicle are usually handled this way.
-<br>While input mappings can be defined directly in code, the usual method is to define them in the **Unreal Engine** editor, so that's how we'll do it in this tutorial.
+
+* **Action Mappings** are useful to think of as "yes or no" inputs, like the buttons on a mouse or joystick. They report when they are pressed, released, double-clicked, or held down for a short time. Discrete actions like jumping, shooting, or interacting with objects are good candidates for this type of mapping.
+* **Axis Mappings** are continuous - think of them as "how much" inputs, like the stick on a joystick, or the position of a mouse cursor. They report their value every frame, even if they're not moving. Things with magnitude or direction, such as walking, looking around, and steering a vehicle are usually handled this way.
+
+While input mappings can be defined directly in code, the usual method is to define them in the **Unreal Engine** editor, so that's how we'll do it in this tutorial.
 
 1. In the **Unreal Engine** editor, under the **Edit** dropdown menu, we'll click the **Project Settings** option.
 
diff --git a/Engine/Documentation/Source/Programming/UnrealArchitecture/Actors/ActorLifecycle/ActorLifecycle.INT.udn b/Engine/Documentation/Source/Programming/UnrealArchitecture/Actors/ActorLifecycle/ActorLifecycle.INT.udn
--- a/Engine/Documentation/Source/Programming/UnrealArchitecture/Actors/ActorLifecycle/ActorLifecycle.INT.udn
+++ b/Engine/Documentation/Source/Programming/UnrealArchitecture/Actors/ActorLifecycle/ActorLifecycle.INT.udn
@@ -95,7 +95,7 @@ These are completely optional, as many Actors will not actually die during play.
 * The lifetime of the Actor has expired
 * Application shut down (All Actors are Destroyed)
 
-Regardless of how this happens, the Actor will be marked RF_PendingKill so during the next garbage collection cycle it will be deallocated. Also, rather than checking for pending kill manually, consider using an FWeakObjectPtr<AActor> as it is cleaner.
+Regardless of how this happens, the Actor will be marked RF_PendingKill so during the next garbage collection cycle it will be deallocated. Also, rather than checking for pending kill manually, consider using an `FWeakObjectPtr<AActor>` as it is cleaner.
 
 **OnDestroy** - This is a legacy response to Destroy. You should probably move anything here to EndPlay as it is called by level transition and other game cleanup functions.
 
