Copying //UE4/Dev-Documentation to //UE4/Samples-Main (Source: //UE4/Dev-Documentation @ 3623665)

diff --git a/Engine/Documentation/Source/Programming/UnrealArchitecture/Objects/Optimizations/UnrealObjectHandling.INT.udn b/Engine/Documentation/Source/Programming/UnrealArchitecture/Objects/Optimizations/UnrealObjectHandling.INT.udn
--- a/Engine/Documentation/Source/Programming/UnrealArchitecture/Objects/Optimizations/UnrealObjectHandling.INT.udn
+++ b/Engine/Documentation/Source/Programming/UnrealArchitecture/Objects/Optimizations/UnrealObjectHandling.INT.udn
@@ -1,13 +1,15 @@
 Availability:Public
-Title: Unreal Object Handling
 Crumbs: %ROOT%
+Title: Unreal Object Handling
 Description:Overview of the features of the UObject system.
-type:overview
-skilllevel:Beginner
+Type: Reference
 Version: 4.16
 parent:UnrealArchitecture
 order:
 tags:Architecture
+related: Gameplay/Networking
+related: Programming/UnrealArchitecture/Reference/Properties
+related: Programming/UnrealArchitecture/TArrays
 
 
 [TOC (start:2 end:2)]
@@ -18,8 +20,6 @@ Marking classes, properties, and functions with the appropriate macros turns the
 
 `UObjects` are automatically zeroed on initialization, before the constructor is called. This happens for the whole class, `UProperties` and native members alike. Members can subsequently be initialized with custom values in the class constructor.
 
-[REGION:note]The constructor of a `UObject` will only run when the engine is first started, as a **Class Default Object** (or "CDO") is created. After that, instantiation of Objects (such as spawning Actors) will involve copying the CDO into the new Actor's memory without re-running the constructor.[/REGION]
-
 ## Automatic Updating of References
 
 When an `AActor` or `UActorComponent` is destroyed or otherwise removed from play, all references to it that are visible to the reflection system (`UProperty` pointers and pointers stored in Unreal Engine container classes such as `TArray`) are automatically nulled. This is beneficial in that it prevents dangling pointers from persisting and causing trouble down the road, but it also means that `AActor` and `UActorComponent` pointers can become null if some other piece of code destroys them. The ultimate advantage of this is that null-checking is more reliable, as it detects both standard-case null pointers and cases where a non-null pointer would have been pointing at deleted memory.
@@ -30,25 +30,25 @@ Another case where a referenced UObject UProperty will be automatically null'ed
 
 ## Serialization
 
-When a `UObject` is serialized, all UProperty values are automatically written or read unless explicitly marked as "transient". For example, you could place an `AEnemy` instance in a level, set its Health to 500, save it and successfully reload it without writing a single line of code beyond the `UClass` definition.
+When a `UObject` is serialized, all `UProperty` values are automatically written or read unless explicitly marked as "transient" or unchanged from the post-constructor default value. For example, you could place an `AEnemy` instance in a level, set its Health to 500, save it and successfully reload it without writing a single line of code beyond the `UClass` definition.
 
 When UProperties are added or removed, loading pre-existing content is handled seamlessly. New properties get default values copied from the new CDO. Removed properties are silently ignored.
 
-If you want custom behavior, it is possible to write custom serialization and versioning code, but this is not common in practice.
+If custom behavior is required, the `UObject::Serialize` function can be overridden. This can be useful to detect data errors, check version numbers, or perform automatic conversions or updates if the data format has changed.
 
 ## Updating of Property Values
 
-When the CDO of a `UClass` has changed, the engine will try to apply those changes to all instances of the class when they are loaded. For a given Object instance, if the updated variable's value matches the value in the old CDO, it will be updated to the value it holds in the new CDO. If the variable has any other value, the assumption is that the value was set intentionally, and those changes will be preserved.
+When the **Class Default Object** (or CDO) of a `UClass` has changed, the engine will try to apply those changes to all instances of the class when they are loaded. For a given Object instance, if the updated variable's value matches the value in the old CDO, it will be updated to the value it holds in the new CDO. If the variable has any other value, the assumption is that the value was set intentionally, and those changes will be preserved.
 
 As an example, let us say you saved a level with several of your `AEnemy` Objects placed in it, and you had set the default Health value in the `AEnemy` constructor to 100. Let us also assume you set the health for Enemy_3 to 500, because he is particularly tough. Now imagine you changed your mind and increased the default value of Health to 150. When you next load your level, Unreal will realize you have changed the CDO and will update all instances of `AEnemy` with the old default Health value (100) to have a Health value of 150. Enemy_3's Health will remain at 500 because it wasn't using the old default value.
  
 ## Editor Integration
 
-`UObjects` and `UProperties` are understood by the editor, and the editor can expose these values automatically for editing without the need to write special code. This can optionally include integration into the Blueprint visual scripting system. There are many options to control the accessibility and exposure of variables and functions.
+`UObjects` and `UProperties` are understood by the Editor, and the Editor can expose these values automatically for editing without the need to write special code. This can optionally include integration into the Blueprint visual scripting system. There are many options to control the accessibility and exposure of variables and functions.
 
-## Type Information at Runtime
+## Run-Time Type Information and Casting
 
-`UObjects` always know what `UClass` they are, and type-related decisions can be made at runtime.
+Because `UObjects` are part of the Unreal Engine's reflection system, they always know what `UClass` they are, and type-related decisions and casts can be made at runtime.
 
 In native code, every `UObject` class has a custom `Super` typedef set to its parent class, which allows easy control of overriding behavior. As an example:
 
@@ -72,7 +72,7 @@ In native code, every `UObject` class has a custom `Super` typedef set to its pa
 
 As you can see, calling `Speak` will result in the MegaBoss saying "Powering up! Time to fight!".
 
-Also, you can safely cast an Object from a base class to a more derived class using the templated Cast function, or query if an Object is of a particular class using `IsA`. A quick example:
+Also, you can safely cast an Object from a base class to a more derived class using the templated Cast function, or query if an Object is of a particular class using `IsA`. A quick example follows:
 
 	class ALegendaryWeapon : public AWeapon
 	{
@@ -94,6 +94,7 @@ Also, you can safely cast an Object from a base class to a more derived class us
 
 Here we have used `Cast` to attempt to cast the `AEnemy` to an `AMegaBoss`. If the Object in question is not actually an `AMegaBoss` (or a child class thereof), the cast will return a null pointer and we can react appropriately. In the code above, the `Incinerate` function will only be called against the MegaBoss.
 
+
 ## Garbage Collection
 
 Unreal implements a garbage collection scheme whereby `UObjects` that are no longer referenced or have been explicitly flagged for destruction will be cleaned up at regular intervals. The engine builds a reference graph to determine which `UObjects` are still in use and which ones are orphaned. At the root of this graph is a set of `UObjects` designated as the "root set". Any `UObject` can be added to the root set. When garbage collection occurs, the engine can track all referenced `UObjects` by searching the tree of known `UObject` references, starting from the root set. Any unreferenced `UObjects`, meaning those which are not found in the tree search, will be assumed to be unneeded, and will be removed.
@@ -103,19 +104,19 @@ One practical implication here is that you typically need to maintain a `UPROPER
 Garbage collection in Unreal Engine 4 is fast and efficient, and has a number of built-in features designed to minimize overhead, such as multithreaded reachability analysis to identify orphaned Objects, and unhashing code optimized to remove Actors from containers as quickly as possible. There are other features that can be adjusted to gain more precise control over how and when garbage collection is performed, most of which are found in **Project Settings** under **Engine - Garbage Collection**. The following settings are commonly used to tune garbage collector performance for a project:
 | Setting(s) | Feature Description |
 | --- | --- |
-| **Create Garbage Collector UObject Clusters** | This can be turned on or off (it is on by default) in the project's settings. If turned on, related Objects will be grouped together in a garbage collection cluster, so that only the cluster itself needs to be checked, rather than each individual Object. This means that reachability is performed faster because an entire cluster can be treated as one object, but it also means that individual items in that cluster will be deleted all at once instead of potentially spreading out over many frames. In the general case, cluster creation improves garbage collection performance. |
-| **Merge GC Clusters** | Cluster merging can be enabled to cause clustered objects to join together when appropriate. **Create Garbage Collector UObject Clusters** must also be turned on for this feature to work. This will lead to a smaller number of garbage collection events, but they will be more frequent. In addition, there may be some cases where garbage collection will not happen with merged clusters when it would have happened without them. This can result in decreased Object creation time, as the Objects that weren't cleared out are still around and can be reused rather than being recreated.
-| **Actor Clustering Enabled** | Actors can be put into clusters by turning on this option in **Project Settings**, and by either setting the `bCanBeInCluster` variable to `true`, or overriding the `CanBeInCluster` function in code so that it returns `true`. This feature is useful for grouping Actors that are expected to be destroyed all at once. |
-| **Blueprint Clustering Enabled** | Blueprint-based Actors can be clustered just like regular Actors by turning this setting on. The **Can Be In Cluster** checkbox in the Blueprint's settings must also be checked. |
-| **Time Between Purging Pending Kill Objects** | The frequency of garbage collection activity can be adjusted in the project's settings. This high-level control is especially useful in preventing hitches. By shortening the time between collections, you decrease the likely amount of destroyed Actors waiting for collection, and can avoid the hitches that can happen when a lot of Actors are cleaned up at the same time. |
+| **Create Garbage Collector UObject Clusters** | This can be turned on or off (it is on by default) in the project's settings. If turned on, related Objects will be grouped together in a garbage collection cluster, so that only the cluster itself needs to be checked, rather than each individual Object. This means that reachability is performed faster because an entire cluster can be treated as one object, but it also means that individual items in that cluster will all be unhashed and prepared for deletion in the same frame, potentially causing a hitch if the cluster is large enough. In the general case, cluster creation improves garbage collection performance and decreases time spent in reachability analysis. |
+| **Merge GC Clusters** | Cluster merging can be enabled to cause clusters to join together when an object from one cluster references an object in another. Note that clearing the reference that caused the merge will not cause the newly-merged cluster to dissolve or break up in any way. **Create Garbage Collector UObject Clusters** must also be turned on for this feature to work. This will lead to the garbage collector unhashing and destroying objects less frequently, but larger numbers of objects will be unhashed and destroyed at once. In addition, there may be some cases where garbage collection will not happen with merged clusters when it would have happened without them, since any reference to any object in the cluster will keep the entire cluster from being garbage-collected.
+| **Actor Clustering Enabled** | Actors can be put into clusters by turning on this option in **Project Settings**, and by either setting the `bCanBeInCluster` variable to `true`, or overriding the `CanBeInCluster` function in code so that it returns `true`. By default, Actors and Components have this turned off, except for Static Mesh Actors and Reflection Capture Components. This feature is useful for grouping Actors that are expected to be destroyed all at once, usually static meshes placed in a level that cannot be destroyed except by unloading the sublevel that contains them. |
+| **Blueprint Clustering Enabled** | A Blueprint's `UBlueprintGeneratedClass` and related data, such as the shared UPROPERTY and UFUNCTION data, can be clustered by turning this setting on. It is important to recognize that this clustering refers to the Blueprint Generated Class itself, not individual instances of the Blueprint. |
+| **Time Between Purging Pending Kill Objects** | The frequency of garbage collection activity can be adjusted in the project's settings. This high-level control is especially useful in preventing hitches. By shortening the time between collections, you decrease the likely amount of unreachable objects that will be discovered on the next reachability analysis pass, and can avoid the hitches that can happen when a lot of Actors are cleaned up at the same time. |
 
 ![ProjectSettingsGarbageCollection.png](ProjectSettingsGarbageCollection.png)
 [REGION:caption]The Garbage Collection settings within Project Settings.[/REGION]
 
 ## Network Replication
 
-The `UObject` system includes a robust set of functionality to facilitate network communication and multiplayer games.
+The `UObject` system includes a robust set of functionality to facilitate [network communication and multiplayer games](Gameplay/Networking/).
 
-`UProperties` can be tagged to tell the Engine to replicate their data during network play. A common model here is that a variable gets changed on the server, and the Engine then detects this change and sends it reliably to all clients. Clients can optionally receive a callback function when the variable changes via replication.
+`UProperties` can be tagged to tell the Engine to [replicate their data during network play](Gameplay/Networking/Actors/Properties/). A common model here is that a variable gets changed on the server, and the Engine then detects this change and sends it reliably to all clients. Clients can optionally receive a callback function when the variable changes via replication.
 
-`UFunctions` can also be tagged to execute on a remote machine. A "server" function, for example, when called on a client machine, will cause the function to execute on the server machine for the server's version of the Actor. A "client" function, on the other hand, can be called from the server and will run on the owning client's version of that Actor.
+`UFunctions` can also be tagged to [execute on a remote machine](Gameplay/Networking/Actors/RPCs/). A "server" function, for example, when called on a client machine, will cause the function to execute on the server machine for the server's version of the Actor. A "client" function, on the other hand, can be called from the server and will run on the owning client's version of that Actor.
diff --git a/Engine/Documentation/Source/Programming/UnrealArchitecture/Reference/Properties/Properties.INT.udn b/Engine/Documentation/Source/Programming/UnrealArchitecture/Reference/Properties/Properties.INT.udn
--- a/Engine/Documentation/Source/Programming/UnrealArchitecture/Reference/Properties/Properties.INT.udn
+++ b/Engine/Documentation/Source/Programming/UnrealArchitecture/Reference/Properties/Properties.INT.udn
@@ -2,7 +2,16 @@ Availability:Public
 Title: Properties
 Crumbs:%ROOT%, Programming, Programming/UnrealArchitecture/Reference
 Description:Reference for creating and implementing properties for gameplay classes.
-Version: 4.12
+type:reference
+skilllevel:beginner
+parent:Programming/UnrealArchitecture
+Version: 4.16
+tags:variable
+tags:UPROPERTY
+tags:UENUM
+tags:enum
+tags:bitmask
+tags:bitfield
 
 ## Property Declaration
 
diff --git a/Engine/Documentation/Source/Programming/UnrealArchitecture/SmartPointerLibrary/SmartPointerLibrary.INT.udn b/Engine/Documentation/Source/Programming/UnrealArchitecture/SmartPointerLibrary/SmartPointerLibrary.INT.udn
--- a/Engine/Documentation/Source/Programming/UnrealArchitecture/SmartPointerLibrary/SmartPointerLibrary.INT.udn
+++ b/Engine/Documentation/Source/Programming/UnrealArchitecture/SmartPointerLibrary/SmartPointerLibrary.INT.udn
@@ -120,7 +120,7 @@ Only one reference controller is created per object, no matter how many shared/w
 
 Shared pointers are non-intrusive which means the class itself does not know whether it is owned by a shared pointer (or reference). In general, this is okay, but sometimes you really want to access the current instance as a shared reference. The solution to this is to derive the class from `TSharedFromThis<>`.
 
-By deriving from TSharedFromThis<>, you can use the **AsSharedRef()** method to convert 'this' to a shared reference. This can be very useful with class factories that always return shared references.
+By deriving from `TSharedFromThis<>`, you can use the **AsSharedRef()** method to convert 'this' to a shared reference. This can be very useful with class factories that always return shared references.
 
 	class FAnimation : public TSharedFromThis<FMyClass>
 	{
@@ -181,7 +181,7 @@ The SharedPointerTesting.h file (located in [UE4RootLocation]/Engine/Source/Runt
 * The "thread-safe" versions of smart pointers are a bit slower -- only use them when needed.
 * You can forward declare shared pointers to incomplete types, just how you would expect to!
 * Shared pointers of compatible types will be converted implicitly (e.g. upcasting).
-* You can create a typedef to TSharedRef< MyClass > to make it easier to type.
+* You can create a typedef to `TSharedRef< MyClass >` to make it easier to type.
 * For best performance, minimize calls to TWeakPtr::Pin (or conversions to TSharedRef/TSharedPtr).
 * Your class can return itself as a shared reference if you derive from TSharedFromThis.
 * To downcast a pointer to a derived object class, use the **StaticCastSharedPtr()** function.
diff --git a/Engine/Documentation/Source/Programming/UnrealArchitecture/StringHandling/CharacterEncoding/CharacterEncoding.INT.udn b/Engine/Documentation/Source/Programming/UnrealArchitecture/StringHandling/CharacterEncoding/CharacterEncoding.INT.udn
--- a/Engine/Documentation/Source/Programming/UnrealArchitecture/StringHandling/CharacterEncoding/CharacterEncoding.INT.udn
+++ b/Engine/Documentation/Source/Programming/UnrealArchitecture/StringHandling/CharacterEncoding/CharacterEncoding.INT.udn
@@ -72,7 +72,7 @@ A string made up of 2 bytes per character with a [BOM](http://en.wikipedia.org/w
 
 All strings in Unreal Engine 4 (UE4) are stored in memory in [UTF-16](http://en.wikipedia.org/wiki/UTF-16/UCS-2) format as FStrings or TCHAR arrays. Most code assumes 2 bytes is one codepoint so only the Basic Multilingual Plane (BMP) is supported so Unreal's internal encoding is more correctly described as UCS-2. Strings are stored in the endian-ness appropriate for the current platform.
 
-When serializing to packages to/from disk or during networking, strings with all TCHAR characters < 0xff are stored as a series of 8-bit bytes, and otherwise as 2-byte UTF-16 strings. The serialization code can deal with any endian conversion as necessary.
+When serializing to packages to/from disk or during networking, strings with all TCHAR characters less than 0xff are stored as a series of 8-bit bytes, and otherwise as 2-byte UTF-16 strings. The serialization code can deal with any endian conversion as necessary.
 
 
 
@@ -147,11 +147,11 @@ We have a number of macros to convert strings to and from various encodings. The
 These use the following helper classes from UnStringConv.h:
 
 
-* typedef TStringConversion<TCHAR,ANSICHAR,FANSIToTCHAR_Convert> FANSIToTCHAR;
-* typedef TStringConversion<ANSICHAR,TCHAR,FTCHARToANSI_Convert> FTCHARToANSI;
-* typedef TStringConversion<ANSICHAR,TCHAR,FTCHARToOEM_Convert> FTCHARToOEM;
-* typedef TStringConversion<ANSICHAR,TCHAR,FTCHARToUTF8_Convert> FTCHARToUTF8;
-* typedef TStringConversion<TCHAR,ANSICHAR,FUTF8ToTCHAR_Convert> FUTF8ToTCHAR;
+`* typedef TStringConversion<TCHAR,ANSICHAR,FANSIToTCHAR_Convert> FANSIToTCHAR;`
+* `typedef TStringConversion<ANSICHAR,TCHAR,FTCHARToANSI_Convert> FTCHARToANSI;`
+* `typedef TStringConversion<ANSICHAR,TCHAR,FTCHARToOEM_Convert> FTCHARToOEM;`
+* `typedef TStringConversion<ANSICHAR,TCHAR,FTCHARToUTF8_Convert> FTCHARToUTF8;`
+* `typedef TStringConversion<TCHAR,ANSICHAR,FUTF8ToTCHAR_Convert> FUTF8ToTCHAR;`
 
 
 It is also critical that when using TCHAR_TO_ANSI you do not assume the number of bytes will be the same as the TCHAR string length. Multiple byte character sets could require multiple bytes per TCHAR character. If you need to know the length of the resulting string in bytes, you can use the helper class instead of the macros. For example:
diff --git a/Engine/Documentation/Source/Programming/UnrealArchitecture/StringHandling/FName/FName.INT.udn b/Engine/Documentation/Source/Programming/UnrealArchitecture/StringHandling/FName/FName.INT.udn
--- a/Engine/Documentation/Source/Programming/UnrealArchitecture/StringHandling/FName/FName.INT.udn
+++ b/Engine/Documentation/Source/Programming/UnrealArchitecture/StringHandling/FName/FName.INT.udn
@@ -50,7 +50,7 @@ FName (such as **IsValidObjectName()**) check the validity of FNames for particu
 
 The == operator can be used to compare two FNames, and returns true or false.  Rather than doing string comparisons, this compares the values in Index, which is a significant CPU savings.
 
-**FName::Compare** can also be used to compare two FNames, and it will return < 0 if this < Other, 0 if this == Other, and > 0 if this > Other.
+**FName::Compare** can also be used to compare two FNames, and it will return less than 0 if this less than Other, 0 if this equal to Other, and greater than 0 if this greater than Other.
 
 	CompareFloat = TestFName.Compare(OtherFName); 
 
diff --git a/Engine/Documentation/Source/Programming/UnrealArchitecture/TArrays/TArrays.INT.udn b/Engine/Documentation/Source/Programming/UnrealArchitecture/TArrays/TArrays.INT.udn
--- a/Engine/Documentation/Source/Programming/UnrealArchitecture/TArrays/TArrays.INT.udn
+++ b/Engine/Documentation/Source/Programming/UnrealArchitecture/TArrays/TArrays.INT.udn
@@ -41,7 +41,7 @@ Memory was allocated from the allocator as the elements were added.  Add and Emp
 * Add will copy (or move) an instance of the element type into the array.
 * Emplace will use the arguments you give it to construct a new instance of the element type.
 
-So in the case of our TArray<FString>, Add will create a temporary FString from the string literal and then move the contents of that temporary into a new FString inside the container, whereas Emplace will just create its FString directly using the string literal.  The end result is the same, but Emplace avoids creating a temporary, which is often undesirable for non-trivial value types like FString.  Push can also be used as a synonym for Add.
+So in the case of our `TArray<FString>`, Add will create a temporary FString from the string literal and then move the contents of that temporary into a new FString inside the container, whereas Emplace will just create its FString directly using the string literal.  The end result is the same, but Emplace avoids creating a temporary, which is often undesirable for non-trivial value types like FString.  Push can also be used as a synonym for Add.
 
 In general, Emplace is better than Add, in that it avoids unnecessary temporaries being created at the call site which are then copied or moved into the container.  As a rule of thumb, use Add for trivial types and Emplace otherwise.  Emplace should never be less efficient than Add, but sometimes Add may read better.
 
@@ -109,7 +109,7 @@ Arrays can be sorted simply by calling the **Sort** function:
 	StrArr.Sort();
 	// StrArr == ["!","Brave","Hello","of","Tomorrow","World"]
 
-Here, the values are sorted by means of the element type's operator<.  In FString's case, this is a case-insensitive lexicographical comparison.  A binary predicate can also be provided to provide different ordering semantics, e.g.:
+Here, the values are sorted by means of the element type's `operator<`.  In FString's case, this is a case-insensitive lexicographical comparison.  A binary predicate can also be provided to provide different ordering semantics, e.g.:
 
 	StrArr.Sort([](const FString& A, const FString& B) {
 		return A.Len() < B.Len();
@@ -407,7 +407,7 @@ Arrays can be compared using the operator== and operator!=.  The order of the el
 
 TArray has functions which support a binary heap data structure.  A heap is a type of binary tree in which any parent node is equivalent to or ordered before all of its child nodes.  When implemented as an array, the root node of the tree is at element 0 and the indices of the left and right child nodes of a node at index N are 2N+1 and 2N+2 respectively.  The children are not in any particular order with respect to one another.
 
-Any existing array can be turned into a heap by calling the **Heapify** function.  This is overloaded to take a predicate or not, where the non-predicated version will use the element type's operator< to determine ordering:
+Any existing array can be turned into a heap by calling the **Heapify** function.  This is overloaded to take a predicate or not, where the non-predicated version will use the element type's `operator<` to determine ordering:
 
 	TArray<int32> HeapArr;
 	for (int32 Val = 10; Val != 0; --Val)
@@ -455,7 +455,7 @@ The **HeapPop** and **HeapPopDiscard** functions are used to remove the top node
 
 It should be noted that HeapPush, HeapPop, HeapPopDiscard and HeapRemoveAt should only be called when the structure is already a valid heap, e.g. after a Heapify() call, any other heap operation, or by manually manipulating the array into a heap.
 
-Also, each of these functions, including Heapify, can take an optional binary predicate to determine the order of the node elements in the heap.  When using a predicate, the same predicate should be used across all heap operations to maintain the correct structure.  When no predicate is specified, the heap operations use the element type's operator< to determine order.
+Also, each of these functions, including Heapify, can take an optional binary predicate to determine the order of the node elements in the heap.  When using a predicate, the same predicate should be used across all heap operations to maintain the correct structure.  When no predicate is specified, the heap operations use the element type's `operator<` to determine order.
 
 Finally, the top node of the heap can be inspected using **HeapTop**, without changing the array:
 
@@ -596,7 +596,7 @@ Use of any of the Uninitialized or Zeroed functions should be treated with cauti
 
 ## Misc
 
-The **BulkSerialize** function is a serialization function which can be used as an alternative operator<< in order to serialize the array as a block of raw bytes instead of doing per-element serialization.  This can be a performance win when your element type is trivial, like a built-in type or a plain data struct.
+The **BulkSerialize** function is a serialization function which can be used as an alternative `operator<<` in order to serialize the array as a block of raw bytes instead of doing per-element serialization.  This can be a performance win when your element type is trivial, like a built-in type or a plain data struct.
 
 The **CountBytes** and **GetAllocatedSize** functions are used to estimate how much memory is currently being utilized by the array.  CountBytes takes an FArchive and GetAllocatedSize can be called directly.  They are typically used for stats reporting.
 
diff --git a/Engine/Documentation/Source/Programming/UnrealArchitecture/TMap/TMap.INT.udn b/Engine/Documentation/Source/Programming/UnrealArchitecture/TMap/TMap.INT.udn
--- a/Engine/Documentation/Source/Programming/UnrealArchitecture/TMap/TMap.INT.udn
+++ b/Engine/Documentation/Source/Programming/UnrealArchitecture/TMap/TMap.INT.udn
@@ -13,7 +13,7 @@ There are two types of map: TMap and **TMultiMap**.  TMap's keys are unique, and
 
 ## TMap
 
-TMaps are primarily defined by two types - a key type and a value type - which are stored as associated pairs in the map.  It is convenient to refer to these pairs as the element type of the map as if it was an individual object.  In this document, element means a key-value pair, while individual components are referred to as the element's key or the element's value.  The element type is actually a **TPair**< KeyType, ElementType >, though it should be rare to need to refer to the TPair type directly.
+TMaps are primarily defined by two types - a key type and a value type - which are stored as associated pairs in the map.  It is convenient to refer to these pairs as the element type of the map as if it was an individual object.  In this document, element means a key-value pair, while individual components are referred to as the element's key or the element's value.  The element type is actually a `TPair< KeyType, ElementType >`, though it should be rare to need to refer to the TPair type directly.
 
 Like TArray, TMap is a homogeneous container, meaning that all of its elements are strictly the same type.  TMap is also a value type, and supports the usual copying, assignment and destructor operations, as well as strong ownership of its elements, which will be destroyed when the map is destroyed.  The key type and value type must also be value types.
 
diff --git a/Engine/Documentation/Source/Programming/UnrealArchitecture/Timers/GameplayTimers.INT.udn b/Engine/Documentation/Source/Programming/UnrealArchitecture/Timers/GameplayTimers.INT.udn
--- a/Engine/Documentation/Source/Programming/UnrealArchitecture/Timers/GameplayTimers.INT.udn
+++ b/Engine/Documentation/Source/Programming/UnrealArchitecture/Timers/GameplayTimers.INT.udn
@@ -1,124 +1,101 @@
 Availability:Public
+Crumbs: %ROOT%
 Title: Gameplay Timers
-Crumbs:%ROOT%, Programming, Programming/UnrealArchitecture/Reference, Programming/UnrealArchitecture/Reference/Functions
 Description:Timer construct for performing actions at set intervals.
-Version: 4.9
+Version: 4.17
+Type: Overview
+SkillLevel: Intermediate
+Parent: Programming/UnrealArchitecture
+Order: 0
+Related: Programming/UnrealArchitecture/Delegates
 
 [EXCERPT:Overview]
-**Timers** allow actions to be performed - once or repeatedly - at designated intervals. They can be 
-used to set actions to happen on a delay or over a period of time. For example, you may want to 
-make the player invincible after obtaining a certain item and then revert after 10 seconds. Or, 
-you may want to apply damage every second the player is in a room filled with toxic gas. Each of 
-these can be achieved through the use of timers.
+**Timers** schedule actions to be performed after a delay, or at a over a period of time. For example, you may want to make the player invulnerable after obtaining a power-up item, and then restore vulnerability after 10 seconds. Or, you may want to apply damage once per second while the player moves through a room filled with toxic gas. Each of these can be achieved through the use of timers.
 [/EXCERPT:Overview]
 
 ## Timer Management
 
-Timers are managed in a global **FTimerManager**, outside the confines of AActor, and can have any 
-of the full range of delegate types assigned. There are several functions in FTimerManager available 
-for managing timers. It is safe to use these functions inside of a timer delegate as the system is ok 
-with manipulating timers while handling a timer. This means, for example, it is ok to set or clear 
-timers inside a timer delegate.
+Timers are managed in a global **Timer Manager** (of type `FTimerManager`). A global Timer Manager exists on the **Game Instance** object, as well as on each **World**. There are two main functions used to set up Timers with a Timer Manager: `SetTimer`, and `SetTimerForNextTick`, each with several overloads. Each one can be attached to any type of object or function delegate, and `SetTimer` can be made to repeat at regular intervals, if desired. See the [Timer Manager API page](API:FTimerManager) for more detail on these functions.
 
-The **AActor::GetWorldTimerManager()** function is used to access the timer manager instance for the 
-current world.
+[REGION:note]Timers will be canceled automatically if the object that they are going to be called on, such as an Actor, is destroyed before the time is up. In this case, the Timer Handle will become invalid and the function will not be called.[/REGION]
 
-### Creating and Clearing Timers
+Accessing the Timer Manager can be done through the `AActor` function called `GetWorldTimerManager`, which calls up to the `GetTimerManager` function in `UWorld`. To access the global Timer Manager, use the `UGameInstance` function, `GetTimerManager'. This is also the fallback used if a World doesn't have its own Timer Manager for any reason, and can be used for function calls that aren't relevant to, or should not depend on, the existence of any specific World.
 
-**FTimerManager::SetTimer()** is used to create a timer to call the given native function at a set interval. 
+Timers can be used with standard C++ function pointers, [`TFunction` Objects](API:TFunction), or [](Programming/UnrealArchitecture/Delegates).
 
-Example:
+### Setting and Clearing Timers
 
-	GetWorldTimerManager().SetTimer(this, &AMatineeActor::CheckPriorityRefresh, 1.0f, true);
+The `SetTimer` functions of `FTimerManager` will set a timer to call a function or delegate after a delay, and can be set to repeat that function call indefinitely. These functions will fill out a **Timer Handle** (type `FTimerHandle`), which can be used to pause (and resume) the countdown, query or change the amount of time remaining, or even cancel the Timer altogether. It is safe to set Timers within a function called by a Timer, even including reuse of the Timer Handle that was used to call the function. One use for this might be to delay initialization of one Actor if it depends on another Actor that hasn't spawned yet, but is expected to spawn soon; the dependent Actor's initialization function could set a Timer to call itself again after a fixed length of time, e.g. one second. Alternately, the initialization function could be called by a looping Timer that clears itself upon success.
 
-[REGION:tip]
-SetTimer() can also reset an existing timer. So, if a timer is already set for this delegate, it will 
-update the current timer to the new parameters and reset its elapsed time to 0.
-[/REGION]
+Timers can also be set to run on the next frame, rather than with a timed interval. This is accomplished by calling `SetTimerForNextTick`, but note that this function does not fill out a Timer Handle.
+
+To clear a Timer, pass the `FTimerHandle` that was filled out during the `SetTimer` call into the `FTimerManager` function called `ClearTimer`. The Timer Handle will become invalid at this point, and can be reused to manage a new Timer. Calling `SetTimer` with an existing Timer Handle will clear the Timer referenced by that Timer Handle and replace it with a new one.
 
-**FTimerManager::ClearTimer()** clears a previously set timer so it will not execute any more.
+Finally, all Timers associated with a specific object can be cleared by calling `ClearAllTimersForObject`.
 
 Example:
 
-	GetWorldTimerManager().ClearTimer(this, &AMatineeActor::CheckPriorityRefresh);
+	void AMyActor::BeginPlay()
+	{
+		Super::BeginPlay();
+		// Call RepeatingFunction once per second, starting two seconds from now.
+		GetWorldTimerManager().SetTimer(MemberTimerHandle, this, &AMyActor::RepeatingFunction, 1.0f, true, 2.0f);
+	}
+	
+	void AMyActor::RepeatingFunction()
+	{
+		// Once we've called this function enough times, clear the Timer.
+		if (--RepeatingCallsRemaining <= 0)
+		{
+			GetWorldTimerManager().ClearTimer(MemberTimerHandle);
+			// MemberTimerHandle can now be reused for any other Timer.
+		}
+		// Do something here...
+	}
 
 [REGION:tip]
-Clearing a timer with ClearTimer() is identical to calling SetTimer() with a <= 0.f rate.
+Calling `SetTimer` with a rate less than or equal to zero is identical to calling `ClearTimer`.
 [/REGION]
 
-### Pausing Timers
+### Pausing and Resuming Timers
 
-**FTimerManager::PauseTimer()** pauses a previously set timer in its current state. This means it will 
-not execute, but the elapsed and remaining time stay the same until the timer is unpaused.
-
-**FTimerManager::UnPauseTimer()** sets a paused timer active once again.
+The `FTimerManager` function, `PauseTimer` uses a Timer Handle to pause a running Timer. This prevents the Timer from executing its function call, but the elapsed and remaining time stay the same while paused. `UnPauseTimer` causes a paused Timer to resume running.
 
 ## Timer Information
 
-In addition to managing timers, FTimermanager also provides funtions for obtaining information - such 
-as timer rate, elapsed time, remaining time, etc. - about a specific timer.
+In addition to managing Timers, Timer Managers also provide funtions for obtaining information - such as the rate, elapsed time, remaining time, etc. - about a specific timer.
 
 ### Is Timer Active
 
-**FTimerManager::IsTimerActive()** is used to determine if the specified timer is curently active and not paused.
+The `IsTimerActive` function of `FTimerManager` is used to determine if the specified timer is curently active and not paused.
 
 Example:
 
-	GetWorldTimerManager().IsTimerActive(this, &APlayerController::DelayedPrepareMapChange)
+	// Is this weapon waiting to be able to fire again?
+	GetWorldTimerManager().IsTimerActive(this, &AUTWeapon::RefireCheckTimer);
 
 ### Timer Rate
 
-**FTimerManager::GetTimerRate()** gets the current rate (time between activations) for the specified timer.
+`FTimerManager` has a function called `GetTimerRate` that gets the current rate (time between activations) of a Timer from its Timer Handle. A Timer's rate cannot be changed directly, but `SetTimer` can be called with its Timer Handle to clear it and create a new Timer, which could be identical except for the rate.  `GetTimerRate` will return a value of `-1` if the Timer Handle is not valid.
 
 Example:
 
-	GetWorldTimerManager().GetTimerRate(this, &AUTWeapon::RefireCheckTimer)
+	// This weapon's rate of fire changes as it warms up. Is it currently waiting to fire, and if so, how long is the current delay between shots?
+	GetWorldTimerManager().GetTimerRate(this, &AUTWeapon::RefireCheckTimer);
 
 ### Elapsed and Remaining Time
 
-**FTimermanager::GetTimerElapsed()** and **FTimerManager::GetTimerRemaining()** return the elapsed and remaining 
-time, respectively, for the specified timer. 
+`FTimermanager` provides functionality, through `GetTimerElapsed` and `GetTimerRemaining` return the elapsed and remaining time, respectively, for the Timer associated with the provided Timer Handle. As with `GetTimerRate`, these functions will return `-1` if the Timer Handle is invalid.
 
 Example:
 
-	GetWorldTimerManager().GetTimerElapsed(this, &AUTWeapon::RefireCheckTimer)
+	// How long will it be until this weapon is ready to fire again? If the answer comes back as -1, it is ready now.
+	GetWorldTimerManager().GetTimerElapsed(this, &AUTWeapon::RefireCheckTimer);
 
 [REGION:note]
 The sum of the elapsed and remaining time for a timer should be equal to the rate of the timer.
 [/REGION]
 
-## Timer Function Overloads
-
-Each of the timer-manipulation functions has 3 flavors of overloads for interacting with timers.
-
-* Taking an object and a function.
-                
-		GetWorldTimerManager().SetTimer( this, &UMyClass::DoSomethingWithNoParams, 5.f, FALSE );
-
-	This is the simplest form and the closest to the older system. It is identical to 
-    
-		GetWorldTimerManager().SetTimer( FTimerDelegate::CreateUObject(this, &UMyClass::DoSomethingWithNoParams), 5.f, FALSE );
-
-	It is a bit redundant, but it is cleaner and expected to be the most common use case.
-
-* Taking any type of non-dynamic delegate you can muster.  It is a bit muddier syntactically but far more flexible.  e.g. 
-
-		GetWorldTimerManager().SetTimer( FTimerDelegate::CreateUObject(this, &AMyClass::MyFunc), 5.0f, TRUE );
-		GetWorldTimerManager().SetTimer( FTimerDelegate::CreateRaw(&MyStaticFunc, MyPassThruArg, MyOtherArg, YesOneMoreArg), 5.0f, TRUE );
-                
-		FTimerDelegate MyDel;
-		Del.BindRaw(this, &AMyClass::MyFunc2);
-		GetWorldTimerManager().SetTimer( MyDel, 5.0f, TRUE );
-
-* Taking a dynamic (i.e. UFunction) delegate.
-                
-		FTimerDynamicDelegate D;
-		D.BindDynamic( this, &AMyClass::MyUFunction );
-		GetWorldTimerManager().SetTimer( D, 5.0f, TRUE );
-
-	These delegates are called via the UFUNCTION event mechanism, meaning they will properly respect things like non-native overrides such as in Blueprints.
-
-
 ## Known Issues
                
-* The code is not threadsafe at the moment. It will cause an assert if accessed outside the game thread.
+* The code is not threadsafe at the moment, and will cause an assert if accessed outside of the game thread.
diff --git a/Engine/Documentation/Source/Programming/UnrealBuildSystem/IWYUReferenceGuide/IWYUReferenceGuide.INT.udn b/Engine/Documentation/Source/Programming/UnrealBuildSystem/IWYUReferenceGuide/IWYUReferenceGuide.INT.udn
--- a/Engine/Documentation/Source/Programming/UnrealBuildSystem/IWYUReferenceGuide/IWYUReferenceGuide.INT.udn
+++ b/Engine/Documentation/Source/Programming/UnrealBuildSystem/IWYUReferenceGuide/IWYUReferenceGuide.INT.udn
@@ -124,3 +124,4 @@ If you want your game to opt-in to IWYU, there are a few tips to keep in mind:
 1.	Include `CoreMinimal.h` at the top of each header file.
 1.	To verify that all of your source files include all of their required dependencies, compile your game project in non-unity mode with PCH files disabled.
 1.	If you need to access **UEngine** or **GEngine**, which are defined in `Runtime\Engine\Classes\Engine\Engine.h`, you can `#include Engine/Engine.h` (distinguishing from the monolithic header file, which is located at `Runtime\Engine\Public\Engine.h`).
+1.	If you use a class that the compiler doesn't recognize, and don't know what you need to include may be missing the header file. This is especially the case if you are converting from non-IWYU code that compiled correctly. You can look up the class in the [API Documentation](API), and find the necessary modules and header files at the bottom of the page.
diff --git a/Engine/Documentation/Source/Resources/Concepts/GameplayConceptExamples.INT.udn b/Engine/Documentation/Source/Resources/Concepts/GameplayConceptExamples.INT.udn
--- a/Engine/Documentation/Source/Resources/Concepts/GameplayConceptExamples.INT.udn
+++ b/Engine/Documentation/Source/Resources/Concepts/GameplayConceptExamples.INT.udn
@@ -63,7 +63,7 @@ Feel free to open up and analyze any element of these examples, as well as use t
 	Blueprint Splines Track
 	[/PARAMLITERAL]
 	[PARAMLITERAL:playlistid]
-	PLZlv_N0_O1gYuWkyTSbfcuyPNOQsK_8bu
+	PLZlv_N0_O1ga0aV9jVqJgog0VWz1cLL5f
 	[/PARAMLITERAL]
 	[PARAMLITERAL:videoid]
 	wR0fH6O9jD8
@@ -114,7 +114,7 @@ Feel free to open up and analyze any element of these examples, as well as use t
 	Turn Based Strategy
 	[/PARAMLITERAL]
 	[PARAMLITERAL:playlistid]
-	PLZlv_N0_O1gYuWkyTSbfcuyPNOQsK_8bu
+	PLZlv_N0_O1ga0aV9jVqJgog0VWz1cLL5f
 	[/PARAMLITERAL]
 	[PARAMLITERAL:videoid]
 	mI7eYXMJ5eI
diff --git a/Engine/Documentation/Source/Resources/ContentExamples/Matinee/2_5/2_5.INT.udn b/Engine/Documentation/Source/Resources/ContentExamples/Matinee/2_5/2_5.INT.udn
--- a/Engine/Documentation/Source/Resources/ContentExamples/Matinee/2_5/2_5.INT.udn
+++ b/Engine/Documentation/Source/Resources/ContentExamples/Matinee/2_5/2_5.INT.udn
@@ -33,6 +33,6 @@ The Director Group also has a few special tracks that can be added.
 
 Fade lets you fade in and out.
 
-Slomo is changing the speed of time 1 = normal speed. <1 is slow motion. >1 fast forward
+Slomo is changing the speed of time 1 equates to normal speed. less than 1 is slow motion. greater than 1 fast forward
 
 
diff --git a/Engine/Documentation/Source/Resources/Showcases/Effects/SnowExamples/SnowExamples.INT.udn b/Engine/Documentation/Source/Resources/Showcases/Effects/SnowExamples/SnowExamples.INT.udn
--- a/Engine/Documentation/Source/Resources/Showcases/Effects/SnowExamples/SnowExamples.INT.udn
+++ b/Engine/Documentation/Source/Resources/Showcases/Effects/SnowExamples/SnowExamples.INT.udn
@@ -51,7 +51,7 @@ Large _unlit_ (i.e. _Emissive_) sprites generate widespread puffs of haze or fog
 The Blizzard Particle System also makes use of a Velocity Grid with low velocity vectors in order to slow down the snowflake sprites as they enter the volume. This slowing effect simulates a sense of drag on the particles, and also aids in collisions, keeping them from feeling quite as rigid. Vector Fields can most easily be visualized in Cascade via the Viewport menu, selecting **View > Vector Fields**.
 
 ![](ShowVectorFields.png)
-<br>_The yellow box represents the Vector Field, while the red lines within it represent the vectors at each location within the volume._
+[REGION:caption]The yellow box represents the Vector Field, while the red lines within it represent the vectors at each location within the volume.[/REGION]
 
 [REGION:tip]
 	If using the _Collision (Scene Depth)_ module, you can boost the **Radius Bias** property to cause particles to appear to "glide around" a surface. This trick was used on the Blizzard Particle System to give the illusion of air influence around a surface.
diff --git a/Engine/Documentation/Source/Resources/Showcases/PhotorealisticCharacter/TwinblastBust.INT.udn b/Engine/Documentation/Source/Resources/Showcases/PhotorealisticCharacter/TwinblastBust.INT.udn
--- a/Engine/Documentation/Source/Resources/Showcases/PhotorealisticCharacter/TwinblastBust.INT.udn
+++ b/Engine/Documentation/Source/Resources/Showcases/PhotorealisticCharacter/TwinblastBust.INT.udn
@@ -278,7 +278,7 @@ Due to the complexity of the eye shader and our recommendation that you use it a
 |**Sclera inner color**| This is the color of the sclera where it meets the iris.|n/a|
 |**Sclera outer color**| This is the color of the sclera at the outer edge of the eye.| n/a|
 |**SecondaryEnvBalance**|This controls how much the [Secondary Environment](#eyeshadersecondaryenvironment) is visible on the surface of the eye.|0.00 - 0.03|
-|**SecondaryEnvRotationAxis**|Since the [Secondary Environment](#eyeshadersecondaryenvironment) is a faked reflection, it must be manually rotated around the eye. This property controls the axis about which it will rotate.| <0.0, 0.0, 1.0> (Z-Axis)|
+|**SecondaryEnvRotationAxis**|Since the [Secondary Environment](#eyeshadersecondaryenvironment) is a faked reflection, it must be manually rotated around the eye. This property controls the axis about which it will rotate.| [0.0, 0.0, 1.0] (Z-Axis)|
 |[REGION:tablesection]Scalar Parameter Values[/REGION]|||
 |**Depth Scale**|This controls the depth of the refraction of the iris underneath the cornea. |1.0-1.4|
 |**Flatten Normal**|This value controls how much flattening of the eye's normal map is taking place, focused particularly on the sclera.|0.9 - 1.1|
diff --git a/Engine/Documentation/Source/Shared/Editor/MapErrors/MapErrors.INT.udn b/Engine/Documentation/Source/Shared/Editor/MapErrors/MapErrors.INT.udn
--- a/Engine/Documentation/Source/Shared/Editor/MapErrors/MapErrors.INT.udn
+++ b/Engine/Documentation/Source/Shared/Editor/MapErrors/MapErrors.INT.udn
@@ -118,7 +118,7 @@ You want to either set Collision Enabled to No Collision on the Actor, or add co
 [EXCERPT:ShadowCasterUsingBoundsScale]
 **Actor casts dynamic shadows and has a BoundsScale greater than 1! This will have a large performance hit**
 
-Serious performance warning... either reduce BoundsScale to be <= 1 or remove dynamic shadows...
+Serious performance warning... either reduce BoundsScale to be less than or equal to 1 or remove dynamic shadows...
 [/EXCERPT:ShadowCasterUsingBoundsScale]
 
 (#InvalidTrace)
diff --git a/Engine/Documentation/Source/Shared/Glossary/P/Glossary_P.INT.udn b/Engine/Documentation/Source/Shared/Glossary/P/Glossary_P.INT.udn
--- a/Engine/Documentation/Source/Shared/Glossary/P/Glossary_P.INT.udn
+++ b/Engine/Documentation/Source/Shared/Glossary/P/Glossary_P.INT.udn
@@ -169,9 +169,38 @@ project
 
 [EXCERPT:Glossary_project]
 A **Project** is a self-contained unit that holds all the content and code that make up an individual game and coincides with a set of directories on your disk. For example, in the image below the Hierarchy Tree of the **Content Browser** contains the same directory structure found inside your Project folder on your disk. 
+
+[OBJECT:ToggleButtonContent]
+[PARAMLITERAL:category]
+OS
+[/PARAMLITERAL]
+[PARAMLITERAL:id]
+windows
+[/PARAMLITERAL]
+[PARAMLITERAL:active]
+[/PARAMLITERAL]
+[PARAM:content]
 [REGION:lightbox]
-[![](ProjectHierarchy.png)(w:440)](ProjectHierarchy.png)
-[/REGION:lightbox]
+[![](ProjectHierarchy_Windows.png)](ProjectHierarchy_Windows.png)
+[/REGION]
+[/PARAM]
+[/OBJECT]
+
+[OBJECT:ToggleButtonContent]
+[PARAMLITERAL:category]
+OS
+[/PARAMLITERAL]
+[PARAMLITERAL:id]
+mac
+[/PARAMLITERAL]
+[PARAMLITERAL:active]
+[/PARAMLITERAL]
+[PARAM:content]
+[REGION:lightbox]
+[![](ProjectHierarchy_Mac.png)](ProjectHierarchy_Mac.png)
+[/REGION]
+[/PARAM]
+[/OBJECT]
 
 _Click image for full view._
 
