Copying //UE4/Dev-Documentation to //UE4/Samples-Main (Source: //UE4/Dev-Documentation @ 3623665)

diff --git a/Engine/Documentation/Source/UE4Home.INT.udn b/Engine/Documentation/Source/UE4Home.INT.udn
--- a/Engine/Documentation/Source/UE4Home.INT.udn
+++ b/Engine/Documentation/Source/UE4Home.INT.udn
@@ -11,9 +11,9 @@ type:landing
 
 [REGION:frontpage]
 	[OBJECT:SearchBox]
-		[PARAM:path]
-		[RELATIVE:Search]
-		[/PARAM]
+		[PARAMLITERAL:path]
+		https://www.unrealengine.com/bing-search
+		[/PARAMLITERAL]
 		[PARAMLITERAL:placeholder]
 		Need Help? Find a topic by searching
 		[/PARAMLITERAL]
