Documentation for build tools

diff --git a/Engine/Documentation/Source/Programming/UnrealBuildSystem/Configuration/BuildConfigProperties/BuildConfigProperties.INT.udn b/Engine/Documentation/Source/Programming/UnrealBuildSystem/Configuration/BuildConfigProperties/BuildConfigProperties.INT.udn
--- a/Engine/Documentation/Source/Programming/UnrealBuildSystem/Configuration/BuildConfigProperties/BuildConfigProperties.INT.udn
+++ b/Engine/Documentation/Source/Programming/UnrealBuildSystem/Configuration/BuildConfigProperties/BuildConfigProperties.INT.udn
@@ -4,6 +4,10 @@ Crumbs:
 Description: This is a procedurally generated markdown page.
 Version: 4.18
 
+### XboxOnePlatform
+
+$ Compiler : Version of the compiler toolchain to use on Windows platform. A value of "default" will be changed to a specific version at UBT startup.
+
 ### BuildConfiguration
 
 $ bGeneratedSYMFile : Whether to generate dSYM files Lists Architectures that you want to build
@@ -22,6 +26,8 @@ $ bEditorDependsOnShaderCompileWorker : Whether to include a dependency on Shade
 
 $ bDebugBuildsActuallyUseDebugCRT : Enables the debug C++ runtime (CRT) for debug builds. By default we always use the release runtime, since the debug version isn't particularly useful when debugging Unreal Engine projects, and linking against the debug CRT libraries forces our third party library dependencies to also be compiled using the debug CRT (and often perform more slowly). Often it can be inconvenient to require a separate copy of the debug versions of third party static libraries simply so that you can debug your program's code.
 
+$ bUseXGEController : Whether the XGE controller worker and modules should be included in the engine build. These are required for distributed shader compilation using the XGE interception interface.
+
 $ bUseUnityBuild : Whether to unify C++ code into larger files for faster compilation.
 
 $ bForceUnityBuild : Whether to force C++ source files to be combined into larger files for faster compilation.
diff --git a/Engine/Documentation/Source/Programming/UnrealBuildSystem/TargetFiles/TargetFilesProperties/TargetFilesProperties.INT.udn b/Engine/Documentation/Source/Programming/UnrealBuildSystem/TargetFiles/TargetFilesProperties/TargetFilesProperties.INT.udn
--- a/Engine/Documentation/Source/Programming/UnrealBuildSystem/TargetFiles/TargetFilesProperties/TargetFilesProperties.INT.udn
+++ b/Engine/Documentation/Source/Programming/UnrealBuildSystem/TargetFiles/TargetFilesProperties/TargetFilesProperties.INT.udn
@@ -128,6 +128,8 @@ $ bForceCompilePerformanceAutomationTests (Boolean): Whether to compile performa
 
 $ bEventDrivenLoader (Boolean): If true, event driven loader will be used in cooked builds. @todoio This needs to be replaced by a runtime solution after async loading refactor.
 
+$ bUseXGEController (Boolean): Whether the XGE controller worker and modules should be included in the engine build. These are required for distributed shader compilation using the XGE interception interface.
+
 $ bEnforceIWYU (Boolean): Enforce "include what you use" rules; warns if monolithic headers (Engine.h, UnrealEd.h, etc...) are used, and checks that source files include their matching header first.
 
 $ bHasExports (Boolean): Whether the final executable should export symbols.
