Documentation for build tools

diff --git a/Engine/Documentation/Source/Programming/UnrealBuildSystem/Configuration/BuildConfigProperties/BuildConfigProperties.INT.udn b/Engine/Documentation/Source/Programming/UnrealBuildSystem/Configuration/BuildConfigProperties/BuildConfigProperties.INT.udn
--- a/Engine/Documentation/Source/Programming/UnrealBuildSystem/Configuration/BuildConfigProperties/BuildConfigProperties.INT.udn
+++ b/Engine/Documentation/Source/Programming/UnrealBuildSystem/Configuration/BuildConfigProperties/BuildConfigProperties.INT.udn
@@ -164,6 +164,8 @@ $ bStopXGECompilationAfterErrors : When enabled, XGE will stop compiling targets
 
 ### ProjectFileGenerator
 
+$ Format : Default list of project file formats to generate
+
 ### VCProjectFileGenerator
 
 $ Version : The version of Visual Studio to generate project files for.
