Documentation for build tools

diff --git a/Engine/Documentation/Source/Programming/UnrealBuildSystem/TargetFiles/TargetFilesProperties/TargetFilesProperties.INT.udn b/Engine/Documentation/Source/Programming/UnrealBuildSystem/TargetFiles/TargetFilesProperties/TargetFilesProperties.INT.udn
--- a/Engine/Documentation/Source/Programming/UnrealBuildSystem/TargetFiles/TargetFilesProperties/TargetFilesProperties.INT.udn
+++ b/Engine/Documentation/Source/Programming/UnrealBuildSystem/TargetFiles/TargetFilesProperties/TargetFilesProperties.INT.udn
@@ -128,6 +128,8 @@ $ bEventDrivenLoader (Boolean): If true, event driven loader will be used in coo
 
 $ bUseXGEController (Boolean): Whether the XGE controller worker and modules should be included in the engine build. These are required for distributed shader compilation using the XGE interception interface.
 
+$ bIWYU (Boolean): Enables "include what you use" by default for modules in this target. Changes the default PCH mode for any module in this project to PCHUsageModule.UseExplicitOrSharedPCHs.
+
 $ bEnforceIWYU (Boolean): Enforce "include what you use" rules; warns if monolithic headers (Engine.h, UnrealEd.h, etc...) are used, and checks that source files include their matching header first.
 
 $ bHasExports (Boolean): Whether the final executable should export symbols.
