Documentation for build tools

diff --git a/Engine/Documentation/Source/Programming/UnrealBuildSystem/Configuration/BuildConfigProperties/BuildConfigProperties.INT.udn b/Engine/Documentation/Source/Programming/UnrealBuildSystem/Configuration/BuildConfigProperties/BuildConfigProperties.INT.udn
--- a/Engine/Documentation/Source/Programming/UnrealBuildSystem/Configuration/BuildConfigProperties/BuildConfigProperties.INT.udn
+++ b/Engine/Documentation/Source/Programming/UnrealBuildSystem/Configuration/BuildConfigProperties/BuildConfigProperties.INT.udn
@@ -162,14 +162,30 @@ $ bShowXGEMonitor : Whether to display the XGE build monitor.
 
 $ bStopXGECompilationAfterErrors : When enabled, XGE will stop compiling targets after a compile error occurs.  Recommended, as it saves computing resources for others.
 
+### CMakefileGenerator
+
+### CodeLiteGenerator
+
+### EddieProjectFileGenerator
+
+### KDevelopGenerator
+
+### MakefileGenerator
+
 ### ProjectFileGenerator
 
 $ Format : Default list of project file formats to generate
 
+### QMakefileGenerator
+
+### VSCodeProjectFileGenerator
+
 ### VCProjectFileGenerator
 
 $ Version : The version of Visual Studio to generate project files for.
 
+### XcodeProjectFileGenerator
+
 ### SourceFileWorkingSet
 
 $ Provider : Sets the provider to use for determining the working set
