Documentation for build tools

diff --git a/Engine/Documentation/Source/Programming/UnrealBuildSystem/Configuration/BuildConfigProperties/BuildConfigProperties.INT.udn b/Engine/Documentation/Source/Programming/UnrealBuildSystem/Configuration/BuildConfigProperties/BuildConfigProperties.INT.udn
--- a/Engine/Documentation/Source/Programming/UnrealBuildSystem/Configuration/BuildConfigProperties/BuildConfigProperties.INT.udn
+++ b/Engine/Documentation/Source/Programming/UnrealBuildSystem/Configuration/BuildConfigProperties/BuildConfigProperties.INT.udn
@@ -70,6 +70,10 @@ $ bUseIncrementalLinking : Whether to use incremental linking or not. Incrementa
 
 $ bAllowLTCG : Whether to allow the use of link time code generation (LTCG).
 
+$ bPGOProfile : Whether to enable Profile Guided Optimization (PGO) instrumentation in this build.
+
+$ bPGOOptimize : Whether to optimize this build with Profile Guided Optimization (PGO).
+
 $ bAllowASLRInShipping : Whether to allow the use of ASLR (address space layout randomization) if supported. Only applies to shipping builds.
 
 $ bSupportEditAndContinue : Whether to support edit and continue.  Only works on Microsoft compilers in 32-bit compiles.
diff --git a/Engine/Documentation/Source/Programming/UnrealBuildSystem/TargetFiles/TargetFilesProperties/TargetFilesProperties.INT.udn b/Engine/Documentation/Source/Programming/UnrealBuildSystem/TargetFiles/TargetFilesProperties/TargetFilesProperties.INT.udn
--- a/Engine/Documentation/Source/Programming/UnrealBuildSystem/TargetFiles/TargetFilesProperties/TargetFilesProperties.INT.udn
+++ b/Engine/Documentation/Source/Programming/UnrealBuildSystem/TargetFiles/TargetFilesProperties/TargetFilesProperties.INT.udn
@@ -188,6 +188,10 @@ $ bUseIncrementalLinking (Boolean): Whether to use incremental linking or not. I
 
 $ bAllowLTCG (Boolean): Whether to allow the use of link time code generation (LTCG).
 
+$ bPGOProfile (Boolean): Whether to enable Profile Guided Optimization (PGO) instrumentation in this build.
+
+$ bPGOOptimize (Boolean): Whether to optimize this build with Profile Guided Optimization (PGO).
+
 $ bAllowASLRInShipping (Boolean): Whether to allow the use of ASLR (address space layout randomization) if supported. Only applies to shipping builds.
 
 $ bSupportEditAndContinue (Boolean): Whether to support edit and continue.  Only works on Microsoft compilers in 32-bit compiles.
@@ -232,6 +236,10 @@ $ bFlushBuildDirOnRemoteMac (Boolean): Whether to clean Builds directory on a re
 
 $ bPrintToolChainTimingInfo (Boolean): Whether to write detailed timing info from the compiler and linker.
 
+$ bHideSymbolsByDefault (Boolean): Whether to hide symbols by default on POSIX platforms
+
+$ bAllowGeneratedIniWhenCooked (Boolean): Whether to load generated ini files in cooked build
+
 $ PCHOutputDirectory (String): The directory to put precompiled header files in. Experimental setting to allow using a path on a faster drive. Defaults to the standard output directory if not set.
 
 $ GlobalDefinitions (List&lt;String&gt;): Macros to define globally across the whole target.
