Documentation for build tools

diff --git a/Engine/Documentation/Source/Programming/UnrealBuildSystem/Configuration/BuildConfigProperties/BuildConfigProperties.INT.udn b/Engine/Documentation/Source/Programming/UnrealBuildSystem/Configuration/BuildConfigProperties/BuildConfigProperties.INT.udn
--- a/Engine/Documentation/Source/Programming/UnrealBuildSystem/Configuration/BuildConfigProperties/BuildConfigProperties.INT.udn
+++ b/Engine/Documentation/Source/Programming/UnrealBuildSystem/Configuration/BuildConfigProperties/BuildConfigProperties.INT.udn
@@ -26,6 +26,8 @@ $ bEditorDependsOnShaderCompileWorker : Whether to include a dependency on Shade
 
 $ bDebugBuildsActuallyUseDebugCRT : Enables the debug C++ runtime (CRT) for debug builds. By default we always use the release runtime, since the debug version isn't particularly useful when debugging Unreal Engine projects, and linking against the debug CRT libraries forces our third party library dependencies to also be compiled using the debug CRT (and often perform more slowly). Often it can be inconvenient to require a separate copy of the debug versions of third party static libraries simply so that you can debug your program's code.
 
+$ bUseInlining : Enable inlining for all modules.
+
 $ bUseXGEController : Whether the XGE controller worker and modules should be included in the engine build. These are required for distributed shader compilation using the XGE interception interface.
 
 $ bUseUnityBuild : Whether to unify C++ code into larger files for faster compilation.
diff --git a/Engine/Documentation/Source/Programming/UnrealBuildSystem/ModuleFiles/ModuleFilesProperties/ModuleFilesProperties.INT.udn b/Engine/Documentation/Source/Programming/UnrealBuildSystem/ModuleFiles/ModuleFilesProperties/ModuleFilesProperties.INT.udn
--- a/Engine/Documentation/Source/Programming/UnrealBuildSystem/ModuleFiles/ModuleFilesProperties/ModuleFilesProperties.INT.udn
+++ b/Engine/Documentation/Source/Programming/UnrealBuildSystem/ModuleFiles/ModuleFilesProperties/ModuleFilesProperties.INT.udn
@@ -83,7 +83,9 @@ $ PublicAdditionalShadowFiles (List&lt;String&gt;): For builds that execute on a
 
 $ PublicDelayLoadDLLs (List&lt;String&gt;): List of delay load DLLs - typically used for External (third party) modules
 
-$ Definitions (List&lt;String&gt;): Additional compiler definitions for this module
+$ PrivateDefinitions (List&lt;String&gt;): Private compiler definitions for this module
+
+$ PublicDefinitions (List&lt;String&gt;): Public compiler definitions for this module
 
 $ DynamicallyLoadedModuleNames (List&lt;String&gt;): Addition modules this module may require at run-time
 
diff --git a/Engine/Documentation/Source/Programming/UnrealBuildSystem/TargetFiles/TargetFilesProperties/TargetFilesProperties.INT.udn b/Engine/Documentation/Source/Programming/UnrealBuildSystem/TargetFiles/TargetFilesProperties/TargetFilesProperties.INT.udn
--- a/Engine/Documentation/Source/Programming/UnrealBuildSystem/TargetFiles/TargetFilesProperties/TargetFilesProperties.INT.udn
+++ b/Engine/Documentation/Source/Programming/UnrealBuildSystem/TargetFiles/TargetFilesProperties/TargetFilesProperties.INT.udn
@@ -94,6 +94,8 @@ $ bCompileSpeedTree (Boolean): Whether to compile SpeedTree support.
 
 $ bForceEnableExceptions (Boolean): Enable exceptions for all modules.
 
+$ bUseInlining (Boolean): Enable inlining for all modules.
+
 $ bForceEnableObjCExceptions (Boolean): Enable exceptions for all modules.
 
 $ bForceEnableRTTI (Boolean): Enable RTTI for all modules.
