Documentation for build tools

diff --git a/Engine/Documentation/Source/Programming/UnrealBuildSystem/ModuleFiles/ModuleFilesProperties/ModuleFilesProperties.INT.udn b/Engine/Documentation/Source/Programming/UnrealBuildSystem/ModuleFiles/ModuleFilesProperties/ModuleFilesProperties.INT.udn
--- a/Engine/Documentation/Source/Programming/UnrealBuildSystem/ModuleFiles/ModuleFilesProperties/ModuleFilesProperties.INT.udn
+++ b/Engine/Documentation/Source/Programming/UnrealBuildSystem/ModuleFiles/ModuleFilesProperties/ModuleFilesProperties.INT.udn
@@ -93,7 +93,7 @@ $ PlatformSpecificDynamicallyLoadedModuleNames (List&lt;String&gt;): Extra modul
 
 $ RuntimeDependencies (RuntimeDependencyList): List of files which this module depends on at runtime. These files will be staged along with the target.
 
-$ AdditionalPropertiesForReceipt (List&lt;ReceiptProperty&gt;): List of additional properties to be added to the build receipt
+$ AdditionalPropertiesForReceipt (ReceiptPropertyList): List of additional properties to be added to the build receipt
 
 $ PrecompileForTargets (PrecompileTargetsType): Which targets this module should be precompiled for
 
diff --git a/Engine/Documentation/Source/Programming/UnrealBuildSystem/TargetFiles/TargetFilesProperties/TargetFilesProperties.INT.udn b/Engine/Documentation/Source/Programming/UnrealBuildSystem/TargetFiles/TargetFilesProperties/TargetFilesProperties.INT.udn
--- a/Engine/Documentation/Source/Programming/UnrealBuildSystem/TargetFiles/TargetFilesProperties/TargetFilesProperties.INT.udn
+++ b/Engine/Documentation/Source/Programming/UnrealBuildSystem/TargetFiles/TargetFilesProperties/TargetFilesProperties.INT.udn
@@ -16,6 +16,8 @@ $ Architecture (String): Architecture that the target is being built for (or an
 
 $ ProjectFile (FileReference): Path to the project file for the project containing this target.
 
+$ Version (ReadOnlyBuildVersion): The current build version
+
 
 ### Read/Write Properties
 $ Type (TargetType): The type of target.
