Documentation for build tools

diff --git a/Engine/Documentation/Source/Programming/UnrealBuildSystem/Configuration/BuildConfigProperties/BuildConfigProperties.INT.udn b/Engine/Documentation/Source/Programming/UnrealBuildSystem/Configuration/BuildConfigProperties/BuildConfigProperties.INT.udn
--- a/Engine/Documentation/Source/Programming/UnrealBuildSystem/Configuration/BuildConfigProperties/BuildConfigProperties.INT.udn
+++ b/Engine/Documentation/Source/Programming/UnrealBuildSystem/Configuration/BuildConfigProperties/BuildConfigProperties.INT.udn
@@ -106,15 +106,13 @@ $ bFlushBuildDirOnRemoteMac : Whether to clean Builds directory on a remote Mac
 
 $ bPrintToolChainTimingInfo : Whether to write detailed timing info from the compiler and linker.
 
-$ PCHOutputDirectory : The directory to put precompiled header files in. Experimental setting to allow using a path on a faster drive. Defaults to the standard output directory if not set.
-
 $ bIgnoreOutdatedImportLibraries : Whether to ignore import library files that are out of date when building targets.  Set this to true to improve iteration time. By default we don't bother relinking targets if only a dependent .lib has changed, as chances are that the import library wasn't actually different unless a dependent header file of this target was also changed, in which case the target would be rebuilt automatically.
 
 $ bUseActionHistory : Whether to generate command line dependencies for compile actions when requested
 
 $ bPrintDebugInfo : Whether debug info should be written to the console.
 
-$ LogFilename : Allows logging to a file
+$ LogFileName : Specifies the file to use for logging
 
 $ bPrintPerformanceInfo : Prints performance diagnostics about include dependencies and other bits
 
@@ -126,8 +124,6 @@ $ bAllowSNDBS : Whether SN-DBS may be used.
 
 $ bShouldDeleteAllOutdatedProducedItems : Whether or not to delete outdated produced items.
 
-$ LogLevel : What level of logging we wish to show
-
 $ JsonExportFile : Whether we should export a JSON file containing detailed target information.
 
 $ bUseIncludeDependencyResolveCache : Whether the dependency cache includes pre-resolved include locations so UBT doesn't have to re-resolve each include location just to check the timestamp. This is technically not fully correct because the dependency cache is global and each module could have a different set of include paths that could cause headers to resolve files differently. In practice this is not the case, and significantly speeds up UBT when nothing is to be done.
diff --git a/Engine/Documentation/Source/Programming/UnrealBuildSystem/ModuleFiles/ModuleFilesProperties/ModuleFilesProperties.INT.udn b/Engine/Documentation/Source/Programming/UnrealBuildSystem/ModuleFiles/ModuleFilesProperties/ModuleFilesProperties.INT.udn
--- a/Engine/Documentation/Source/Programming/UnrealBuildSystem/ModuleFiles/ModuleFilesProperties/ModuleFilesProperties.INT.udn
+++ b/Engine/Documentation/Source/Programming/UnrealBuildSystem/ModuleFiles/ModuleFilesProperties/ModuleFilesProperties.INT.udn
@@ -19,6 +19,10 @@ $ ShortName (String): Specifies an alternate name for intermediate directories a
 
 $ PCHUsage (PCHUsageMode): Precompiled header usage for this module
 
+$ bTreatAsEngineModule (Boolean): Whether this module should be treated as an engine module (eg. using engine definitions, PCHs, compiled with optimizations enabled in DebugGame configurations, etc...). Initialized to a default based on the rules assembly it was created from.
+
+$ bUseBackwardsCompatibleDefaults (Boolean): Whether to use backwards compatible defaults for this module. By default, engine modules always use the latest default settings, while project modules do not (to support an easier migration path).
+
 $ bUseRTTI (Boolean): Use run time type information
 
 $ bUseAVX (Boolean): Use AVX instructions
@@ -51,6 +55,10 @@ $ bEnforceIWYU (Boolean): Enforce "include what you use" rules when PCHUsage is
 
 $ bAddDefaultIncludePaths (Boolean): Whether to add all the default include paths to the module (eg. the Source/Classes folder, subfolders under Source/Public).
 
+$ bPrecompile (Boolean): Whether this module should be precompiled. Defaults to the bPrecompile flag from the target. Clear this flag to prevent a module being precompiled.
+
+$ bUsePrecompiled (Boolean): Whether this module should use precompiled data. Always true for modules created from installed assemblies.
+
 $ PublicIncludePathModuleNames (List&lt;String&gt;): List of modules names (no path needed) with header files that our module's public headers needs access to, but we don't need to "import" or link against.
 
 $ PublicDependencyModuleNames (List&lt;String&gt;): List of public dependency module names (no path needed) (automatically does the private/public include). These are modules that are required by our public source files.
@@ -89,8 +97,6 @@ $ PublicDefinitions (List&lt;String&gt;): Public compiler definitions for this m
 
 $ DynamicallyLoadedModuleNames (List&lt;String&gt;): Addition modules this module may require at run-time
 
-$ PlatformSpecificDynamicallyLoadedModuleNames (List&lt;String&gt;): Extra modules this module may require at run time, that are on behalf of another platform (i.e. shader formats and the like)
-
 $ RuntimeDependencies (RuntimeDependencyList): List of files which this module depends on at runtime. These files will be staged along with the target.
 
 $ AdditionalPropertiesForReceipt (ReceiptPropertyList): List of additional properties to be added to the build receipt
@@ -99,4 +105,6 @@ $ PrecompileForTargets (PrecompileTargetsType): Which targets this module should
 
 $ ExternalDependencies (List&lt;String&gt;): External files which invalidate the makefile if modified. Relative paths are resolved relative to the .build.cs file.
 
+$ bLegacyPublicIncludePaths (Nullable&lt;Boolean&gt;): Whether this module qualifies included headers from other modules relative to the root of their 'Public' folder. This reduces the number of search paths that have to be passed to the compiler, improving performance and reducing the length of the compiler command line.
+
 
diff --git a/Engine/Documentation/Source/Programming/UnrealBuildSystem/TargetFiles/TargetFilesProperties/TargetFilesProperties.INT.udn b/Engine/Documentation/Source/Programming/UnrealBuildSystem/TargetFiles/TargetFilesProperties/TargetFilesProperties.INT.udn
--- a/Engine/Documentation/Source/Programming/UnrealBuildSystem/TargetFiles/TargetFilesProperties/TargetFilesProperties.INT.udn
+++ b/Engine/Documentation/Source/Programming/UnrealBuildSystem/TargetFiles/TargetFilesProperties/TargetFilesProperties.INT.udn
@@ -132,6 +132,8 @@ $ bEventDrivenLoader (Boolean): If true, event driven loader will be used in coo
 
 $ bUseXGEController (Boolean): Whether the XGE controller worker and modules should be included in the engine build. These are required for distributed shader compilation using the XGE interception interface.
 
+$ bUseBackwardsCompatibleDefaults (Boolean): Whether to use backwards compatible defaults for this module. By default, engine modules always use the latest default settings, while project modules do not (to support an easier migration path).
+
 $ bIWYU (Boolean): Enables "include what you use" by default for modules in this target. Changes the default PCH mode for any module in this project to PCHUsageModule.UseExplicitOrSharedPCHs.
 
 $ bEnforceIWYU (Boolean): Enforce "include what you use" rules; warns if monolithic headers (Engine.h, UnrealEd.h, etc...) are used, and checks that source files include their matching header first.
@@ -244,10 +246,12 @@ $ bHideSymbolsByDefault (Boolean): Whether to hide symbols by default on POSIX p
 
 $ bAllowGeneratedIniWhenCooked (Boolean): Whether to load generated ini files in cooked build
 
-$ PCHOutputDirectory (String): The directory to put precompiled header files in. Experimental setting to allow using a path on a faster drive. Defaults to the standard output directory if not set.
+$ bLegacyPublicIncludePaths (Boolean): Add all the public folders as include paths for the compile environment.
 
 $ GlobalDefinitions (List&lt;String&gt;): Macros to define globally across the whole target.
 
+$ ProjectDefinitions (List&lt;String&gt;): Macros to define across all macros in the project.
+
 $ ExtraModuleNames (List&lt;String&gt;): List of additional modules to be compiled into the target.
 
 $ BuildEnvironment (TargetBuildEnvironment): Specifies the build environment for this target. See TargetBuildEnvironment for more infomation on the available options.
@@ -256,6 +260,4 @@ $ PreBuildSteps (List&lt;String&gt;): Specifies a list of steps which should be
 
 $ PostBuildSteps (List&lt;String&gt;): Specifies a list of steps which should be executed after this target is built, in the context of the host platform's shell. The following variables will be expanded before execution: $(EngineDir), $(ProjectDir), $(TargetName), $(TargetPlatform), $(TargetConfiguration), $(TargetType), $(ProjectFile).
 
-$ OverrideExecutableFileExtension (String): Can be set to override the file extension of the executable file (normally .exe or .dll on Windows, for example)
-
 
