Documentation for build tools

diff --git a/Engine/Documentation/Source/Programming/UnrealBuildSystem/TargetFiles/TargetFilesProperties/TargetFilesProperties.INT.udn b/Engine/Documentation/Source/Programming/UnrealBuildSystem/TargetFiles/TargetFilesProperties/TargetFilesProperties.INT.udn
--- a/Engine/Documentation/Source/Programming/UnrealBuildSystem/TargetFiles/TargetFilesProperties/TargetFilesProperties.INT.udn
+++ b/Engine/Documentation/Source/Programming/UnrealBuildSystem/TargetFiles/TargetFilesProperties/TargetFilesProperties.INT.udn
@@ -244,6 +244,8 @@ $ bPrintToolChainTimingInfo (Boolean): Whether to write detailed timing info fro
 
 $ bHideSymbolsByDefault (Boolean): Whether to hide symbols by default on POSIX platforms
 
+$ ToolChainName (String): Allows overriding the toolchain to be created for this target. This must match the name of a class declared in the UnrealBuildTool assembly.
+
 $ bAllowGeneratedIniWhenCooked (Boolean): Whether to load generated ini files in cooked build
 
 $ bLegacyPublicIncludePaths (Boolean): Add all the public folders as include paths for the compile environment.
