Documentation for build tools

diff --git a/Engine/Documentation/Source/Programming/Development/BuildGraph/ScriptAnatomy/Tasks/PredefinedTasks/BuildGraphScriptPredefinedTasks.INT.udn b/Engine/Documentation/Source/Programming/Development/BuildGraph/ScriptAnatomy/Tasks/PredefinedTasks/BuildGraphScriptPredefinedTasks.INT.udn
--- a/Engine/Documentation/Source/Programming/Development/BuildGraph/ScriptAnatomy/Tasks/PredefinedTasks/BuildGraphScriptPredefinedTasks.INT.udn
+++ b/Engine/Documentation/Source/Programming/Development/BuildGraph/ScriptAnatomy/Tasks/PredefinedTasks/BuildGraphScriptPredefinedTasks.INT.udn
@@ -129,6 +129,19 @@ Moves files from one directory to another.
 | Tag   | Tag List  | Optional | Tag to be applied to build products of this task                                  |
 | If    | Condition | Optional | Whether to execute this task. It is ignored if this condition evaluates to false. |
 
+### MsBuild
+
+Executes MsBuild
+
+|               |           |          |                                                                                                               |
+| ------------- | --------- | -------- | ------------------------------------------------------------------------------------------------------------- |
+| Project       | String    | Required | The C# project file to be compile. More than one project file can be specified by separating with semicolons. |
+| Configuration | String    | Optional | The configuration to compile                                                                                  |
+| Platform      | String    | Optional | The platform to compile                                                                                       |
+| Arguments     | String    | Optional | Additional options to pass to the compiler                                                                    |
+| Verbosity     | String    | Optional | The MSBuild output verbosity                                                                                  |
+| If            | Condition | Optional | Whether to execute this task. It is ignored if this condition evaluates to false.                             |
+
 ### PakFile
 
 Creates a PAK file from a given set of files.
