Documentation for build tools

diff --git a/Engine/Documentation/Source/Programming/UnrealBuildSystem/Configuration/BuildConfigProperties/BuildConfigProperties.INT.udn b/Engine/Documentation/Source/Programming/UnrealBuildSystem/Configuration/BuildConfigProperties/BuildConfigProperties.INT.udn
--- a/Engine/Documentation/Source/Programming/UnrealBuildSystem/Configuration/BuildConfigProperties/BuildConfigProperties.INT.udn
+++ b/Engine/Documentation/Source/Programming/UnrealBuildSystem/Configuration/BuildConfigProperties/BuildConfigProperties.INT.udn
@@ -166,28 +166,88 @@ $ bStopXGECompilationAfterErrors : When enabled, XGE will stop compiling targets
 
 ### CMakefileGenerator
 
+$ bIncludeDocumentation : True if we should include documentation in the generated projects
+
+$ MasterProjectName : Name of the master project file (e.g. base file name for the solution file for Visual Studio, or the Xcode project file on Mac)
+
+$ bMasterProjectNameFromFolder : If true, sets the master project name according to the name of the folder it's in
+
 ### CodeLiteGenerator
 
+$ bIncludeDocumentation : True if we should include documentation in the generated projects
+
+$ MasterProjectName : Name of the master project file (e.g. base file name for the solution file for Visual Studio, or the Xcode project file on Mac)
+
+$ bMasterProjectNameFromFolder : If true, sets the master project name according to the name of the folder it's in
+
 ### EddieProjectFileGenerator
 
+$ bIncludeDocumentation : True if we should include documentation in the generated projects
+
+$ MasterProjectName : Name of the master project file (e.g. base file name for the solution file for Visual Studio, or the Xcode project file on Mac)
+
+$ bMasterProjectNameFromFolder : If true, sets the master project name according to the name of the folder it's in
+
 ### KDevelopGenerator
 
+$ bIncludeDocumentation : True if we should include documentation in the generated projects
+
+$ MasterProjectName : Name of the master project file (e.g. base file name for the solution file for Visual Studio, or the Xcode project file on Mac)
+
+$ bMasterProjectNameFromFolder : If true, sets the master project name according to the name of the folder it's in
+
 ### MakefileGenerator
 
+$ bIncludeDocumentation : True if we should include documentation in the generated projects
+
+$ MasterProjectName : Name of the master project file (e.g. base file name for the solution file for Visual Studio, or the Xcode project file on Mac)
+
+$ bMasterProjectNameFromFolder : If true, sets the master project name according to the name of the folder it's in
+
 ### ProjectFileGenerator
 
 $ Format : Default list of project file formats to generate
 
+$ bIncludeDocumentation : True if we should include documentation in the generated projects
+
+$ MasterProjectName : Name of the master project file (e.g. base file name for the solution file for Visual Studio, or the Xcode project file on Mac)
+
+$ bMasterProjectNameFromFolder : If true, sets the master project name according to the name of the folder it's in
+
 ### QMakefileGenerator
 
+$ bIncludeDocumentation : True if we should include documentation in the generated projects
+
+$ MasterProjectName : Name of the master project file (e.g. base file name for the solution file for Visual Studio, or the Xcode project file on Mac)
+
+$ bMasterProjectNameFromFolder : If true, sets the master project name according to the name of the folder it's in
+
 ### VSCodeProjectFileGenerator
 
+$ bIncludeDocumentation : True if we should include documentation in the generated projects
+
+$ MasterProjectName : Name of the master project file (e.g. base file name for the solution file for Visual Studio, or the Xcode project file on Mac)
+
+$ bMasterProjectNameFromFolder : If true, sets the master project name according to the name of the folder it's in
+
 ### VCProjectFileGenerator
 
 $ Version : The version of Visual Studio to generate project files for.
 
+$ bIncludeDocumentation : True if we should include documentation in the generated projects
+
+$ MasterProjectName : Name of the master project file (e.g. base file name for the solution file for Visual Studio, or the Xcode project file on Mac)
+
+$ bMasterProjectNameFromFolder : If true, sets the master project name according to the name of the folder it's in
+
 ### XcodeProjectFileGenerator
 
+$ bIncludeDocumentation : True if we should include documentation in the generated projects
+
+$ MasterProjectName : Name of the master project file (e.g. base file name for the solution file for Visual Studio, or the Xcode project file on Mac)
+
+$ bMasterProjectNameFromFolder : If true, sets the master project name according to the name of the folder it's in
+
 ### SourceFileWorkingSet
 
 $ Provider : Sets the provider to use for determining the working set
@@ -248,6 +308,10 @@ $ RSyncUsername : Username on the remote machine to connect to with RSync
 
 $ Compiler : Version of the compiler toolchain to use on Windows platform. A value of "default" will be changed to a specific version at UBT startup.
 
+$ CompilerVersion : The specific toolchain version to use. This may be a specific version number (eg. "14.13.26128") or the string "Latest" to select the newest available version. By default, we use the toolchain version indicated by WindowsPlatform.DefaultToolChainVersion if it is available, or the latest version otherwise.
+
+$ WindowsSdkVersion : The specific Windows SDK version to use. This may be a specific version number (eg. "8.1", "10.0", or "10.0.10150.0") or the string "Latest" to select the newest available version. By default, we use the Windows SDK version indicated by WindowsPlatform.DefaultWindowsSdkVersion if it is available, or the latest version otherwise.
+
 $ StaticAnalyzer : The static analyzer to use
 
 $ bStrictConformanceMode : Enables strict standard conformance mode (/permissive-) in VS2017+.
diff --git a/Engine/Documentation/Source/Programming/UnrealBuildSystem/TargetFiles/TargetFilesProperties/TargetFilesProperties.INT.udn b/Engine/Documentation/Source/Programming/UnrealBuildSystem/TargetFiles/TargetFilesProperties/TargetFilesProperties.INT.udn
--- a/Engine/Documentation/Source/Programming/UnrealBuildSystem/TargetFiles/TargetFilesProperties/TargetFilesProperties.INT.udn
+++ b/Engine/Documentation/Source/Programming/UnrealBuildSystem/TargetFiles/TargetFilesProperties/TargetFilesProperties.INT.udn
@@ -218,7 +218,7 @@ $ bCheckLicenseViolations (Boolean): Tells the UBT to check if module currently
 
 $ bBreakBuildOnLicenseViolation (Boolean): Tells the UBT to break build if module currently being built is violating EULA.
 
-$ bUseFastPDBLinking (Boolean): Whether to use the :FASTLINK option when building with /DEBUG to create local PDBs on Windows. Fast, but currently seems to have problems finding symbols in the debugger.
+$ bUseFastPDBLinking (Nullable&lt;Boolean&gt;): Whether to use the :FASTLINK option when building with /DEBUG to create local PDBs on Windows. Fast, but currently seems to have problems finding symbols in the debugger.
 
 $ bCreateMapFile (Boolean): Outputs a map file as part of the build.
 
@@ -250,16 +250,22 @@ $ bAllowGeneratedIniWhenCooked (Boolean): Whether to load generated ini files in
 
 $ bLegacyPublicIncludePaths (Boolean): Add all the public folders as include paths for the compile environment.
 
+$ BuildVersion (String): The build version string
+
 $ GlobalDefinitions (List&lt;String&gt;): Macros to define globally across the whole target.
 
 $ ProjectDefinitions (List&lt;String&gt;): Macros to define across all macros in the project.
 
 $ ExtraModuleNames (List&lt;String&gt;): List of additional modules to be compiled into the target.
 
-$ BuildEnvironment (TargetBuildEnvironment): Specifies the build environment for this target. See TargetBuildEnvironment for more infomation on the available options.
+$ BuildEnvironment (TargetBuildEnvironment): Specifies the build environment for this target. See TargetBuildEnvironment for more information on the available options.
 
 $ PreBuildSteps (List&lt;String&gt;): Specifies a list of steps which should be executed before this target is built, in the context of the host platform's shell. The following variables will be expanded before execution: $(EngineDir), $(ProjectDir), $(TargetName), $(TargetPlatform), $(TargetConfiguration), $(TargetType), $(ProjectFile).
 
 $ PostBuildSteps (List&lt;String&gt;): Specifies a list of steps which should be executed after this target is built, in the context of the host platform's shell. The following variables will be expanded before execution: $(EngineDir), $(ProjectDir), $(TargetName), $(TargetPlatform), $(TargetConfiguration), $(TargetType), $(ProjectFile).
 
+$ AdditionalCompilerArguments (String): Additional arguments to pass to the compiler
+
+$ AdditionalLinkerArguments (String): Additional arguments to pass to the linker
+
 
