Documentation for build tools

diff --git a/Engine/Documentation/Source/Programming/UnrealBuildSystem/Configuration/BuildConfigProperties/BuildConfigProperties.INT.udn b/Engine/Documentation/Source/Programming/UnrealBuildSystem/Configuration/BuildConfigProperties/BuildConfigProperties.INT.udn
--- a/Engine/Documentation/Source/Programming/UnrealBuildSystem/Configuration/BuildConfigProperties/BuildConfigProperties.INT.udn
+++ b/Engine/Documentation/Source/Programming/UnrealBuildSystem/Configuration/BuildConfigProperties/BuildConfigProperties.INT.udn
@@ -164,6 +164,14 @@ $ bShowXGEMonitor : Whether to display the XGE build monitor.
 
 $ bStopXGECompilationAfterErrors : When enabled, XGE will stop compiling targets after a compile error occurs.  Recommended, as it saves computing resources for others.
 
+### CLionGenerator
+
+$ bIncludeDocumentation : True if we should include documentation in the generated projects
+
+$ MasterProjectName : Name of the master project file (e.g. base file name for the solution file for Visual Studio, or the Xcode project file on Mac)
+
+$ bMasterProjectNameFromFolder : If true, sets the master project name according to the name of the folder it's in
+
 ### CMakefileGenerator
 
 $ bIncludeDocumentation : True if we should include documentation in the generated projects
