Documentation for build tools

diff --git a/Engine/Documentation/Source/Programming/Development/BuildGraph/ScriptAnatomy/Tasks/PredefinedTasks/BuildGraphScriptPredefinedTasks.INT.udn b/Engine/Documentation/Source/Programming/Development/BuildGraph/ScriptAnatomy/Tasks/PredefinedTasks/BuildGraphScriptPredefinedTasks.INT.udn
--- a/Engine/Documentation/Source/Programming/Development/BuildGraph/ScriptAnatomy/Tasks/PredefinedTasks/BuildGraphScriptPredefinedTasks.INT.udn
+++ b/Engine/Documentation/Source/Programming/Development/BuildGraph/ScriptAnatomy/Tasks/PredefinedTasks/BuildGraphScriptPredefinedTasks.INT.udn
@@ -73,13 +73,14 @@ Cook a selection of maps for a certain platform
 
 Copies files from one directory to another.
 
-|       |           |          |                                                                                   |
-| ----- | --------- | -------- | --------------------------------------------------------------------------------- |
-| Files | File Spec | Optional | Filter to be applied to the list of input files. Optional.                        |
-| From  | File Spec | Required | The pattern(s) to copy from (eg. Engine/*.txt)                                    |
-| To    | File Spec | Required | The directory or to copy to                                                       |
-| Tag   | Tag List  | Optional | Tag to be applied to build products of this task                                  |
-| If    | Condition | Optional | Whether to execute this task. It is ignored if this condition evaluates to false. |
+|           |           |          |                                                                                   |
+| --------- | --------- | -------- | --------------------------------------------------------------------------------- |
+| Files     | File Spec | Optional | Filter to be applied to the list of input files. Optional.                        |
+| From      | File Spec | Required | The pattern(s) to copy from (eg. Engine/*.txt)                                    |
+| To        | File Spec | Required | The directory or to copy to                                                       |
+| Overwrite | Boolean   | Optional | Whether or not to overwrite existing files                                        |
+| Tag       | Tag List  | Optional | Tag to be applied to build products of this task                                  |
+| If        | Condition | Optional | Whether to execute this task. It is ignored if this condition evaluates to false. |
 
 ### CsCompile
 
diff --git a/Engine/Documentation/Source/Programming/UnrealBuildSystem/Configuration/BuildConfigProperties/BuildConfigProperties.INT.udn b/Engine/Documentation/Source/Programming/UnrealBuildSystem/Configuration/BuildConfigProperties/BuildConfigProperties.INT.udn
--- a/Engine/Documentation/Source/Programming/UnrealBuildSystem/Configuration/BuildConfigProperties/BuildConfigProperties.INT.udn
+++ b/Engine/Documentation/Source/Programming/UnrealBuildSystem/Configuration/BuildConfigProperties/BuildConfigProperties.INT.udn
@@ -18,12 +18,20 @@ $ bEnableThreadSanitizer : Enables thread sanitizer (TSan)
 
 $ bEnableUndefinedBehaviorSanitizer : Enables undefined behavior sanitizer (UBSan)
 
+$ bWriteSolutionOptionFile : Whether to write a solution option (suo) file for the sln
+
 $ bAddFastPDBToProjects : Whether to add the -FastPDB option to build command lines by default
 
 $ bUsePerFileIntellisense : Whether to generate per-file intellisense data
 
 $ bEditorDependsOnShaderCompileWorker : Whether to include a dependency on ShaderCompileWorker when generating project files for the editor.
 
+$ ProcessorCountMultiplier : Processor count multiplier for local execution. Can be below 1 to reserve CPU for other tasks. When using the local executor (not XGE), run a single action on each CPU core.  Note that you can set this to a larger value to get slightly faster build times in many cases, but your computer's responsiveness during compiling may be much worse.
+
+$ MaxProcessorCount : Maximum processor count for local execution.
+
+$ bStopCompilationAfterErrors : When enabled, will stop compiling targets after a compile error occurs.
+
 $ bDebugBuildsActuallyUseDebugCRT : Enables the debug C++ runtime (CRT) for debug builds. By default we always use the release runtime, since the debug version isn't particularly useful when debugging Unreal Engine projects, and linking against the debug CRT libraries forces our third party library dependencies to also be compiled using the debug CRT (and often perform more slowly). Often it can be inconvenient to require a separate copy of the debug versions of third party static libraries simply so that you can debug your program's code.
 
 $ bUseInlining : Enable inlining for all modules.
@@ -142,6 +150,8 @@ $ bUseUBTMakefiles : Enables support for very fast iterative builds by caching t
 
 $ bAllowDistcc : Whether DMUCS/Distcc may be used. Distcc requires some setup - so by default disable it so we don't break local or remote building
 
+$ bAllowParallelExecutor : Whether to allow using parallel executor on Windows.
+
 $ bIgnoreJunk : Whether to skip checking for files identified by the junk manifest
 
 $ bAllowDistccLocalFallback : When enabled allows DMUCS/Distcc to fallback to local compilation when remote compiling fails. Defaults to true as separation of pre-process and compile stages can introduce non-fatal errors.
@@ -154,10 +164,6 @@ $ DMUCSCoordinator : DMUCS coordinator hostname or IP address.
 
 $ DMUCSDistProp : The DMUCS distinguishing property value to request for compile hosts.
 
-$ ProcessorCountMultiplier : Processor count multiplier for local execution. Can be below 1 to reserve CPU for other tasks. When using the local executor (not XGE), run a single action on each CPU core.  Note that you can set this to a larger value to get slightly faster build times in many cases, but your computer's responsiveness during compiling may be much worse.
-
-$ MaxProcessorCount : Maximum processor count for local execution.
-
 $ bXGENoWatchdogThread : Whether to use the no_watchdog_thread option to prevent VS2015 toolchain stalls.
 
 $ bShowXGEMonitor : Whether to display the XGE build monitor.
@@ -264,6 +270,16 @@ $ RepositoryPath : Sets the path to use for the repository. Interpreted relative
 
 $ GitPath : Sets the path to use for the Git executable. Defaults to "git" (assuming it's in the PATH).
 
+### VCMacProjectFileGenerator
+
+$ Version : The version of Visual Studio to generate project files for.
+
+$ bIncludeDocumentation : True if we should include documentation in the generated projects
+
+$ MasterProjectName : Name of the master project file (e.g. base file name for the solution file for Visual Studio, or the Xcode project file on Mac)
+
+$ bMasterProjectNameFromFolder : If true, sets the master project name according to the name of the folder it's in
+
 ### IOSToolChain
 
 $ IOSSDKVersion : Which version of the iOS SDK to target at build time
diff --git a/Engine/Documentation/Source/Programming/UnrealBuildSystem/TargetFiles/TargetFilesProperties/TargetFilesProperties.INT.udn b/Engine/Documentation/Source/Programming/UnrealBuildSystem/TargetFiles/TargetFilesProperties/TargetFilesProperties.INT.udn
--- a/Engine/Documentation/Source/Programming/UnrealBuildSystem/TargetFiles/TargetFilesProperties/TargetFilesProperties.INT.udn
+++ b/Engine/Documentation/Source/Programming/UnrealBuildSystem/TargetFiles/TargetFilesProperties/TargetFilesProperties.INT.udn
@@ -40,6 +40,8 @@ $ bBuildAllPlugins (Boolean): Build all the plugins that we can find, even if th
 
 $ AdditionalPlugins (List&lt;String&gt;): A list of additional plugins which need to be included in this target. This allows referencing non-optional plugin modules which cannot be disabled, and allows building against specific modules in program targets which do not fit the categories in ModuleHostType.
 
+$ ExcludePlugins (List&lt;String&gt;): List of plugins to be excluded from this target. Note that the project file may still reference them, so they should be marked as optional to avoid failing to find them at runtime.
+
 $ PakSigningKeysFile (String): Path to the set of pak signing keys to embed in the executable.
 
 $ SolutionDirectory (String): Allows a Program Target to specify it's own solution folder path.
