Copying //UE4/Dev-Documentation to //UE4/Samples-Main (Source: //UE4/Dev-Documentation @ 4091886)
============================
  MAJOR FEATURES & CHANGES
============================

Change 3725068 by [REDACTED_PERSON]

	Fixed typo

Change 3809004 by [REDACTED_PERSON]

	mesh Distance Field missing updates for cms copy.

Change 3809144 by [REDACTED_PERSON]

	Updated file against INTSourceChangelist:3599948

Change 3809151 by [REDACTED_PERSON]

	Updated file against INTSourceChangelist:3151855

Change 3809156 by [REDACTED_PERSON]

	Updated file against INTSourceChangelist:3151855

Change 3809164 by [REDACTED_PERSON]

	Updated file against INTSourceChangelist:3717114

Change 3809166 by [REDACTED_PERSON]

	Updated file against INTSourceChangelist:3151855

Change 3809213 by [REDACTED_PERSON]

	Updated file against INTSourceChangelist:3717114

Change 3809217 by [REDACTED_PERSON]

	Updated fie against INTSourceChangelist:3782314

Change 3809225 by [REDACTED_PERSON]

	Updated file against INTSourceChangelist:3151855

Change 3809238 by [REDACTED_PERSON]

	Updated file against INTSourceChangelist:3151855

Change 3809243 by [REDACTED_PERSON]

	Updated file against INTSourceChangelist:3151855

Change 3809247 by [REDACTED_PERSON]

	Updated file against INTSourceChangelist:3151855

Change 3809606 by [REDACTED_PERSON]

	More fixes for migration

Change 3809611 by [REDACTED_PERSON]

	Fixed incorrect syntax

Change 3809644 by [REDACTED_PERSON]

	Adding XML templates

Change 3811009 by [REDACTED_PERSON]

	Updated file against INTSourceChangelist:3718347

Change 3811010 by [REDACTED_PERSON]

	Updated file against INTSourceChangelist:3810557

Change 3811011 by [REDACTED_PERSON]

	Error fixed (column alignment error)

Change 3811015 by [REDACTED_PERSON]

	Updated file against INTSourceChangelist:3810557

Change 3811017 by [REDACTED_PERSON]

	Fixed image missing error (by updating DoNet)

Change 3811026 by [REDACTED_PERSON]

	Updated file against INTSourceChangelist:3810557

Change 3811032 by [REDACTED_PERSON]

	Updated file against INTSourceChangelist:3810557

Change 3811043 by [REDACTED_PERSON]

	Updated file against INTSourceChangelist:3782314

Change 3811054 by [REDACTED_PERSON]

	Updated file against INTSourceChangelist:3726228

Change 3811063 by [REDACTED_PERSON]

	Updated file against INTSourceChangelist:3810557

Change 3811070 by [REDACTED_PERSON]

	Updated file against INTSourceChangelist:3810467

Change 3811079 by [REDACTED_PERSON]

	Updated file against INTSourceChangelist:3810557

Change 3811080 by [REDACTED_PERSON]

	Updated file against INTSourceChangelist:3782314

Change 3811083 by [REDACTED_PERSON]

	Updated file against INTSourceChangelist:3735117

Change 3811090 by [REDACTED_PERSON]

	Updated file against INTSourceChangelist:3801260

Change 3811092 by [REDACTED_PERSON]

	Fixed error

Change 3811098 by [REDACTED_PERSON]

	Updated file against INTSourceChangelist:3782314

Change 3811099 by [REDACTED_PERSON]

	Updated file against INTSourceChangelist:3782314

Change 3811101 by [REDACTED_PERSON]

	Updated file against INTSourceChangelist:3782314

Change 3811150 by [REDACTED_PERSON]

	#loc UE4DocJpn File updated against INTSourceChangelist:3782314

Change 3811151 by [REDACTED_PERSON]

	#loc UE4DocCHN moved or deleted

Change 3811152 by [REDACTED_PERSON]

	#loc UE4DocJpn File updated against 3782314

Change 3811157 by [REDACTED_PERSON]

	#loc UE4DocJpn File updated against INTSourceChangelist:3782314

Change 3811160 by [REDACTED_PERSON]

	#loc UE4DocJpn File updated against INTSourceChangelist:3782314

Change 3811161 by [REDACTED_PERSON]

	#loc UE4DocJpn File updated against INTSourceChangelist:3782314

Change 3811164 by [REDACTED_PERSON]

	#loc UE4DocJpn File updated against INTSourceChangelist:3782314

Change 3812250 by [REDACTED_PERSON]

	#loc UE4DocJpn File updated against INTSourceChangelist:3782314

Change 3812252 by [REDACTED_PERSON]

	#loc UE4DocJpn File updated against INTSourceChangelist:3782314

Change 3812256 by [REDACTED_PERSON]

	#loc UE4DocJpn File updated against INTSourceChangelist:3810467

Change 3812258 by [REDACTED_PERSON]

	#loc UE4DocJpn File updated against INTSourceChangelist:3782314

Change 3812262 by [REDACTED_PERSON]

	#loc UE4DocJpn File updated against INTSourceChangelist:3782314

Change 3812265 by [REDACTED_PERSON]

	#loc UE4DocJpn File updated against INTSourceChangelist:3782314

Change 3812268 by [REDACTED_PERSON]

	#loc UE4DocJpn File updated against INTSourceChangelist:3782314

Change 3812270 by [REDACTED_PERSON]

	#loc UE4DocJpn File updated against INTSourceChangelist:3782314

Change 3812276 by [REDACTED_PERSON]

	#loc UE4DocJpn File updated against INTSourceChangelist:3782314

Change 3812280 by [REDACTED_PERSON]

	#loc UE4DocJpn File updated against INTSourceChangelist:3782314

Change 3812281 by [REDACTED_PERSON]

	#loc UE4DocJpn File updated against INTSourceChangelist:3782314

Change 3812289 by [REDACTED_PERSON]

	#loc UE4DocJpn File updated against INTSourceChangelist:3782314

Change 3812296 by [REDACTED_PERSON]

	#loc UE4DocJpn File updated against INTSourceChangelist:3782314

Change 3812306 by [REDACTED_PERSON]

	#loc UE4DocJpn File updated against INTSourceChangelist:3782314

Change 3812309 by [REDACTED_PERSON]

	#loc UE4DocJpn File updated against INTSourceChangelist:3782314

Change 3812326 by [REDACTED_PERSON]

	Updated file against INTSourceChangelist:3810557

Change 3812498 by [REDACTED_PERSON]

	Updated file agaisnt INTSourceChangelist:3678273

Change 3813612 by [REDACTED_PERSON]

	Minor translation change

Change 3813654 by [REDACTED_PERSON]

	#loc UE4DocJpn File updated against INTSourceChangelist:3782314

Change 3813663 by [REDACTED_PERSON]

	#loc UE4DocJpn File updated against INTSourceChangelist:3782314

Change 3813668 by [REDACTED_PERSON]

	#loc UE4DocJpn File updated against INTSourceChangelist:3782314

Change 3813686 by [REDACTED_PERSON]

	#loc UE4DocJpn File updated against INTSourceChangelist:3782314

Change 3813687 by [REDACTED_PERSON]

	#loc UE4DocJpn File updated against INTSourceChangelist:3810467

Change 3813711 by [REDACTED_PERSON]

	#loc UE4DocJpn File updated against INTSourceChangelist:3782314

Change 3813725 by [REDACTED_PERSON]

	#loc UE4DocJpn File updated against INTSourceChangelist:3782314

Change 3813730 by [REDACTED_PERSON]

	#loc UE4DocJpn File updated against INTSourceChangelist:3782314

Change 3813732 by [REDACTED_PERSON]

	#loc UE4DocJpn File updated against INTSourceChangelist:3782314

Change 3813734 by [REDACTED_PERSON]

	#loc UE4DocJpn File updated against INTSourceChangelist:3782314

Change 3813737 by [REDACTED_PERSON]

	#loc UE4DocJpn File updated against INTSourceChangelist:3782314

Change 3813740 by [REDACTED_PERSON]

	#loc UE4DocJpn File updated against INTSourceChangelist:3782314

Change 3813741 by [REDACTED_PERSON]

	#loc UE4DocJpn File updated against INTSourceChangelist:3782314

Change 3813742 by [REDACTED_PERSON]

	#loc UE4DocJpn File updated against INTSourceChangelist:3782314

Change 3813748 by [REDACTED_PERSON]

	#loc UE4DocJpn File updated against INTSourceChangelist:3782314

Change 3813800 by [REDACTED_PERSON]

	#loc UE4DocKOR

Change 3814338 by [REDACTED_PERSON]

	#loc UE4DocJpn File updated against INTSourceChangelist:3810467

Change 3814339 by [REDACTED_PERSON]

	#loc UE4DocJpn File updated against INTSourceChangelist:3807090

Change 3814340 by [REDACTED_PERSON]

	#loc UE4DocJpn File updated against INTSourceChangelist:3807090

Change 3814342 by [REDACTED_PERSON]

	#loc UE4DocJpn File updated against INTSourceChangelist:3782314

Change 3814354 by [REDACTED_PERSON]

	INTSourceChangelist:3810557

Change 3814355 by [REDACTED_PERSON]

	#loc UE4DocJpn File updated against INTSourceChangelist:3810557

Change 3814356 by [REDACTED_PERSON]

	#loc UE4DocJpn File updated against INTSourceChangelist:3810557

Change 3814357 by [REDACTED_PERSON]

	#loc UE4DocJpn File updated against INTSourceChangelist:3810557

Change 3814359 by [REDACTED_PERSON]

	#loc UE4DocJpn File updated against INTSourceChangelist:3810557

Change 3814776 by [REDACTED_PERSON]

	#loc UE4DocJpn File updated against INTSourceChangelist:3801260

Change 3814780 by [REDACTED_PERSON]

	#loc UE4DocJpn File updated against INTSourceChangelist:3801260

Change 3814784 by [REDACTED_PERSON]

	#loc UE4DocJpn File updated against INTSourceChangelist:3782314

Change 3814786 by [REDACTED_PERSON]

	#loc UE4DocJpn File updated against INTSourceChangelist:3782314

Change 3814787 by [REDACTED_PERSON]

	#loc UE4DocJpn File updated against INTSourceChangelist:3782314

Change 3814789 by [REDACTED_PERSON]

	#loc UE4DocJpn File updated against INTSourceChangelist:3782314

Change 3814790 by [REDACTED_PERSON]

	#loc UE4DocJpn File updated against INTSourceChangelist:3782314

Change 3814792 by [REDACTED_PERSON]

	#loc UE4DocJpn File updated against INTSourceChangelist:3782314

Change 3814794 by [REDACTED_PERSON]

	#loc UE4DocJpn File updated against INTSourceChangelist:3782314

Change 3814798 by [REDACTED_PERSON]

	#loc UE4DocJpn File updated against INTSourceChangelist:3782314

Change 3814799 by [REDACTED_PERSON]

	#loc UE4DocJpn File updated against INTSourceChangelist:3782314

Change 3814801 by [REDACTED_PERSON]

	#loc UE4DocJpn File updated against INTSourceChangelist:3782314

Change 3814808 by [REDACTED_PERSON]

	#loc UE4DocJpn File updated against INTSourceChangelist:3810557

Change 3814814 by [REDACTED_PERSON]

	#loc UE4DocJpn File updated against INTSourceChangelist:3810467

Change 3814815 by [REDACTED_PERSON]

	#loc UE4DocJpn File updated against INTSourceChangelist:3810467

Change 3814817 by [REDACTED_PERSON]

	#loc UE4DocJpn File updated against INTSourceChangelist:3810467

Change 3814821 by [REDACTED_PERSON]

	#loc UE4DocJpn File updated against INTSourceChangelist:3810467

Change 3814823 by [REDACTED_PERSON]

	#loc UE4DocJpn File updated against INTSourceChangelist:3810467

Change 3814824 by [REDACTED_PERSON]

	#loc UE4DocJpn File updated against INTSourceChangelist:3810467

Change 3814825 by [REDACTED_PERSON]

	#loc UE4DocJpn File updated against INTSourceChangelist:3810467

Change 3814828 by [REDACTED_PERSON]

	#loc UE4DocJpn File updated against INTSourceChangelist:3809606

Change 3814832 by [REDACTED_PERSON]

	#loc UE4DocJpn File updated against INTSourceChangelist:3809606

Change 3814833 by [REDACTED_PERSON]

	#loc UE4DocJpn File updated against INTSourceChangelist:3809606

Change 3814834 by [REDACTED_PERSON]

	#loc UE4DocJpn File updated against INTSourceChangelist:3809606

Change 3814840 by [REDACTED_PERSON]

	#loc UE4DocJpn File updated against INTSourceChangelist:3809606

Change 3814843 by [REDACTED_PERSON]

	#loc UE4DocJpn File updated against INTSourceChangelist:3807090

Change 3814844 by [REDACTED_PERSON]

	#loc UE4DocJpn File updated against INTSourceChangelist:3810557

Change 3814845 by [REDACTED_PERSON]

	#loc UE4DocJpn File updated against INTSourceChangelist:3782314

Change 3814940 by [REDACTED_PERSON]

	#loc UE4DocJpn File updated against INT# INTSourceChangelist:3782314

Change 3814941 by [REDACTED_PERSON]

	#loc UE4DocJpn File updated against INT# INTSourceChangelist:3782314

Change 3814942 by [REDACTED_PERSON]

	#loc UE4DocJpn File updated against INTSourceChangelist:3807090

Change 3814943 by [REDACTED_PERSON]

	#loc UE4DocJpn File updated against INTSourceChangelist:3807090

Change 3814944 by [REDACTED_PERSON]

	#loc UE4DocJpn File updated against INTSourceChangelist:3782314

Change 3814945 by [REDACTED_PERSON]

	#loc UE4DocJpn File updated against INTSourceChangelist:3782314

Change 3815773 by [REDACTED_PERSON]

	#loc UE4DocJpn File updated against INTSourceChangelist:3782314

Change 3815774 by [REDACTED_PERSON]

	#loc UE4DocJpn File updated against INTSourceChangelist:3810557

Change 3815776 by [REDACTED_PERSON]

	#loc UE4DocJpn File updated against INTSourceChangelist:3810557

Change 3815779 by [REDACTED_PERSON]

	#loc UE4DocJpn File updated against INTSourceChangelist:3810557

Change 3815780 by [REDACTED_PERSON]

	#loc UE4DocJpn File updated against INTSourceChangelist:3810557

Change 3815785 by [REDACTED_PERSON]

	#loc UE4DocJpn File updated against INTSourceChangelist:3810557

Change 3815791 by [REDACTED_PERSON]

	#loc UE4DocJpn File updated against INTSourceChangelist:3810557

Change 3815793 by [REDACTED_PERSON]

	#loc UE4DocJpn File updated against INTSourceChangelist:3810557

Change 3815796 by [REDACTED_PERSON]

	#loc UE4DocJpn File updated against INTSourceChangelist:3806721

Change 3815797 by [REDACTED_PERSON]

	#loc UE4DocJpn File updated against INTSourceChangelist:3804170

Change 3815799 by [REDACTED_PERSON]

	#loc UE4DocJpn File updated against INTSourceChangelist:3810467

Change 3815801 by [REDACTED_PERSON]

	#loc UE4DocJpn File updated against INTSourceChangelist:3810467

Change 3815803 by [REDACTED_PERSON]

	#loc UE4DocJpn File updated against INTSourceChangelist:3810467

Change 3815808 by [REDACTED_PERSON]

	#loc UE4DocJpn File updated against INTSourceChangelist:3810467

Change 3815814 by [REDACTED_PERSON]

	#loc UE4DocJpn File updated against INTSourceChangelist:3807090

Change 3815815 by [REDACTED_PERSON]

	#loc UE4DocJpn File updated against INTSourceChangelist:3807090

Change 3815816 by [REDACTED_PERSON]

	#loc UE4DocJpn File updated against INTSourceChangelist:3807090

Change 3815818 by [REDACTED_PERSON]

	#loc UE4DocJpn File updated against INTSourceChangelist:3801260

Change 3815820 by [REDACTED_PERSON]

	#loc UE4DocJpn File updated against INTSourceChangelist:3782314

Change 3815822 by [REDACTED_PERSON]

	#loc UE4DocJpn File updated against INTSourceChangelist:3782314

Change 3815826 by [REDACTED_PERSON]

	#loc UE4DocJpn File updated against INTSourceChangelist:3782314

Change 3815828 by [REDACTED_PERSON]

	#loc UE4DocJpn File updated against INTSourceChangelist:3782314

Change 3815832 by [REDACTED_PERSON]

	#loc UE4DocJpn File updated against INTSourceChangelist:3808342

Change 3815834 by [REDACTED_PERSON]

	#loc UE4DocJpn File updated against INTSourceChangelist:3808342

Change 3815835 by [REDACTED_PERSON]

	#loc UE4DocJpn File updated against INTSourceChangelist:3808342

Change 3815836 by [REDACTED_PERSON]

	#loc UE4DocJpn File updated against INTSourceChangelist:3807090

Change 3816009 by [REDACTED_PERSON]

	Updated file against INTSourceChangelist:3718323

Change 3816010 by [REDACTED_PERSON]

	Updated file against INTSourceChangelist:3727000

Change 3816012 by [REDACTED_PERSON]

	Updated file against INTSourceChangelist:3810467

Change 3816037 by [REDACTED_PERSON]

	Updated file against INTSourceChangelist:3810467

Change 3816040 by [REDACTED_PERSON]

	Updated file against INTSourceChangelist:3622284

Change 3816049 by [REDACTED_PERSON]

	Updated file against INTSourceChangelist:3680097

Change 3816050 by [REDACTED_PERSON]

	Updated file against INTSourceChangelist:3692730

Change 3816052 by [REDACTED_PERSON]

	Updated file against INTSourceChangelist:3782314

Change 3816053 by [REDACTED_PERSON]

	Updated file against INTSourceChangelist:3692730

Change 3816089 by [REDACTED_PERSON]

	#loc UE4DocJpn File updated against INTSourceChangelist:3782314

Change 3816091 by [REDACTED_PERSON]

	#loc UE4DocJpn File updated against INTSourceChangelist:3782314

Change 3816094 by [REDACTED_PERSON]

	#loc UE4DocJpn File updated against INTSourceChangelist:3782314

Change 3816098 by [REDACTED_PERSON]

	#loc UE4DocJpn File updated against INTSourceChangelist:3782314

Change 3816102 by [REDACTED_PERSON]

	#loc UE4DocJpn File updated against INTSourceChangelist:3809606

Change 3816375 by [REDACTED_PERSON]

	Updated file against INTSourceChangelist:3809606

Change 3816733 by [REDACTED_PERSON]

	Updated file against INTSourceChangelist:3689149

Change 3816736 by [REDACTED_PERSON]

	Fixed bookmark error

Change 3816745 by [REDACTED_PERSON]

	Updated file against INTSourceChangelist:3723511

Change 3819579 by [REDACTED_PERSON]

	#loc UE4DocKOR

Change 3819679 by [REDACTED_PERSON]

	#loc UE4DocKOR

Change 3822726 by [REDACTED_PERSON]

	Updated file against  INTSourceChangelist:3397556

Change 3822745 by [REDACTED_PERSON]

	Updated file against INTSourceChangelist:3244370

Change 3822750 by [REDACTED_PERSON]

	Updated file against INTSourceChangelist:3244370

Change 3822758 by [REDACTED_PERSON]

	Updated file against INTSourceChangelist:3244370

Change 3822771 by [REDACTED_PERSON]

	Updated file against INTSourceChangelist:3244370

Change 3822778 by [REDACTED_PERSON]

	Updated file against INTSourceChangelist:3244370

Change 3827349 by [REDACTED_PERSON]

	Updated file against INTSourceChangelist:3244370

Change 3827377 by [REDACTED_PERSON]

	Updated file against INTSourceChangelist:3244370

Change 3827379 by [REDACTED_PERSON]

	Updated file against INTSourceChangelist:3244370

Change 3827384 by [REDACTED_PERSON]

	Updated file against INTSourceChangelist:3244370

Change 3827386 by [REDACTED_PERSON]

	Updated file against INTSourceChangelist:3244370

Change 3827391 by [REDACTED_PERSON]

	Updated file against INTSourceChangelist:3243860

Change 3827396 by [REDACTED_PERSON]

	Updated file against INTSourceChangelist:3726228

Change 3827404 by [REDACTED_PERSON]

	Updated file against INTSourceChangelist:3474032

Change 3827405 by [REDACTED_PERSON]

	#loc UE4DocJpn File updated against INTSourceChangelist:3801260

Change 3827409 by [REDACTED_PERSON]

	#loc UE4DocJpn File updated against INTSourceChangelist:3782314

Change 3827411 by [REDACTED_PERSON]

	#loc UE4DocJpn File updated against INTSourceChangelist:3782314

Change 3827416 by [REDACTED_PERSON]

	#loc UE4DocJpn File updated against INTSourceChangelist:3782314

Change 3827419 by [REDACTED_PERSON]

	#loc UE4DocJpn File updated against INTSourceChangelist:3782314

Change 3827421 by [REDACTED_PERSON]

	#loc UE4DocJpn File updated against INTSourceChangelist:3782314

Change 3827426 by [REDACTED_PERSON]

	#loc UE4DocJpn File updated against INTSourceChangelist:3782314

Change 3827447 by [REDACTED_PERSON]

	#loc UE4DocJpn File updated againstINTSourceChangelist:3782314

Change 3827450 by [REDACTED_PERSON]

	#loc UE4DocJpn File updated againstINTSourceChangelist:3782314

Change 3827455 by [REDACTED_PERSON]

	#loc UE4DocJpn File updated againstINTSourceChangelist:3782314

Change 3827457 by [REDACTED_PERSON]

	#loc UE4DocJpn File updated againstINTSourceChangelist:3782314

Change 3827458 by [REDACTED_PERSON]

	Updated file against INTSourceChangelist:3801260

Change 3827463 by [REDACTED_PERSON]

	#loc UE4DocKor

Change 3827464 by [REDACTED_PERSON]

	#loc UE4DocJpn File updated againstINTSourceChangelist:3782314

Change 3827465 by [REDACTED_PERSON]

	#loc UE4DocJpn File updated againstINTSourceChangelist:3782314

Change 3827466 by [REDACTED_PERSON]

	#loc UE4DocJpn File updated againstINTSourceChangelist:3782314

Change 3827470 by [REDACTED_PERSON]

	#loc UE4DocJpn File updated againstINTSourceChangelist:3782314

Change 3827472 by [REDACTED_PERSON]

	#loc UE4DocJpn File updated againstINTSourceChangelist:3782314

Change 3827475 by [REDACTED_PERSON]

	Updated file against INTSourceChangelist:3801260

Change 3827486 by [REDACTED_PERSON]

	#loc UE4DocJpn File updated againstINTSourceChangelist:3782314

Change 3827537 by [REDACTED_PERSON]

	Updated file against INTSourceChangelist:3716175

Change 3827538 by [REDACTED_PERSON]

	Fixed title translation

Change 3827541 by [REDACTED_PERSON]

	Updated file against INTSourceChangelist:3211040

Change 3827548 by [REDACTED_PERSON]

	Updated file against INTSourceChangelist:3807086

Change 3827550 by [REDACTED_PERSON]

	Updated file against INTSourceChangelist:3766333

Change 3827717 by [REDACTED_PERSON]

	#loc UE4DocKOR

Change 3829650 by [REDACTED_PERSON]

	File updated against INTSourceChangelist:3782314

Change 3829653 by [REDACTED_PERSON]

	Updated file against INTSourceChangelist:3782314

Change 3829656 by [REDACTED_PERSON]

	Updated file against INTSourceChangelist:3782314

Change 3829658 by [REDACTED_PERSON]

	Updated file against INTSourceChangelist:3782314

Change 3829662 by [REDACTED_PERSON]

	Updated file against INTSourceChangelist:3782314

Change 3829663 by [REDACTED_PERSON]

	Updated file against INTSourceChangelist:3782314

Change 3829677 by [REDACTED_PERSON]

	Updated file against INTSourceChangelist:3782314

Change 3829681 by [REDACTED_PERSON]

	Updated file against INTSourceChangelist:3782314

Change 3829686 by [REDACTED_PERSON]

	Updated file against INTSourceChangelist:3782314

Change 3829694 by [REDACTED_PERSON]

	Updated file against INTSourceChangelist:3782314

Change 3829704 by [REDACTED_PERSON]

	Updated file against INTSourceChangelist:3782314

Change 3829714 by [REDACTED_PERSON]

	Updated file against INTSourceChangelist:3782314

Change 3829721 by [REDACTED_PERSON]

	Updated file against INTSourceChangelist:3782314

Change 3829725 by [REDACTED_PERSON]

	Updated file against INTSourceChangelist:3782314

Change 3829727 by [REDACTED_PERSON]

	Updated file against INTSourceChangelist:3782314

Change 3829732 by [REDACTED_PERSON]

	Updated file against INTSourceChangelist:3782314

Change 3829735 by [REDACTED_PERSON]

	Updated file against INTSourceChangelist:3782314

Change 3829736 by [REDACTED_PERSON]

	Updated file against INTSourceChangelist:3782314

Change 3829739 by [REDACTED_PERSON]

	Updated file against INTSourceChangelist:3782314

Change 3829741 by [REDACTED_PERSON]

	Updated file against INTSourceChangelist:3782314

Change 3829742 by [REDACTED_PERSON]

	Updated file against INTSourceChangelist:3782314

Change 3829751 by [REDACTED_PERSON]

	Updated file against INTSourceChangelist:3782314

Change 3829754 by [REDACTED_PERSON]

	Updated file against INTSourceChangelist:3782314

Change 3829793 by [REDACTED_PERSON]

	Updated file against INTSourceChangelist:3782314

Change 3829798 by [REDACTED_PERSON]

	Updated file against INTSourceChangelist:3782314

Change 3829799 by [REDACTED_PERSON]

	Updated file against INTSourceChangelist:3782314

Change 3829803 by [REDACTED_PERSON]

	Updated file against INTSourceChangelist:3782314

Change 3829804 by [REDACTED_PERSON]

	Updated file against INTSourceChangelist:3782314

Change 3829805 by [REDACTED_PERSON]

	Updated file against INTSourceChangelist:3782314

Change 3829808 by [REDACTED_PERSON]

	Updated file against INTSourceChangelist:3782314

Change 3829809 by [REDACTED_PERSON]

	Updated file against INTSourceChangelist:3782314

Change 3829814 by [REDACTED_PERSON]

	Updated file against INTSourceChangelist:3782314

Change 3829816 by [REDACTED_PERSON]

	Updated file against INTSourceChangelist:3782314

Change 3829818 by [REDACTED_PERSON]

	Updated file against INTSourceChangelist:3782314

Change 3829821 by [REDACTED_PERSON]

	Updated file against INTSourceChangelist:3782314

Change 3829823 by [REDACTED_PERSON]

	Updated file against INTSourceChangelist:3782314

Change 3829824 by [REDACTED_PERSON]

	Updated file against INTSourceChangelist:3782314

Change 3829825 by [REDACTED_PERSON]

	Updated file against INTSourceChangelist:3782314

Change 3829827 by [REDACTED_PERSON]

	Updated file against INTSourceChangelist:3782314

Change 3829828 by [REDACTED_PERSON]

	Updated file against INTSourceChangelist:3782314

Change 3829829 by [REDACTED_PERSON]

	Updated file against INTSourceChangelist:3782314

Change 3829833 by [REDACTED_PERSON]

	Updated file against INTSourceChangelist:3782314

Change 3829837 by [REDACTED_PERSON]

	Updated file against INTSourceChangelist:3782314

Change 3829850 by [REDACTED_PERSON]

	Updated file against INTSourceChangelist:3782314

Change 3829855 by [REDACTED_PERSON]

	Updated file against INTSourceChangelist:3782314

Change 3829859 by [REDACTED_PERSON]

	Updated file against INTSourceChangelist:3782314

Change 3829862 by [REDACTED_PERSON]

	Updated file against INTSourceChangelist:3782314

Change 3829864 by [REDACTED_PERSON]

	Updated file against INTSourceChangelist:3782314

Change 3829866 by [REDACTED_PERSON]

	Updated file against INTSourceChangelist:3782314

Change 3829869 by [REDACTED_PERSON]

	Updated file against INTSourceChangelist:3782314

Change 3829870 by [REDACTED_PERSON]

	Updated file against INTSourceChangelist:3782314

Change 3829872 by [REDACTED_PERSON]

	Updated file against INTSourceChangelist:3782314

Change 3829876 by [REDACTED_PERSON]

	Updated file against INTSourceChangelist:3782314

Change 3829877 by [REDACTED_PERSON]

	Updated file against INTSourceChangelist:3782314

Change 3829878 by [REDACTED_PERSON]

	Updated file against INTSourceChangelist:3782314

Change 3829882 by [REDACTED_PERSON]

	Updated file against INTSourceChangelist:3782314

Change 3829895 by [REDACTED_PERSON]

	Updated file against INTSourceChangelist:3782314

Change 3829898 by [REDACTED_PERSON]

	Updated file against INTSourceChangelist:3782314

Change 3829900 by [REDACTED_PERSON]

	Updated file against INTSourceChangelist:3782314

Change 3829901 by [REDACTED_PERSON]

	Updated file against INTSourceChangelist:3782314

Change 3829902 by [REDACTED_PERSON]

	Updated file against INTSourceChangelist:3782314

Change 3829904 by [REDACTED_PERSON]

	Updated file against INTSourceChangelist:3782314

Change 3829905 by [REDACTED_PERSON]

	Updated file against INTSourceChangelist:3782314

Change 3829906 by [REDACTED_PERSON]

	Updated file against INTSourceChangelist:3782314

Change 3829907 by [REDACTED_PERSON]

	Updated file against INTSourceChangelist:3782314

Change 3829909 by [REDACTED_PERSON]

	Updated file against INTSourceChangelist:3782314

Change 3829911 by [REDACTED_PERSON]

	Updated file against INTSourceChangelist:3782314

Change 3831636 by [REDACTED_PERSON]

	#loc UE4DocJpn File updated against INTSourceChangelist:3782314

Change 3831644 by [REDACTED_PERSON]

	#loc UE4DocJpn File updated against INTSourceChangelist:3782314

Change 3831649 by [REDACTED_PERSON]

	Updated file against INTSourceChangelist:3782314

Change 3831651 by [REDACTED_PERSON]

	Updated file against INTSourceChangelist:3782314

Change 3831654 by [REDACTED_PERSON]

	Updated file against INTSourceChangelist:3782314

Change 3831656 by [REDACTED_PERSON]

	#loc UE4DocJpn File updated against INTSourceChangelist:3782314

Change 3831664 by [REDACTED_PERSON]

	#loc UE4DocJpn File updated against INTSourceChangelist:3782314

Change 3831667 by [REDACTED_PERSON]

	Updated file against INTSourceChangelist:3782314

Change 3831671 by [REDACTED_PERSON]

	#loc UE4DocJpn File updated against INTSourceChangelist:3782314

Change 3831672 by [REDACTED_PERSON]

	Updated file against INTSourceChangelist:3782314

Change 3831674 by [REDACTED_PERSON]

	Updated file against INTSourceChangelist:3782314

Change 3831682 by [REDACTED_PERSON]

	Updated file against INTSourceChangelist:3782314

Change 3831687 by [REDACTED_PERSON]

	Updated file against INTSourceChangelist:3782314

Change 3831690 by [REDACTED_PERSON]

	Updated file against INTSourceChangelist:3782314

Change 3831694 by [REDACTED_PERSON]

	Updated file against INTSourceChangelist:3782314

Change 3831695 by [REDACTED_PERSON]

	#loc UE4DocJpn File updated against INTSourceChangelist:3782314

Change 3831705 by [REDACTED_PERSON]

	Updated file against INTSourceChangelist:3782314

Change 3831716 by [REDACTED_PERSON]

	Updated file against INTSourceChangelist:3782314

Change 3831721 by [REDACTED_PERSON]

	Updated file against INTSourceChangelist:3782314

Change 3831722 by [REDACTED_PERSON]

	Updated file against INTSourceChangelist:3782314

Change 3831739 by [REDACTED_PERSON]

	Updated file against INTSourceChangelist:3782314

Change 3831744 by [REDACTED_PERSON]

	Updated file against INTSourceChangelist:3782314

Change 3831745 by [REDACTED_PERSON]

	#loc UE4DocJpn File updated againstINTSourceChangelist:3782314

Change 3831752 by [REDACTED_PERSON]

	Updated file against INTSourceChangelist:3782314

Change 3831772 by [REDACTED_PERSON]

	#loc UE4DocJpn File updated againstINTSourceChangelist:3782314

Change 3831779 by [REDACTED_PERSON]

	#loc UE4DocJpn File updated againstINTSourceChangelist:3782314

Change 3831783 by [REDACTED_PERSON]

	#loc UE4DocJpn File updated againstINTSourceChangelist:3782314

Change 3831793 by [REDACTED_PERSON]

	#loc UE4DocJpn File updated againstINTSourceChangelist:3782314

Change 3831794 by [REDACTED_PERSON]

	#loc UE4DocJpn File updated againstINTSourceChangelist:3782314

Change 3831796 by [REDACTED_PERSON]

	#loc UE4DocJpn File updated againstINTSourceChangelist:3782314

Change 3831803 by [REDACTED_PERSON]

	#loc UE4DocJpn File updated againstINTSourceChangelist:3782314

Change 3831805 by [REDACTED_PERSON]

	Updated file against INTSourceChangelist:3782314

Change 3831808 by [REDACTED_PERSON]

	#loc UE4DocJpn File updated against INTSourceChangelist:3782314

Change 3831811 by [REDACTED_PERSON]

	Updated file against INTSourceChangelist:3782314

Change 3831812 by [REDACTED_PERSON]

	#loc UE4DocJpn File updated againstINTSourceChangelist:3782314

Change 3831815 by [REDACTED_PERSON]

	Updated file against INTSourceChangelist:3782314

Change 3831826 by [REDACTED_PERSON]

	Updated file against INTSourceChangelist:3782314

Change 3831828 by [REDACTED_PERSON]

	Updated file against INTSourceChangelist:3782314

Change 3831835 by [REDACTED_PERSON]

	Updated file against INTSourceChangelist:3782314

Change 3831837 by [REDACTED_PERSON]

	Updated file against INTSourceChangelist:3782314

Change 3831839 by [REDACTED_PERSON]

	Updated file against INTSourceChangelist:3782314

Change 3831845 by [REDACTED_PERSON]

	Updated file against INTSourceChangelist:3782314

Change 3831847 by [REDACTED_PERSON]

	Updated file against INTSourceChangelist:3782314

Change 3831850 by [REDACTED_PERSON]

	Updated file against INTSourceChangelist:3782314

Change 3831851 by [REDACTED_PERSON]

	Updated file against INTSourceChangelist:3782314

Change 3831859 by [REDACTED_PERSON]

	Updated file against INTSourceChangelist:3782314

Change 3831862 by [REDACTED_PERSON]

	Updated file against INTSourceChangelist:3782314

Change 3831863 by [REDACTED_PERSON]

	Updated file against INTSourceChangelist:3782314

Change 3831865 by [REDACTED_PERSON]

	Updated file against INTSourceChangelist:3782314

Change 3831870 by [REDACTED_PERSON]

	Updated file against INTSourceChangelist:3782314

Change 3831928 by [REDACTED_PERSON]

	c:\UE4\Main\Engine\Documentation\Source\Engine\Performance\Visualizers\OutOfBoundsPixels\OutOfBoundsPixels.INT.udn

Change 3831931 by [REDACTED_PERSON]

	Updated file against INTSourceChangelist:3782314

Change 3831934 by [REDACTED_PERSON]

	Updated file against INTSourceChangelist:3782314

Change 3831939 by [REDACTED_PERSON]

	Updated file against INTSourceChangelist:3782314

Change 3831956 by [REDACTED_PERSON]

	Updated file against INTSourceChangelist:3782314

Change 3831957 by [REDACTED_PERSON]

	Updated file against INTSourceChangelist:3782314

Change 3831964 by [REDACTED_PERSON]

	Updated file against INTSourceChangelist:3782314

Change 3831977 by [REDACTED_PERSON]

	Updated file against INTSourceChangelist:3782314

Change 3831978 by [REDACTED_PERSON]

	Updated file against INTSourceChangelist:3782314

Change 3831981 by [REDACTED_PERSON]

	Updated file against INTSourceChangelist:3782314

Change 3831986 by [REDACTED_PERSON]

	Updated file against INTSourceChangelist:3782314

Change 3831989 by [REDACTED_PERSON]

	Updated file against INTSourceChangelist:3782314

Change 3831991 by [REDACTED_PERSON]

	Updated file against INTSourceChangelist:3782314

Change 3831997 by [REDACTED_PERSON]

	Updated file against INTSourceChangelist:3782314

Change 3831998 by [REDACTED_PERSON]

	Updated file against INTSourceChangelist:3782314

Change 3832001 by [REDACTED_PERSON]

	Updated file against INTSourceChangelist:3782314

Change 3832004 by [REDACTED_PERSON]

	Updated file against INTSourceChangelist:3782314

Change 3832006 by [REDACTED_PERSON]

	Updated file against INTSourceChangelist:3782314

Change 3832015 by [REDACTED_PERSON]

	Updated file against INTSourceChangelist:3782314

Change 3832016 by [REDACTED_PERSON]

	Updated file against INTSourceChangelist:3782314

Change 3832018 by [REDACTED_PERSON]

	Updated file against INTSourceChangelist:3782314

Change 3832019 by [REDACTED_PERSON]

	Updated file against INTSourceChangelist:3782314

Change 3832021 by [REDACTED_PERSON]

	Updated file against INTSourceChangelist:3782314

Change 3832022 by [REDACTED_PERSON]

	Updated file against INTSourceChangelist:3782314

Change 3832026 by [REDACTED_PERSON]

	Updated file against INTSourceChangelist:3782314

Change 3832027 by [REDACTED_PERSON]

	Updated file against INTSourceChangelist:3782314

Change 3832028 by [REDACTED_PERSON]

	Updated file against INTSourceChangelist:3782314

Change 3832032 by [REDACTED_PERSON]

	Updated file against INTSourceChangelist:3782314

Change 3832040 by [REDACTED_PERSON]

	Updated file against INTSourceChangelist:3782314

Change 3832042 by [REDACTED_PERSON]

	Updated file against INTSourceChangelist:3782314

Change 3832051 by [REDACTED_PERSON]

	Updated file against INTSourceChangelist:3782314

Change 3832053 by [REDACTED_PERSON]

	Updated file against INTSourceChangelist:3782314

Change 3832063 by [REDACTED_PERSON]

	#loc UE4DocJpn File updated against INTSourceChangelist:3782314

Change 3832064 by [REDACTED_PERSON]

	Updated file against INTSourceChangelist:3782314

Change 3832072 by [REDACTED_PERSON]

	Updated file against INTSourceChangelist:3782314

Change 3832073 by [REDACTED_PERSON]

	#loc UE4DocJpn File updated against INTSourceChangelist:3782314

Change 3832074 by [REDACTED_PERSON]

	Updated file against INTSourceChangelist:3782314

Change 3832076 by [REDACTED_PERSON]

	#loc UE4DocJpn File updated against INTSourceChangelist:3782314

Change 3832078 by [REDACTED_PERSON]

	Updated file against INTSourceChangelist:3782314

Change 3832080 by [REDACTED_PERSON]

	Updated file against INTSourceChangelist:3782314

Change 3832083 by [REDACTED_PERSON]

	#loc UE4DocJpn File updated against INTSourceChangelist:3782314

Change 3832086 by [REDACTED_PERSON]

	Updated file against INTSourceChangelist:3782314

Change 3832088 by [REDACTED_PERSON]

	#loc UE4DocJpn File updated against INTSourceChangelist:3809606

Change 3832089 by [REDACTED_PERSON]

	Updated file against INTSourceChangelist:3782314

Change 3832092 by [REDACTED_PERSON]

	Updated file against INTSourceChangelist:3782314

Change 3832093 by [REDACTED_PERSON]

	Updated file against INTSourceChangelist:3782314

Change 3832094 by [REDACTED_PERSON]

	Updated file against INTSourceChangelist:3782314

Change 3832100 by [REDACTED_PERSON]

	Updated file against INTSourceChangelist:3782314

Change 3832103 by [REDACTED_PERSON]

	Updated file against INTSourceChangelist:3782314

Change 3832105 by [REDACTED_PERSON]

	Updated file against INTSourceChangelist:3782314

Change 3832107 by [REDACTED_PERSON]

	Updated file against INTSourceChangelist:3782314

Change 3832111 by [REDACTED_PERSON]

	Updated file against INTSourceChangelist:3782314

Change 3832114 by [REDACTED_PERSON]

	#loc UE4DocJpn File updated against INTSourceChangelist:3782314

Change 3832115 by [REDACTED_PERSON]

	Updated file against INTSourceChangelist:3782314

Change 3832116 by [REDACTED_PERSON]

	Updated file against INTSourceChangelist:3782314

Change 3832119 by [REDACTED_PERSON]

	Updated file against INTSourceChangelist:3782314

Change 3832125 by [REDACTED_PERSON]

	#loc UE4DocJpn File updated against INTSourceChangelist:3782314

Change 3832132 by [REDACTED_PERSON]

	Updated file against INTSourceChangelist:3782314

Change 3832133 by [REDACTED_PERSON]

	Updated file against INTSourceChangelist:3782314

Change 3832137 by [REDACTED_PERSON]

	Updated file against INTSourceChangelist:3782314

Change 3832138 by [REDACTED_PERSON]

	Updated file against INTSourceChangelist:3782314

Change 3832141 by [REDACTED_PERSON]

	Updated file against INTSourceChangelist:3782314

Change 3832142 by [REDACTED_PERSON]

	Updated file against INTSourceChangelist:3782314

Change 3833977 by [REDACTED_PERSON]

	Fixed link error

Change 3833981 by [REDACTED_PERSON]

	Fixed link error

Change 3833990 by [REDACTED_PERSON]

	Updated file against INTSourceChangelist:3782314

Change 3833995 by [REDACTED_PERSON]

	Updated file against INTSourceChangelist:3782314

Change 3834027 by [REDACTED_PERSON]

	Updated file against INTSourceChangelist:3782314

Change 3834030 by [REDACTED_PERSON]

	Updated file against INTSourceChangelist:3782314

Change 3834038 by [REDACTED_PERSON]

	#loc UE4DocJpn File updated againstINTSourceChangelist:3810557

Change 3834041 by [REDACTED_PERSON]

	Updated file against INTSourceChangelist:3782314

Change 3834060 by [REDACTED_PERSON]

	Updated file against INTSourceChangelist:3782314

Change 3834072 by [REDACTED_PERSON]

	Updated file against INTSourceChangelist:3782314

Change 3834078 by [REDACTED_PERSON]

	Updated file against INTSourceChangelist:3782314

Change 3834086 by [REDACTED_PERSON]

	Updated file against INTSourceChangelist:3782314

Change 3834089 by [REDACTED_PERSON]

	Updated file against INTSourceChangelist:3782314

Change 3834097 by [REDACTED_PERSON]

	#loc UE4DocJpn File updated againstINTSourceChangelist:3782314

Change 3834098 by [REDACTED_PERSON]

	Updated file against INTSourceChangelist:3782314

Change 3834100 by [REDACTED_PERSON]

	Updated file against INTSourceChangelist:3782314

Change 3834102 by [REDACTED_PERSON]

	Updated file against INTSourceChangelist:3782314

Change 3834106 by [REDACTED_PERSON]

	Updated file against INTSourceChangelist:3782314

Change 3834117 by [REDACTED_PERSON]

	#loc UE4DocJpn File updated against INTSourceChangelist:3782314

Change 3834124 by [REDACTED_PERSON]

	Updated file against INTSourceChangelist:3782314

Change 3834126 by [REDACTED_PERSON]

	Updated file against INTSourceChangelist:3782314

Change 3834134 by [REDACTED_PERSON]

	Updated file against INTSourceChangelist:3782314

Change 3834136 by [REDACTED_PERSON]

	Updated file against INTSourceChangelist:3782314

Change 3834138 by [REDACTED_PERSON]

	Updated file against INTSourceChangelist:3782314

Change 3834140 by [REDACTED_PERSON]

	#loc UE4DocJpn File updated against INTSourceChangelist:3782314

Change 3834147 by [REDACTED_PERSON]

	Updated file against INTSourceChangelist:3782314

Change 3834155 by [REDACTED_PERSON]

	Updated file against INTSourceChangelist:3782314

Change 3834157 by [REDACTED_PERSON]

	Updated file against INTSourceChangelist:3782314

Change 3834158 by [REDACTED_PERSON]

	#loc UE4DocJpn File updated againstINTSourceChangelist:3782314

Change 3834159 by [REDACTED_PERSON]

	Updated file against INTSourceChangelist:3782314

Change 3834163 by [REDACTED_PERSON]

	Updated file against INTSourceChangelist:3782314

Change 3834168 by [REDACTED_PERSON]

	Updated file against INTSourceChangelist:3782314

Change 3834172 by [REDACTED_PERSON]

	#loc UE4DocJpn File updated against INTSourceChangelist:3782314

Change 3834173 by [REDACTED_PERSON]

	Updated file against INTSourceChangelist:3782314

Change 3834175 by [REDACTED_PERSON]

	Updated file against INTSourceChangelist:3782314

Change 3834183 by [REDACTED_PERSON]

	Updated file against INTSourceChangelist:3782314

Change 3834185 by [REDACTED_PERSON]

	#loc UE4DocJpn File updated againstINTSourceChangelist:3782314

Change 3834187 by [REDACTED_PERSON]

	Updated file against INTSourceChangelist:3782314

Change 3834188 by [REDACTED_PERSON]

	Updated file against INTSourceChangelist:3782314

Change 3834190 by [REDACTED_PERSON]

	Updated file against INTSourceChangelist:3782314

Change 3834193 by [REDACTED_PERSON]

	#loc UE4DocJpn File updated against INTSourceChangelist:3782314

Change 3834196 by [REDACTED_PERSON]

	Updated file against INTSourceChangelist:3782314

Change 3834197 by [REDACTED_PERSON]

	Updated file against INTSourceChangelist:3782314

Change 3834201 by [REDACTED_PERSON]

	Updated file against INTSourceChangelist:3782314

Change 3834203 by [REDACTED_PERSON]

	Updated file against INTSourceChangelist:3782314

Change 3834210 by [REDACTED_PERSON]

	Updated file against INTSourceChangelist:3782314

Change 3834231 by [REDACTED_PERSON]

	Updated file against INTSourceChangelist:3782314

Change 3834232 by [REDACTED_PERSON]

	#loc UE4DocJpn File updated against INTSourceChangelist:3782314

Change 3834233 by [REDACTED_PERSON]

	#loc UE4DocJpn File updated against INTSourceChangelist:3782314

Change 3834236 by [REDACTED_PERSON]

	Updated file against INTSourceChangelist:3782314

Change 3834238 by [REDACTED_PERSON]

	Updated file against INTSourceChangelist:3782314

Change 3834242 by [REDACTED_PERSON]

	#loc UE4DocJpn File updated against INTSourceChangelist:3782314

Change 3834243 by [REDACTED_PERSON]

	Updated file against INTSourceChangelist:3782314

Change 3834244 by [REDACTED_PERSON]

	#loc UE4DocJpn File updated against INTSourceChangelist:3782314

Change 3834245 by [REDACTED_PERSON]

	#loc UE4DocJpn File updated against INTSourceChangelist:3782314

Change 3834246 by [REDACTED_PERSON]

	Updated file against INTSourceChangelist:3782314

Change 3834248 by [REDACTED_PERSON]

	#loc UE4DocJpn File updated against INTSourceChangelist:3782314

Change 3834249 by [REDACTED_PERSON]

	#loc UE4DocJpn File updated against INTSourceChangelist:3782314

Change 3834250 by [REDACTED_PERSON]

	#loc UE4DocJpn File updated against INTSourceChangelist:3782314

Change 3834251 by [REDACTED_PERSON]

	Updated file against INTSourceChangelist:3782314

Change 3834252 by [REDACTED_PERSON]

	#loc UE4DocJpn File updated against INTSourceChangelist:3782314

Change 3834253 by [REDACTED_PERSON]

	#loc UE4DocJpn File updated against INTSourceChangelist:3782314

Change 3834259 by [REDACTED_PERSON]

	#loc UE4DocJpn File updated against INTSourceChangelist:3782314

Change 3834260 by [REDACTED_PERSON]

	Updated file against INTSourceChangelist:3782314

Change 3834261 by [REDACTED_PERSON]

	#loc UE4DocJpn File updated against INTSourceChangelist:3782314

Change 3834262 by [REDACTED_PERSON]

	Updated file against INTSourceChangelist:3782314

Change 3834263 by [REDACTED_PERSON]

	#loc UE4DocJpn File updated against INTSourceChangelist:3782314

Change 3834264 by [REDACTED_PERSON]

	#loc UE4DocJpn File updated against INTSourceChangelist:3782314

Change 3834265 by [REDACTED_PERSON]

	Updated file against INTSourceChangelist:3782314

Change 3834266 by [REDACTED_PERSON]

	#loc UE4DocJpn File updated against INTSourceChangelist:3782314

Change 3834267 by [REDACTED_PERSON]

	Updated file against INTSourceChangelist:3782314

Change 3834268 by [REDACTED_PERSON]

	Updated file against INTSourceChangelist:3782314

Change 3834271 by [REDACTED_PERSON]

	Updated file against INTSourceChangelist:3782314

Change 3834272 by [REDACTED_PERSON]

	Updated file against INTSourceChangelist:3782314

Change 3834274 by [REDACTED_PERSON]

	#loc UE4DocJpn File updated against INTSourceChangelist:3782314

Change 3834275 by [REDACTED_PERSON]

	Updated file against INTSourceChangelist:3782314

Change 3834276 by [REDACTED_PERSON]

	Updated file against INTSourceChangelist:3782314

Change 3834282 by [REDACTED_PERSON]

	#loc UE4DocJpn File updated against INTSourceChangelist:3782314

Change 3834298 by [REDACTED_PERSON]

	Updated file against INTSourceChangelist:3782314

Change 3834299 by [REDACTED_PERSON]

	Updated file against INTSourceChangelist:3782314

Change 3834304 by [REDACTED_PERSON]

	Updated file against INTSourceChangelist:3807090

Change 3834305 by [REDACTED_PERSON]

	Updated file against INTSourceChangelist:3782314

Change 3834311 by [REDACTED_PERSON]

	#loc UE4DocJpn File updated against INTSourceChangelist:3782314

Change 3834313 by [REDACTED_PERSON]

	Updated file against INTSourceChangelist:3782314

Change 3834314 by [REDACTED_PERSON]

	Updated file against INTSourceChangelist:3782314

Change 3834315 by [REDACTED_PERSON]

	Updated file against INTSourceChangelist:3782314

Change 3834316 by [REDACTED_PERSON]

	#loc UE4DocJpn File updated against INTSourceChangelist:3782314

Change 3834318 by [REDACTED_PERSON]

	Updated file against INTSourceChangelist:3782314

Change 3834319 by [REDACTED_PERSON]

	#loc UE4DocJpn File updated against INTSourceChangelist:3782314

Change 3834322 by [REDACTED_PERSON]

	#loc UE4DocJpn File updated against INTSourceChangelist:3782314

Change 3834331 by [REDACTED_PERSON]

	Updated file against INTSourceChangelist:3782314

Change 3834334 by [REDACTED_PERSON]

	Updated file against INTSourceChangelist:3782314

Change 3834336 by [REDACTED_PERSON]

	Updated file against INTSourceChangelist:3782314

Change 3834342 by [REDACTED_PERSON]

	#loc UE4DocJpn File updated against INTSourceChangelist:3782314

Change 3834344 by [REDACTED_PERSON]

	#loc UE4DocJpn File updated against INTSourceChangelist:3782314

Change 3834347 by [REDACTED_PERSON]

	#loc UE4DocJpn File updated against INTSourceChangelist:3782314

Change 3834351 by [REDACTED_PERSON]

	#loc UE4DocJpn File updated against INTSourceChangelist:3782314

Change 3834364 by [REDACTED_PERSON]

	Updated file against INTSourceChangelist:3782314

Change 3834366 by [REDACTED_PERSON]

	Updated file against INTSourceChangelist:3782314

Change 3834367 by [REDACTED_PERSON]

	#loc UE4DocJpn File updated against INTSourceChangelist:3782314

Change 3834378 by [REDACTED_PERSON]

	#loc UE4DocJpn File updated against INTSourceChangelist:3782314

Change 3834384 by [REDACTED_PERSON]

	Updated file against INTSourceChangelist:3782314

Change 3834388 by [REDACTED_PERSON]

	Updated file against INTSourceChangelist:3782314

Change 3834390 by [REDACTED_PERSON]

	Updated file against INTSourceChangelist:3782314

Change 3834392 by [REDACTED_PERSON]

	#loc UE4DocJpn File updated against INTSourceChangelist:3782314

Change 3834395 by [REDACTED_PERSON]

	Updated file against INTSourceChangelist:3782314

Change 3834396 by [REDACTED_PERSON]

	Updated file against INTSourceChangelist:3782314

Change 3834399 by [REDACTED_PERSON]

	Updated file against INTSourceChangelist:3782314

Change 3834402 by [REDACTED_PERSON]

	Updated file against INTSourceChangelist:3782314

Change 3834409 by [REDACTED_PERSON]

	Updated file against INTSourceChangelist:3782314

Change 3834410 by [REDACTED_PERSON]

	#loc UE4DocJpn File updated against INTSourceChangelist:3782314

Change 3834412 by [REDACTED_PERSON]

	Updated file against INTSourceChangelist:3782314

Change 3834418 by [REDACTED_PERSON]

	Updated file against INTSourceChangelist:3782314

Change 3834419 by [REDACTED_PERSON]

	#loc UE4DocJpn File updated against INTSourceChangelist:3782314

Change 3834428 by [REDACTED_PERSON]

	Updated file against INTSourceChangelist:3782314

Change 3834430 by [REDACTED_PERSON]

	Updated file against INTSourceChangelist:3782314

Change 3834439 by [REDACTED_PERSON]

	Updated file against INTSourceChangelist:3782314

Change 3834444 by [REDACTED_PERSON]

	Updated file against INTSourceChangelist:3782314

Change 3834450 by [REDACTED_PERSON]

	Updated file against INTSourceChangelist:3782314

Change 3834451 by [REDACTED_PERSON]

	Updated file against INTSourceChangelist:3782314

Change 3834452 by [REDACTED_PERSON]

	Updated file against INTSourceChangelist:3782314

Change 3834456 by [REDACTED_PERSON]

	Updated file against INTSourceChangelist:3782314

Change 3834459 by [REDACTED_PERSON]

	Updated file against INTSourceChangelist:3782314

Change 3835269 by [REDACTED_PERSON]

	#UE4Doc
	Updated Toolchain table to include link to newest cross compile toolchain.

Change 3837426 by [REDACTED_PERSON]

	Updated file against INTSourceChangelist:3782314

Change 3837428 by [REDACTED_PERSON]

	#loc UE4DocJpn File updated against3810557

Change 3837430 by [REDACTED_PERSON]

	Updated file against INTSourceChangelist:3835269

Change 3837433 by [REDACTED_PERSON]

	Updated file against INTSourceChangelist:3810557

Change 3837436 by [REDACTED_PERSON]

	Updated file against INTSourceChangelist:3810557

Change 3837438 by [REDACTED_PERSON]

	Updated file against INTSourceChangelist:3810557

Change 3837483 by [REDACTED_PERSON]

	Updatd file against INTSourceChangelist:3810557

Change 3837484 by [REDACTED_PERSON]

	Updated file against INTSourceChangelist:3810557

Change 3837492 by [REDACTED_PERSON]

	Updated file against INTSourceChangelist:3810557

Change 3837495 by [REDACTED_PERSON]

	Updated file against INTSourceChangelist:3810557

Change 3837504 by [REDACTED_PERSON]

	Updated file against INTSourceChangelist:3810557

Change 3837506 by [REDACTED_PERSON]

	#loc UE4DocJpn File updated against INTSourceChangelist:3761563

Change 3837507 by [REDACTED_PERSON]

	Updated file against INTSourceChangelist:3810557

Change 3837518 by [REDACTED_PERSON]

	#loc UE4DocJpn File updated against INTSourceChangelist:3750352

Change 3837519 by [REDACTED_PERSON]

	Updated file against INTSourceChangelist:3810557

Change 3837521 by [REDACTED_PERSON]

	Updated file against INTSourceChangelist:3810557

Change 3837522 by [REDACTED_PERSON]

	#loc UE4DocJpn File updated against INTSourceChangelist:3804170

Change 3837523 by [REDACTED_PERSON]

	Updated file against INTSourceChangelist:3810557

Change 3837524 by [REDACTED_PERSON]

	#loc UE4DocJpn File updated against INTSourceChangelist:3782314

Change 3837526 by [REDACTED_PERSON]

	Updated file agaisnt INTSourceChangelist:3810557

Change 3837527 by [REDACTED_PERSON]

	#loc UE4DocJpn File updated against INTSourceChangelist:3782314

Change 3837531 by [REDACTED_PERSON]

	#loc UE4DocJpn File updated against INTSourceChangelist:3782314

Change 3837537 by [REDACTED_PERSON]

	#loc UE4DocJpn File updated against INTSourceChangelist:3782314

Change 3837557 by [REDACTED_PERSON]

	#loc UE4DocJpn File updated againstINTSourceChangelist:3782314

Change 3837569 by [REDACTED_PERSON]

	#loc UE4DocJpn File updated against INTSourceChangelist:3782314

Change 3837580 by [REDACTED_PERSON]

	Updated file against INTSourceChangelist:3810557

Change 3837581 by [REDACTED_PERSON]

	Updated file against INTSourceChangelist:3810557

Change 3837585 by [REDACTED_PERSON]

	Fixed translation

Change 3837590 by [REDACTED_PERSON]

	#loc UE4DocJpn File updated against INTSourceChangelist:3782314

Change 3837591 by [REDACTED_PERSON]

	Updated file against INTSourceChangelist:3810557

Change 3837592 by [REDACTED_PERSON]

	Updated file against INTSourceChangelist:3810557

Change 3837593 by [REDACTED_PERSON]

	Updated file against sINTSourceChangelist:3810557

Change 3837595 by [REDACTED_PERSON]

	Updated file against INTSourceChangelist:3810557

Change 3837597 by [REDACTED_PERSON]

	Updated file against INTSourceChangelist:3810557

Change 3837598 by [REDACTED_PERSON]

	#loc UE4DocJpn File updated against INTSourceChangelist:3807090

Change 3837599 by [REDACTED_PERSON]

	Updated file against INTSourceChangelist:3810557

Change 3837605 by [REDACTED_PERSON]

	Updated file against INTSourceChangelist:3810557

Change 3837606 by [REDACTED_PERSON]

	Updated file against INTSourceChangelist:3810557

Change 3837607 by [REDACTED_PERSON]

	Updated file against INTSourceChangelist:3810557

Change 3837609 by [REDACTED_PERSON]

	Updated file against INTSourceChangelist:3810557

Change 3837612 by [REDACTED_PERSON]

	Updated file against INTSourceChangelist:3810557

Change 3837614 by [REDACTED_PERSON]

	Updated file against INTSourceChangelist:3810557

Change 3837618 by [REDACTED_PERSON]

	Updated file against INTSourceChangelist:3810557

Change 3837621 by [REDACTED_PERSON]

	Updated file against INTSourceChangelist:3810557

Change 3837624 by [REDACTED_PERSON]

	Updated file against INTSourceChangelist:3810557

Change 3837626 by [REDACTED_PERSON]

	Updated file against INTSourceChangelist:3810557

Change 3837627 by [REDACTED_PERSON]

	Updated file against INTSourceChangelist:3810557

Change 3837640 by [REDACTED_PERSON]

	#loc UE4DocJpn File updated against INTSourceChangelist:3807090

Change 3837643 by [REDACTED_PERSON]

	Updated file against INTSourceChangelist:3810557

Change 3837649 by [REDACTED_PERSON]

	Updated file against INTSourceChangelist:3810467

Change 3837650 by [REDACTED_PERSON]

	#loc UE4DocJpn File updated against INTSourceChangelist:3807090

Change 3837652 by [REDACTED_PERSON]

	Updated file against INTSourceChangelist:3810467

Change 3837657 by [REDACTED_PERSON]

	Updatead file against INTSourceChangelist:3810467

Change 3837658 by [REDACTED_PERSON]

	#loc UE4DocJpn File updated against INTSourceChangelist:3807090

Change 3837660 by [REDACTED_PERSON]

	Updated file against INTSourceChangelist:3810467

Change 3837661 by [REDACTED_PERSON]

	Updated file against INTSourceChangelist:3810467

Change 3837665 by [REDACTED_PERSON]

	Updated file against INTSourceChangelist:3810467

Change 3837666 by [REDACTED_PERSON]

	Updated file against INTSourceChangelist:3810467

Change 3837667 by [REDACTED_PERSON]

	#loc UE4DocJpn File updated against INTSourceChangelist:3782314

Change 3837668 by [REDACTED_PERSON]

	#loc UE4DocJpn File updated against INTSourceChangelist:3782314

Change 3837669 by [REDACTED_PERSON]

	Updated file against INTSourceChangelist:3810467

Change 3837671 by [REDACTED_PERSON]

	#loc UE4DocJpn File updated against INTSourceChangelist:3782314

Change 3837672 by [REDACTED_PERSON]

	#loc UE4DocJpn File updated against INTSourceChangelist:3782314

Change 3837682 by [REDACTED_PERSON]

	Updated file against INTSourceChangelist:3810467

Change 3837683 by [REDACTED_PERSON]

	#loc UE4DocJpn File updated against INTSourceChangelist:3782314

Change 3837687 by [REDACTED_PERSON]

	Updated file against INTSourceChangelist:3810467

Change 3837688 by [REDACTED_PERSON]

	#loc UE4DocJpn File updated against INTSourceChangelist:3782314

Change 3837689 by [REDACTED_PERSON]

	Updated file against INTSourceChangelist:3810467

Change 3837690 by [REDACTED_PERSON]

	#loc UE4DocJpn File updated against INTSourceChangelist:3782314

Change 3837691 by [REDACTED_PERSON]

	#loc UE4DocJpn File updated against INTSourceChangelist:3782314

Change 3837695 by [REDACTED_PERSON]

	Updated file against INTSourceChangelist:3810467

Change 3837697 by [REDACTED_PERSON]

	Updated file against INTSourceChangelist:3810467

Change 3837698 by [REDACTED_PERSON]

	Updated file against INTSourceChangelist:3810467

Change 3837699 by [REDACTED_PERSON]

	Updated file against INTSourceChangelist:3810467

Change 3837703 by [REDACTED_PERSON]

	#loc UE4DocJpn File updated against INTSourceChangelist:3782314

Change 3837707 by [REDACTED_PERSON]

	Updated file against INTSourceChangelist:3810467

Change 3837708 by [REDACTED_PERSON]

	Updated file against INTSourceChangelist:3810467

Change 3837710 by [REDACTED_PERSON]

	Updated file against INTSourceChangelist:3810467

Change 3837713 by [REDACTED_PERSON]

	Updated file against INTSourceChangelist:3810467

Change 3837716 by [REDACTED_PERSON]

	Updated file against INTSourceChangelist:3810467

Change 3837806 by [REDACTED_PERSON]

	#loc UE4DocKor

Change 3839390 by [REDACTED_PERSON]

	#loc UE4DocJpn File updated against INTSourceChangelist:3782314

Change 3839428 by [REDACTED_PERSON]

	Updated file against INTSourceChangelist:3782314

Change 3839430 by [REDACTED_PERSON]

	#loc UE4DocJpn File updated against INTSourceChangelist:3782314

Change 3839440 by [REDACTED_PERSON]

	#loc UE4DocJpn File updated against INTSourceChangelist:3804170

Change 3839460 by [REDACTED_PERSON]

	#loc UE4DocJpn File updated against INTSourceChangelist:3782314

Change 3839476 by [REDACTED_PERSON]

	#loc UE4DocJpn File updated against INTSourceChangelist:3782314

Change 3839480 by [REDACTED_PERSON]

	#loc UE4DocJpn File updated againstINTSourceChangelist:3804170

Change 3839482 by [REDACTED_PERSON]

	#loc UE4DocJpn File updated against INTSourceChangelist:3782314

Change 3839501 by [REDACTED_PERSON]

	Updated file against INTSourceChangelist:3807090

Change 3839507 by [REDACTED_PERSON]

	#loc UE4DocJpn File updated against INTSourceChangelist:3810467

Change 3839522 by [REDACTED_PERSON]

	Updated file against INTSourceChangelist:3807090

Change 3839539 by [REDACTED_PERSON]

	Updated file against INTSourceChangelist:3807090

Change 3839547 by [REDACTED_PERSON]

	#loc UE4DocJpn File updated against INTSourceChangelist:3782314

Change 3839556 by [REDACTED_PERSON]

	Updated file against INTSourceChangelist:3807090

Change 3839558 by [REDACTED_PERSON]

	Updated file against INTSourceChangelist:3810467

Change 3839562 by [REDACTED_PERSON]

	#loc UE4DocJpn File updated against INTSourceChangelist:3782314

Change 3839565 by [REDACTED_PERSON]

	Updated file against INTSourceChangelist:3810467

Change 3839567 by [REDACTED_PERSON]

	#loc UE4DocJpn File updated against INTSourceChangelist:3782314

Change 3839568 by [REDACTED_PERSON]

	Updated file against INTSourceChangelist:3809606

Change 3839630 by [REDACTED_PERSON]

	#loc UE4DocJpn File updated against INTSourceChangelist:3782314

Change 3839632 by [REDACTED_PERSON]

	#loc UE4DocJpn File updated against INTSourceChangelist:3810467

Change 3839648 by [REDACTED_PERSON]

	#loc UE4DocJpn File updated against INTSourceChangelist:3782314

Change 3839649 by [REDACTED_PERSON]

	#loc UE4DocJpn File updated against INTSourceChangelist:3782314

Change 3839651 by [REDACTED_PERSON]

	#loc UE4DocJpn File updated against INTSourceChangelist:3750351

Change 3839654 by [REDACTED_PERSON]

	Updated file against INTSourceChangelist:3761485

Change 3839657 by [REDACTED_PERSON]

	Updated file against INTSourceChangelist:3750407

Change 3839659 by [REDACTED_PERSON]

	Updated file against INTSourceChangelist:3750348

Change 3839660 by [REDACTED_PERSON]

	Updated file against INTSourceChangelist:3750348

Change 3839662 by [REDACTED_PERSON]

	Updated file against INTSourceChangelist:3750348

Change 3839664 by [REDACTED_PERSON]

	#loc UE4DocJpn File updated against INTSourceChangelist:3782314

Change 3839665 by [REDACTED_PERSON]

	Updated file against INTSourceChangelist:3750348

Change 3839667 by [REDACTED_PERSON]

	Updated file against INTSourceChangelist:3750348

Change 3839668 by [REDACTED_PERSON]

	Updated file against INTSourceChangelist:3750348

Change 3839670 by [REDACTED_PERSON]

	Updated file against Updated file against INTSourceChangelist:3750348

Change 3839673 by [REDACTED_PERSON]

	Updated file against INTSourceChangelist:3736300

Change 3839682 by [REDACTED_PERSON]

	Updated file against INTSourceChangelist:3736300

Change 3839683 by [REDACTED_PERSON]

	Updated file against INTSourceChangelist:3736300

Change 3839684 by [REDACTED_PERSON]

	#loc UE4DocJpn File updated against INTSourceChangelist:3810557

Change 3839686 by [REDACTED_PERSON]

	Foliage Instanced Mesh system > Foliage Mode

Change 3839688 by [REDACTED_PERSON]

	#loc UE4DocJpn File updated against INTSourceChangelist:3782314

Change 3839690 by [REDACTED_PERSON]

	#loc UE4DocJpn File updated against INTSourceChangelist:3782314

Change 3839693 by [REDACTED_PERSON]

	Updated file against INTSourceChangelist:3736212

Change 3839696 by [REDACTED_PERSON]

	Updated file against INTSourceChangelist:3736212

Change 3839697 by [REDACTED_PERSON]

	Updated file against INTSourceChangelist:3736250

Change 3839699 by [REDACTED_PERSON]

	Updated file against INTSourceChangelist:3244370

Change 3839702 by [REDACTED_PERSON]

	Updated file against INTSourceChangelist:3243860

Change 3839706 by [REDACTED_PERSON]

	#loc UE4DocJpn File updated against INTSourceChangelist:3807090

Change 3839707 by [REDACTED_PERSON]

	Updated file against INTSourceChangelist:3243860

Change 3839708 by [REDACTED_PERSON]

	Updated file against INTSourceChangelist:3243860

Change 3839709 by [REDACTED_PERSON]

	Updated file against INTSourceChangelist:3243860

Change 3839710 by [REDACTED_PERSON]

	Updated file against INTSourceChangelist:3243860

Change 3839711 by [REDACTED_PERSON]

	Updated file against INTSourceChangelist:3243860

Change 3839712 by [REDACTED_PERSON]

	Updated file against INTSourceChangelist:3243860

Change 3839714 by [REDACTED_PERSON]

	Updated file against INTSourceChangelist:3243860

Change 3839715 by [REDACTED_PERSON]

	Updated file against INTSourceChangelist:3242329

Change 3839716 by [REDACTED_PERSON]

	Updated file against INTSourceChangelist:3242329


diff --git a/Engine/Documentation/Source/DocumentationGuidelines/BrandGlossary/F/F.INT.udn b/Engine/Documentation/Source/DocumentationGuidelines/BrandGlossary/F/F.INT.udn
--- a/Engine/Documentation/Source/DocumentationGuidelines/BrandGlossary/F/F.INT.udn
+++ b/Engine/Documentation/Source/DocumentationGuidelines/BrandGlossary/F/F.INT.udn
@@ -47,4 +47,4 @@ The Foliage Instanced Mesh system allows you to quickly paint or erase sets of S
 
 ##Foliage Mode
 
-The Foliage Instanced Mesh system enables you to quickly paint or erase batched sets of Static Meshes that use hardware instancing for efficient rendering on Landscape Actors, other Static Mesh Actors, and BSP geometry.
\ No newline at end of file
+The Foliage Mode enables you to quickly paint or erase batched sets of Static Meshes that use hardware instancing for efficient rendering on Landscape Actors, other Static Mesh Actors, and BSP geometry.
\ No newline at end of file
diff --git a/Engine/Documentation/Source/DocumentationGuidelines/Syntax/EpicMarkdownSyntax.INT.udn b/Engine/Documentation/Source/DocumentationGuidelines/Syntax/EpicMarkdownSyntax.INT.udn
--- a/Engine/Documentation/Source/DocumentationGuidelines/Syntax/EpicMarkdownSyntax.INT.udn
+++ b/Engine/Documentation/Source/DocumentationGuidelines/Syntax/EpicMarkdownSyntax.INT.udn
@@ -123,7 +123,7 @@ It was desirable to have the ability to provide some formatting instructions for
     [refepiclogo]: Logo_Epic-New.jpg "Here's a title"(w:50 h:50 a:left convert:true quality:75 fill:#000000)
 	This text is in a paragraph alongside the image. A special div element follows ending floating.
 	[REGION:clear]
-		[COMMENT]
+		[COMMENT:none]
  clear float 
 [/COMMENT]
 	[/REGION]    
@@ -133,7 +133,7 @@ It was desirable to have the ability to provide some formatting instructions for
 	<p>This is the reference image <img width="50" height="50" style="float: left;" title="Here's a title" alt="Epic Logo" src="./../../../images\EpicMarkdownDocGuides\EpicMarkdownExtension\Logo_Epic-New.jpg"/> with width, height, and left float.</p>
 	<p>This text is in a paragraph alongside the image. A special div element follows ending floating.</p>
 	[REGION:clear]
-		[COMMENT]
+		[COMMENT:none]
  clear float 
 [/COMMENT]
 	[/REGION]
@@ -148,7 +148,7 @@ This is the reference image ![Epic Logo][refepiclogo] with width, height, and le
 This text is in a paragraph alongside the image. A special div element follows ending floating.
 
 [REGION:clear]
-	[COMMENT]
+	[COMMENT:none]
  clear float 
 [/COMMENT]
 [/REGION]
diff --git a/Engine/Documentation/Source/Engine/AI/BehaviorTrees/HowUE4BehaviorTreesDiffer/HowUE4BehaviorTreesDiffer.INT.udn b/Engine/Documentation/Source/Engine/AI/BehaviorTrees/HowUE4BehaviorTreesDiffer/HowUE4BehaviorTreesDiffer.INT.udn
--- a/Engine/Documentation/Source/Engine/AI/BehaviorTrees/HowUE4BehaviorTreesDiffer/HowUE4BehaviorTreesDiffer.INT.udn
+++ b/Engine/Documentation/Source/Engine/AI/BehaviorTrees/HowUE4BehaviorTreesDiffer/HowUE4BehaviorTreesDiffer.INT.udn
@@ -5,9 +5,7 @@ Description:How Unreal Engine 4 Behavior Trees Differ From Traditional Behavior
 
 [TOC(start:2 end:3)]
 
-This section is for people who are familiar with **Behavior Trees** generally and would like to dive into the Unreal Engine 4 implementation as quickly as possible.  For those who have not used behavior trees before, you may find some of the explanation here confusing. [COMMENT:none]
- You may want to check out the user guide first. You can return to this section if you're interested in what we've done differently than others. 
-[/COMMENT]
+This section is for people who are familiar with **Behavior Trees** generally and would like to dive into the Unreal Engine 4 implementation as quickly as possible.  For those who have not used behavior trees before, you may find some of the explanation here confusing.
 
 There are three critical ways in which the UE4 implementation of Behavior Trees differs from "standard" behavior trees:
 
@@ -19,9 +17,7 @@ Having an event-driven architecture grants improvements to both performance and
 
 Since the code does not have to iterate through the entire tree every tick, performance is much better!  Conceptually, instead of constantly asking "Are we there yet?", we can just rest until we are prodded and told "We are there!"
 
-When stepping forward and backward through the behavior tree's execution history to visually debug the behavior, it is ideal to make the history show relevant changes and not show irrelevant ones.  In the event-driven implementation, it is not necessary to filter out irrelevant steps that iterated over the tree and chose the same behavior as before, because that additional iteration never had to happen in the first place!  Instead, only changes to the location of execution in the tree or to blackboard values matter, and it is easy to just show those differences.  [COMMENT:none]
- (See the debugging section for more information on how to step through the execution history and other debugging tricks.) 
-[/COMMENT]
+When stepping forward and backward through the behavior tree's execution history to visually debug the behavior, it is ideal to make the history show relevant changes and not show irrelevant ones.  In the event-driven implementation, it is not necessary to filter out irrelevant steps that iterated over the tree and chose the same behavior as before, because that additional iteration never had to happen in the first place!  Instead, only changes to the location of execution in the tree or to blackboard values matter, and it is easy to just show those differences.
 
 ## Conditionals Are Not Leaf Nodes
 
@@ -35,9 +31,7 @@ First, conditional decorators make the behavior tree UI more intuitive and easie
 
 In this section of a Behavior Tree, the Decorators **Close Enough** and **Blackboard** can prevent the execution of the Sequence node's children.
 
-Another advantage of conditional decorators is that it is easy to make those decorators act as observers (waiting for events) at critical nodes in the tree.  This feature is critical to gaining full advantage from the event-driven nature of the trees.  [COMMENT:none]
- Read more about how to use Conditional Decorators and their special properties here. 
-[/COMMENT]
+Another advantage of conditional decorators is that it is easy to make those decorators act as observers (waiting for events) at critical nodes in the tree.  This feature is critical to gaining full advantage from the event-driven nature of the trees. 
 
 ## Special Handling for Concurrent Behaviors
 
@@ -81,20 +75,10 @@ For example, a service can be used to determine which enemy is the best choice f
 
 Services are active only as long as execution remains in the subtree rooted at the composite node with which the service is associated.
 
-[COMMENT:none]
- Add screenshot showing what Services look like. 
-[/COMMENT]
-
 #### Decorator "Observer Aborts" Property
 
 One common usage case for standard Parallel nodes is to constantly check conditions so that a task can abort if the conditions it requires becomes false.  For example, if you have a cat that performs a sequence, such as "Shake Rear End" and "Pounce", you may want to give up immediately if the mouse escapes into its mouse hole.  With Parallel nodes, you would have a child that checks if the mouse can be pounced on, and then another child that's the sequence to perform.  Since our Behavior Trees are event-driven, we instead handle this by having our conditional decorators observe their values and abort when necessary.  (In this example, you would just have the "Mouse Can Be Pounced On?" decorator on the sequence itself, with "Observer Aborts" set to "Self". )
 
-[COMMENT:none]
- For more information, see the <LINK>"Observer Aborts" explanation page. 
-[/COMMENT]
-
-
-
 
 ## Advantages of UE4's Approach to Concurrent Behaviors
 
diff --git a/Engine/Documentation/Source/Engine/Animation/AnimPose/AnimBPPoseNodes/AnimBPPoseNodes.INT.udn b/Engine/Documentation/Source/Engine/Animation/AnimPose/AnimBPPoseNodes/AnimBPPoseNodes.INT.udn
--- a/Engine/Documentation/Source/Engine/Animation/AnimPose/AnimBPPoseNodes/AnimBPPoseNodes.INT.udn
+++ b/Engine/Documentation/Source/Engine/Animation/AnimPose/AnimBPPoseNodes/AnimBPPoseNodes.INT.udn
@@ -126,7 +126,7 @@ On the **Pose Driver** node itself, you have the following options:
 | **Type** | [INCLUDE:Engine\Animation\AnimPose\AnimBPPoseNodes\#typetable] | 
 | **Twist Axis** | The axis to use when **Type** is set to **SwingOnly** (X, Y or Z). | 
 
-[COMMENT]
+[COMMENT:none]
 
 [EXCERPT:typetable]
 | Option | Description |
diff --git a/Engine/Documentation/Source/Engine/Animation/Persona/PersonaEditor.INT.udn b/Engine/Documentation/Source/Engine/Animation/Persona/PersonaEditor.INT.udn
--- a/Engine/Documentation/Source/Engine/Animation/Persona/PersonaEditor.INT.udn
+++ b/Engine/Documentation/Source/Engine/Animation/Persona/PersonaEditor.INT.udn
@@ -125,7 +125,7 @@ See the individual editor pages for information on their interfaces and usage:
 
 
 
-[COMMENT]
+[COMMENT:none]
 
 [EXCERPT:SkeletonOverview]
 **Skeleton mode** allows for visualization and control over those aspects of a character or Skeletal Mesh involving its bone or joint hierarchy. Example uses include:
@@ -163,10 +163,6 @@ See the individual editor pages for information on their interfaces and usage:
 * Blending between different poses and animation results using the AnimGraph.
 [/EXCERPT:GraphOverview]
 
-[/COMMENT]
-
-[COMMENT]
-
 [EXCERPT:SkeletonUI]
 [REGION:simpletable]
 |Available Interface Components						|																					|
diff --git a/Engine/Documentation/Source/Engine/Animation/Persona/Viewport/PersonaViewport.INT.udn b/Engine/Documentation/Source/Engine/Animation/Persona/Viewport/PersonaViewport.INT.udn
--- a/Engine/Documentation/Source/Engine/Animation/Persona/Viewport/PersonaViewport.INT.udn
+++ b/Engine/Documentation/Source/Engine/Animation/Persona/Viewport/PersonaViewport.INT.udn
@@ -112,8 +112,62 @@
 | **Mute Audio** | Mutes audio from the preview. |
 | **Use Audio Attenuation** | Use audio attenuation when playing back audio in the preview. |
 
+### Physics
+
+The **Physics** button provides the following options: 
+
+![](PhysicsMenu.png)
+
+| Option | Description |
+| ------ | ----------- |
+| [REGION:tablesection] Anim [/REGION] ||
+| **Physics Bend** | Lets you manually control the physics and animation. |
+| **Update Joint from Animations** | When enabled, any joints will update from its animations while simulating. |
+| **Physics Update Mode** | Determines whether the simulations of the root body also updates the component transform. |
+| **Poke Pause Time** | The time between poking the ragdoll and starting to blend back. |
+| **Poke Blend Time** | The time that is taken to blend from physics to the animation. |
+| [REGION:tablesection] Simulation [/REGION] ||
+| **Grav Scale** | Scale factor for the gravity used while simulating. |
+| **Max FPS** | The max frames per second (FPS) for simulation in the Physics Asset Editor. This is helpful for targeting the same FPS as your game. A value of -1 will disable a set Max FPS. |
+| [REGION:tablesection] Mouse Spring [/REGION] ||
+| **Handle Linear Damping** | Linear damping for mouse spring forces. |
+| **Handle Linear Stiffness** | Linear stiffness for mouse spring forces. |
+| **Handle Angular Damping** | Angular damping for mouse springs forces.  |
+| **Handle Angular Stiffness** | Angular stiffness for mouse spring forces. |
+| **Interpolation Speed** | How quickly interpolation takes place for the physics target transform for mouse spring forces. |
+| [REGION:tablesection] Poking [/REGION] ||
+| **Poke Strength** | The strength of the impulse used when poking with left mouse button. |
+
+
+### Preview LODs
+
+The **Preview LODs** button provides the following options: 
+
+![](lodmenu.png)
+
+| Option | Description |
+| ------ | ----------- |
+| **LOD Auto** | Automatically selects the Level Of Detail (LOD) (this is the default setting). |
+| **LOD 0** | Selects LOD 0 as the LOD. |
+
+### Playback Speed
+
+The **Playback Speed** button provides the following options: 
+
+![](playbackspeed.png)
+
+| Option | Description |
+| ------ | ----------- |
+| **x0.1** | Sets the playback speed to 0.1.|
+| **x0.25** | Sets the playback speed to 0.25. |
+| **x0.5** | Sets the playback speed to 0.5. |
+| **x1.0** | Sets the playback speed to 1.0 (this is the default setting). |
+| **x2.0** | Sets the playback speed to 2.0. |
+| **x5.0** | Sets the playback speed to 5.0. |
+| **x10.0** | Sets the playback speed to 10.0. |
+[/EXCERPT:viewoptions]
+
 [COMMENT:none]
--
 [INCLUDE:#advanced]
 
 [EXCERPT:Animation]
@@ -178,67 +232,8 @@
 | **Mesh Info** | How mesh info should be displayed in the viewport. |
 | **Mesh Overlay Drawing** | How the mesh should be displayed while in the viewport. |
 [/EXCERPT:mesh]
-
--
 [/COMMENT]
 
-
-### Physics
-
-The **Physics** button provides the following options: 
-
-![](PhysicsMenu.png)
-
-| Option | Description |
-| ------ | ----------- |
-| [REGION:tablesection] Anim [/REGION] ||
-| **Physics Bend** | Lets you manually control the physics and animation. |
-| **Update Joint from Animations** | When enabled, any joints will update from its animations while simulating. |
-| **Physics Update Mode** | Determines whether the simulations of the root body also updates the component transform. |
-| **Poke Pause Time** | The time between poking the ragdoll and starting to blend back. |
-| **Poke Blend Time** | The time that is taken to blend from physics to the animation. |
-| [REGION:tablesection] Simulation [/REGION] ||
-| **Grav Scale** | Scale factor for the gravity used while simulating. |
-| **Max FPS** | The max frames per second (FPS) for simulation in the Physics Asset Editor. This is helpful for targeting the same FPS as your game. A value of -1 will disable a set Max FPS. |
-| [REGION:tablesection] Mouse Spring [/REGION] ||
-| **Handle Linear Damping** | Linear damping for mouse spring forces. |
-| **Handle Linear Stiffness** | Linear stiffness for mouse spring forces. |
-| **Handle Angular Damping** | Angular damping for mouse springs forces.  |
-| **Handle Angular Stiffness** | Angular stiffness for mouse spring forces. |
-| **Interpolation Speed** | How quickly interpolation takes place for the physics target transform for mouse spring forces. |
-| [REGION:tablesection] Poking [/REGION] ||
-| **Poke Strength** | The strength of the impulse used when poking with left mouse button. |
-
-
-### Preview LODs
-
-The **Preview LODs** button provides the following options: 
-
-![](lodmenu.png)
-
-| Option | Description |
-| ------ | ----------- |
-| **LOD Auto** | Automatically selects the Level Of Detail (LOD) (this is the default setting). |
-| **LOD 0** | Selects LOD 0 as the LOD. |
-
-### Playback Speed
-
-The **Playback Speed** button provides the following options: 
-
-![](playbackspeed.png)
-
-| Option | Description |
-| ------ | ----------- |
-| **x0.1** | Sets the playback speed to 0.1.|
-| **x0.25** | Sets the playback speed to 0.25. |
-| **x0.5** | Sets the playback speed to 0.5. |
-| **x1.0** | Sets the playback speed to 1.0 (this is the default setting). |
-| **x2.0** | Sets the playback speed to 2.0. |
-| **x5.0** | Sets the playback speed to 5.0. |
-| **x10.0** | Sets the playback speed to 10.0. |
-[/EXCERPT:viewoptions]
-
-
 ## Movement Toolbar
 
 [EXCERPT:movement]
diff --git a/Engine/Documentation/Source/Engine/Animation/RootMotion/RootMotion.INT.udn b/Engine/Documentation/Source/Engine/Animation/RootMotion/RootMotion.INT.udn
--- a/Engine/Documentation/Source/Engine/Animation/RootMotion/RootMotion.INT.udn
+++ b/Engine/Documentation/Source/Engine/Animation/RootMotion/RootMotion.INT.udn
@@ -183,7 +183,7 @@ YuD29-Xr7Oc
 The physics state of the Character is taken into account during Root Motion. For example, if the character physics is Walking or Falling, then the Z-axis of Root Motion will be ignored, and gravity will be applied. Characters will fall, go down slopes or up stairs. If the character physics is Flying, then the full Root Motion is applied and gravity is ignored.
 
 
-[COMMENT]
+[COMMENT:none]
  This should be a HowTo. It's actually a checkbox now instead of using AnimMontages
 ## Setting Up Root Motion
 
diff --git a/Engine/Documentation/Source/Engine/Animation/Sequences/Notifies/AnimNotifies.INT.udn b/Engine/Documentation/Source/Engine/Animation/Sequences/Notifies/AnimNotifies.INT.udn
--- a/Engine/Documentation/Source/Engine/Animation/Sequences/Notifies/AnimNotifies.INT.udn
+++ b/Engine/Documentation/Source/Engine/Animation/Sequences/Notifies/AnimNotifies.INT.udn
@@ -158,7 +158,7 @@ To use an AnimTrail emitter you must:
 1. Create a [Particle System](Engine/Rendering/ParticleSystems) that contains an emitter with AnimTrail type data. 
 2. Create and add a Notify State of the **Trail** type to an animation.
 
-[COMMENT]
+[COMMENT:none]
  To be included as part of a HowTo.
 Setting Up the Particle System
 
@@ -231,7 +231,7 @@ There are also some debug rendering options available:
 | **Render Tangents** | This toggles the debug rendering of the curve tangents. |
 | **Render Tessellation** | This toggles debug rendering of all vertices on the curve. This shows how the curve is tessellated between the spawn points. |
 
-[COMMENT]
+[COMMENT:none]
  To be included with AnimTrail HowTo.
 ### Controlling Animation Trails in Blueprint
 While the above method is the standard approach to controlling AnimTrails, you can also control them from inside a Blueprint without the need of an Anim Notify. 
diff --git a/Engine/Documentation/Source/Engine/Animation/SyncGroups/SyncGroups.INT.udn b/Engine/Documentation/Source/Engine/Animation/SyncGroups/SyncGroups.INT.udn
--- a/Engine/Documentation/Source/Engine/Animation/SyncGroups/SyncGroups.INT.udn
+++ b/Engine/Documentation/Source/Engine/Animation/SyncGroups/SyncGroups.INT.udn
@@ -43,9 +43,7 @@ Leaders and followers are designated simply by which animation has the greater b
 |[INCLUDE:#walkanimationscaled] |In this instance, the blend weight for the run animation would be greater than the walk animation. Since they are in the same Sync Group, the walk scales down to match the length of the run animation.|
 |[INCLUDE:#runanimationscaled]|Here, the weighting has switched such that the walking animation is receiving a greater blend weight. This causes the run animation to scale up to match the length of the walk. | 
 
-[COMMENT]
--
-
+[COMMENT:none]
 
 [EXCERPT:RunAnimation]
 [REGION:raw]
@@ -81,7 +79,6 @@ Leaders and followers are designated simply by which animation has the greater b
 * Has its Notifies suppressed, deferring instead to Leader's Notifies.
 [/EXCERPT:FollowerList]
 
---
 [/COMMENT]
 
 
@@ -96,13 +93,11 @@ To set the Sync Group, select the animation node in the AnimGraph and look in th
 |**Group Name**|This is where the Sync Group can be set. Generally, you will want all related animations that are meant to blend together (such as body motions for sneak, walk, run, sprint) to be in the same Sync Group.|
 |**Group Role**|This is where you set the type of membership in the group. The available values are "Can Be Leader" (the default), "Always Leader", and "Never Leader".|
 
-[COMMENT]
--
+[COMMENT:none]
 Group Role has 3 options:
 * Can Be Leader - This is the default setting. This means that if this animation has the greatest blend weight, then it can become the Leader and other animations will scale to match its length.
 * Always Leader - If set, then this animation will always be the Leader and all animations will always scale to match its length.
 * Never Leader - If set, then this animation can never be the Leader and will always scale to match the next weighted animation.
--
 [/COMMENT]
 
 
diff --git a/Engine/Documentation/Source/Engine/Audio/WAV/WAV.INT.udn b/Engine/Documentation/Source/Engine/Audio/WAV/WAV.INT.udn
--- a/Engine/Documentation/Source/Engine/Audio/WAV/WAV.INT.udn
+++ b/Engine/Documentation/Source/Engine/Audio/WAV/WAV.INT.udn
@@ -31,41 +31,7 @@ The engine currently supports importing uncompressed little endian 16 bit wave f
 
 Importing with the normal sound factory yields a simple sound node wave that can be referenced by a sound cue for playback. These sounds can be mono or stereo. Importing with a special naming convention gives the potential of multichannel (e.g. 5.1) sounds. Up to 8 channel sounds can be compressed on all platforms, although not all platforms can play them back efficiently. The package saving process does all the necessary work to convert to the platform's native format. There is no strict speaker mapping, so the number of channels infers which channels are played on which speakers (see table below).
 
-[EXCERPT:multiChannelExtensionTable]
-
-[REGION:raw]
-|                   | Extension | 4.0         | 5.1         | 6.1         | 7.1         |
-| ---               | ---       | ---         | ---         | ---         | ---         |
-| **FrontLeft**     | _fl       |  &#x2713;   |  &#x2713;   |  &#x2713;   |  &#x2713;   |
-| **FrontRight**    | _fr       |  &#x2713;   |  &#x2713;   |  &#x2713;   |  &#x2713;   |
-| **FrontCenter**   | _fc       |             |  &#x2713;   |  &#x2713;   |  &#x2713;   |
-| **LowFrequency**  | _lf       |             |  &#x2713;   |  &#x2713;   |  &#x2713;   |
-| **SideLeft**      | _sl       |  &#x2713;   |  &#x2713;   |  &#x2713;   |  &#x2713;   |
-| **SideRight**     | _sr       |  &#x2713;   |  &#x2713;   |  &#x2713;   |  &#x2713;   |
-| **BackLeft**      | _bl       |             |             |  &#x2713;   |  &#x2713;   |
-| **BackRight**     | _br       |             |             |             |  &#x2713;   |
-[/REGION]
-
-[/EXCERPT:multiChannelExtensionTable]
-
-[COMMENT:none]
-
-
-[REGION:raw]
-|                   | Extension | 4.0 | 5.1 | 6.1 | 7.1 |  
-| --- | --- | --- | --- | --- | --- | 
-| **FrontLeft**     | _fl       | ![](icon_MatEd_Apply_40x.png) | ![](icon_MatEd_Apply_40x.png)   | ![](icon_MatEd_Apply_40x.png)   | ![](icon_MatEd_Apply_40x.png)   | 
-| **FrontRight**    | _fr       | ![](icon_MatEd_Apply_40x.png)   | ![](icon_MatEd_Apply_40x.png)   | ![](icon_MatEd_Apply_40x.png)   | ![](icon_MatEd_Apply_40x.png)   | 
-| **FrontCenter**   | _fc       |     | ![](icon_MatEd_Apply_40x.png)   | ![](icon_MatEd_Apply_40x.png)   | ![](icon_MatEd_Apply_40x.png)   | 
-| **LowFrequency**  | _lf       |     | ![](icon_MatEd_Apply_40x.png)   | ![](icon_MatEd_Apply_40x.png)   | ![](icon_MatEd_Apply_40x.png)   | 
-| **SideLeft**      | _sl       | ![](icon_MatEd_Apply_40x.png)   | ![](icon_MatEd_Apply_40x.png)   | ![](icon_MatEd_Apply_40x.png)   | ![](icon_MatEd_Apply_40x.png)   | 
-| **SideRight**     | _sr       | ![](icon_MatEd_Apply_40x.png)   | ![](icon_MatEd_Apply_40x.png)   | ![]( icon_MatEd_Apply_40x.png)   | ![](icon_MatEd_Apply_40x.png)   | 
-|**BackLeft**      | _bl       |     |     | ![](icon_MatEd_Apply_40x.png)   | ![](icon_MatEd_Apply_40x.png)   |  
-| **BackRight**     | _br       |     |     |     | ![](icon_MatEd_Apply_40x.png)   |  
-[/REGION]
-
-
-[/COMMENT]
+[INCLUDE:#multiChannelExtensionTable]
 
 Compression settings are tweaked to maintain similar quality across platforms. After testing many sounds, our musician determined an XMA quality of 40 was equivalent to an ogg vorbis quality of 0.15. 
 
@@ -98,7 +64,24 @@ Or
 
 [/EXCERPT:wav01]
 
+[COMMENT:none]
+[EXCERPT:multiChannelExtensionTable]
 
+[REGION:raw]
+|                   | Extension | 4.0         | 5.1         | 6.1         | 7.1         |
+| ---               | ---       | ---         | ---         | ---         | ---         |
+| **FrontLeft**     | _fl       |  &#x2713;   |  &#x2713;   |  &#x2713;   |  &#x2713;   |
+| **FrontRight**    | _fr       |  &#x2713;   |  &#x2713;   |  &#x2713;   |  &#x2713;   |
+| **FrontCenter**   | _fc       |             |  &#x2713;   |  &#x2713;   |  &#x2713;   |
+| **LowFrequency**  | _lf       |             |  &#x2713;   |  &#x2713;   |  &#x2713;   |
+| **SideLeft**      | _sl       |  &#x2713;   |  &#x2713;   |  &#x2713;   |  &#x2713;   |
+| **SideRight**     | _sr       |  &#x2713;   |  &#x2713;   |  &#x2713;   |  &#x2713;   |
+| **BackLeft**      | _bl       |             |             |  &#x2713;   |  &#x2713;   |
+| **BackRight**     | _br       |             |             |             |  &#x2713;   |
+[/REGION]
+
+[/EXCERPT:multiChannelExtensionTable]
+[/COMMENT]
 
 [REGION:note]
 Exporting a multi channel sound asset will create a mono .WAV file for each channel with the above speaker extension naming convention.
@@ -107,7 +90,5 @@ Exporting a multi channel sound asset will create a mono .WAV file for each chan
 
 ## Related
 
-[REGION:topics half]
 %Engine/Audio:topiccompact%
 %Engine/Audio/Overview:topiccompact%
-[/REGION]
diff --git a/Engine/Documentation/Source/Engine/Basics/Projects/UnrealGameProjects.INT.udn b/Engine/Documentation/Source/Engine/Basics/Projects/UnrealGameProjects.INT.udn
--- a/Engine/Documentation/Source/Engine/Basics/Projects/UnrealGameProjects.INT.udn
+++ b/Engine/Documentation/Source/Engine/Basics/Projects/UnrealGameProjects.INT.udn
@@ -11,11 +11,7 @@ related:Programming/Development/CompilingProjects
 related:Programming/Development/ManagingGameCode
 
 [EXCERPT:Main]
-[EXCERPT:Intro]
-
-[INCLUDE:Shared/Glossary/P/#Glossary_project]
-
-[/EXCERPT:Intro]
+[INCLUDE:#intro]
 
 %Engine/Basics/Projects/Browser:topic%
 [OBJECT:Topic]
@@ -33,20 +29,12 @@ related:Programming/Development/ManagingGameCode
 	[/PARAM]
 [/OBJECT]
 
+[/EXCERPT:Main]
 
 [COMMENT:none]
+[EXCERPT:Intro]
 
-Projects can be one of two types:
-
-* **Content-Only** - These are projects where all development is done through the editor using content assets such as maps, art assets, _Blueprints_, etc.
-* **Code** - These projects allow programmer development in addition to regular content development available in _content-only_ projects. 
-
-[REGION:tip]
-Content projects can be converted to code projects using the [](Programming/Development/ManagingGameCode/CppClassWizard).
-[/REGION]
-
-
-
-[/COMMENT]
+[INCLUDE:Shared/Glossary/P/#Glossary_project]
 
-[/EXCERPT:Main]
+[/EXCERPT:Intro]
+[/COMMENT]
\ No newline at end of file
diff --git a/Engine/Documentation/Source/Engine/Blueprints/BP_HowTo/BlueprintHowTo.INT.udn b/Engine/Documentation/Source/Engine/Blueprints/BP_HowTo/BlueprintHowTo.INT.udn
--- a/Engine/Documentation/Source/Engine/Blueprints/BP_HowTo/BlueprintHowTo.INT.udn
+++ b/Engine/Documentation/Source/Engine/Blueprints/BP_HowTo/BlueprintHowTo.INT.udn
@@ -22,30 +22,12 @@ When you are ready and have an understanding of the core aspects behind Blueprin
 [DIR(output:"fancy" type:"how-to" skill:"Beginner" parent:"Engine/Blueprints/BP_HowTo" org:"hierarchy")]
 
 ## Intermediate
-[REGION:wrapper transparent]
-[](FANCY:Engine/Blueprints/BP_HowTo/ActorReference "%Engine/Blueprints/BP_HowTo/ActorReference:description%")
-[](FANCY:Engine/Blueprints/BP_HowTo/MakingFunctions "%Engine/Blueprints/BP_HowTo/MakingFunctions:description%")
-[](FANCY:Engine/Blueprints/BP_HowTo/WorkingWithArrays "%Engine/Blueprints/BP_HowTo/WorkingWithArrays:description%")
-[](FANCY:Engine/Blueprints/UserGuide/Variables "%Engine/Blueprints/UserGuide/Variables:description%")
-[/REGION]
 
-## Advanced
-[REGION:wrapper transparent]
-[](FANCY:Engine/Blueprints/BP_HowTo/DirectBlueprintComs "%Engine/Blueprints/BP_HowTo/DirectBlueprintComs:description%")
-[](FANCY:Engine/Blueprints/BP_HowTo/Casting "%Engine/Blueprints/BP_HowTo/Casting:description%")
-[](FANCY:Engine/Blueprints/BP_HowTo/Interfaces "%Engine/Blueprints/BP_HowTo/Interfaces:description%")
-[](FANCY:Engine/Blueprints/BP_HowTo/EventDispatcher "%Engine/Blueprints/BP_HowTo/EventDispatcher:description%")
-[](FANCY:Engine/Blueprints/BP_HowTo/MakingMacros "%Engine/Blueprints/BP_HowTo/MakingMacros:description%")
-[](FANCY:Engine/Blueprints/BP_HowTo/MacroLibraries "%Engine/Blueprints/BP_HowTo/MacroLibraries:description%")
-[/REGION]
-
-[COMMENT:none]
- To remove since information is redundant. 
-
-| **[](Engine/Blueprints/BP_HowTo/BasicUsage)** | %Engine/Blueprints/BP_HowTo/BasicUsage:description% |
+[DIR(output:"fancy" type:"how-to" skill:"Intermediate" parent:"Engine/Blueprints/BP_HowTo" org:"hierarchy")]
 
+## Advanced
 
-[/COMMENT]
+[DIR(output:"fancy" type:"how-to" skill:"Advanced" parent:"Engine/Blueprints/BP_HowTo" org:"hierarchy")]
 
 
 
diff --git a/Engine/Documentation/Source/Engine/Blueprints/Editor/UIComponents/GraphEditor/GraphEditor.INT.udn b/Engine/Documentation/Source/Engine/Blueprints/Editor/UIComponents/GraphEditor/GraphEditor.INT.udn
--- a/Engine/Documentation/Source/Engine/Blueprints/Editor/UIComponents/GraphEditor/GraphEditor.INT.udn
+++ b/Engine/Documentation/Source/Engine/Blueprints/Editor/UIComponents/GraphEditor/GraphEditor.INT.udn
@@ -31,6 +31,20 @@ The **Graph Editor** panel is the heart of the Blueprint system. It is here that
 
 The **Graph Editor** tab can be navigated using the controls below:
 
+[INCLUDE:#Controls]
+
+[REGION:tip]
+StaticMesh, SoundCue, SkeletalMesh, and ParticleSystem assets can be dragged from the **Content Browser** and 
+dropped on the **Graph Editor** tab to create a new AddComponent function call with the asset automatically assigned.
+
+[REGION:raw]
+![Blueprints - Drag and Drop Component](k2_component_graph_drag.png)(w:900 convert:false)
+[/REGION:raw]
+
+[/REGION:tip]
+[/EXCERPT:Main]
+
+[COMMENT:none]
 [EXCERPT:Controls]
 
 | Control | Action |
@@ -49,16 +63,4 @@ The **Graph Editor** tab can be navigated using the controls below:
 | **Alt + Click** on a pin | Remove all wires connected to the selected pin. |
 
 [/EXCERPT:Controls]
-
-
-[REGION:tip]
-StaticMesh, SoundCue, SkeletalMesh, and ParticleSystem assets can be dragged from the **Content Browser** and 
-dropped on the **Graph Editor** tab to create a new AddComponent function call with the asset automatically assigned.
-
-[REGION:raw]
-![Blueprints - Drag and Drop Component](k2_component_graph_drag.png)(w:900 convert:false)
-[/REGION:raw]
-
-[/REGION:tip]
-[/EXCERPT:Main]
-
+[/COMMENT]
\ No newline at end of file
diff --git a/Engine/Documentation/Source/Engine/Blueprints/Editor/UIComponents/Menu/BlueprintEditorMenu.INT.udn b/Engine/Documentation/Source/Engine/Blueprints/Editor/UIComponents/Menu/BlueprintEditorMenu.INT.udn
--- a/Engine/Documentation/Source/Engine/Blueprints/Editor/UIComponents/Menu/BlueprintEditorMenu.INT.udn
+++ b/Engine/Documentation/Source/Engine/Blueprints/Editor/UIComponents/Menu/BlueprintEditorMenu.INT.udn
@@ -13,11 +13,61 @@
 
 [EXCERPT:Main]
 
-
-
-
 ### File
 
+[INCLUDE:#FileMenu]
+
+### Edit
+
+[VAR:BlueprintSearch]
+| **Search** | Find references to functions, events, variables, and pins in the current Blueprint.|
+| **Find in Blueprints** | Find references to functions, events, variables, and pins in ALL Blueprints. |
+| **Delete Unused Variables** | Deletes any variables that are never used. |
+[/VAR]
+
+[INCLUDE:#EditMenu]
+
+
+### Asset
+
+[INCLUDE:#AssetMenu]
+
+### View
+
+[INCLUDE:ViewMenu]
+
+### Debug
+
+[INCLUDE:#DebugMenu]
+
+### Window 
+
+The **Window** menu in the **Blueprint Editor** has a specific subsection for displaying Blueprint Editor-specific tabs. Not all tabs will be present in the menu while the Blueprint Editor is in the **Defaults** and **Components** modes.
+
+
+| Command | Description |
+| --- | --- |
+| **Toolbar** | Shows the Toolbar if it is currently not visible. |
+| **Details** | Shows the [Details](Engine/Blueprints\Editor/UIComponents\Details) pane if it is currently not visible. |
+| **Debug** | Shows the **Debug** pane if if it currently not visible. |
+| **Palette** | Shows the [](Engine/Blueprints\Editor/UIComponents\Palette) pane if it is not currently visible. |
+| **My Blueprint** | Shows the [](Engine/Blueprints\Editor/UIComponents\MyBlueprint) pane if it is not currently visible.  |
+| **Compiler Results** | Shows the **Compiler Results** pane if it is not currently visible.  |
+| **Find Results** | Shows the **Find Results** pane if it is not currently visible.  | 
+| **Components** | Shows the **Components panel** if it is not currently visible. |
+| **Viewport** | Shows the **Preview Viewport** pane if it is not currently visible.  |
+| **Content Browser** | Opens a sub-menu with access to all four **Content Browsers**. |
+| **Developer Tools** | [INCLUDE:#devtools] |
+| **Project Launcher** | Shows the Project Launcher which enables you to run your project on any appropriately setup and connected devices. |
+| **Plugins** | Shows the Plugin Tab where you can load/unload plugins. |
+| **Reset Layout...** | Restores the default layout for the entire Unreal Engine 4 Editor. This requires the Editor to restart, but it will re-open the current project. |
+| **Save Layout** | Saves the current interface layout. |
+| **Enable Fullscreen** | Enables fullscreen mode for the Main Editor window. |
+
+
+[/EXCERPT:Main]
+
+[COMMENT:none]
 [EXCERPT:FileMenu]
 | Command | Description |
 | --- | --- |
@@ -34,15 +84,6 @@
 | **Diff** | Diff against previous revisions. Requires source control to be active. |
 | **Developer** | Open the Developer Menu, where you can change compiler settings and recompile modules like the Graph Editor. |
 [/EXCERPT:FileMenu]
-
-### Edit
-
-[VAR:BlueprintSearch]
-| **Search** | Find references to functions, events, variables, and pins in the current Blueprint.|
-| **Find in Blueprints** | Find references to functions, events, variables, and pins in ALL Blueprints. |
-| **Delete Unused Variables** | Deletes any variables that are never used. |
-[/VAR]
-
 [EXCERPT:EditMenu]
 | Command | Description |
 | --- | --- |
@@ -57,9 +98,6 @@
 | **Project Settings** | Opens the settings for the current project. |
 | **Plugins** | Opens the Plugin Browser tab. |
 [/EXCERPT:EditMenu]
-
-
-### Asset
 [EXCERPT:AssetMenu]
 [VAR:AssetMenu]
 | Command | Description |
@@ -71,9 +109,6 @@
 
 %assetmenu%
 [/EXCERPT:AssetMenu]
-
-### View
-
 [EXCERPT:ViewMenu]
 | Command | Description |
 | --- | --- |
@@ -82,12 +117,9 @@
 | **Hide Unused Pins** | Hides all pins with no connections and no default values. |
 | **Hide Unconnected Pins** | Hides all pins with no connections.  This option will hide input pins for inputs that have been set directly on the node. |
 | [REGION:tablesection]**Zoom**[/REGION] ||
-| **Zoom to Graph Extents ** | Fits the current view to the entire graph. |
+| **Zoom to Graph Extents** | Fits the current view to the entire graph. |
 | **Zoom to Selection** | Fits the current view to the selection.  |
 [/EXCERPT:ViewMenu]
-
-### Debug
-
 [EXCERPT:DebugMenu]
 | Command | Description |
 | --- | --- |
@@ -98,36 +130,6 @@
 | [REGION:tablesection]**Watches**[/REGION] ||
 | **Delete All Watches** | Removes all watch values in all graphs of the current Blueprint or Level Blueprint. |
 [/EXCERPT:DebugMenu]
-
-### Window 
-
-The **Window** menu in the **Blueprint Editor** has a specific subsection for displaying Blueprint Editor-specific tabs. Not all tabs will be present in the menu while the Blueprint Editor is in the **Defaults** and **Components** modes.
-
-
-| Command | Description |
-| --- | --- |
-| **Toolbar** | Shows the Toolbar if it is currently not visible. |
-| **Details** | Shows the [Details](Engine/Blueprints\Editor/UIComponents\Details) pane if it is currently not visible. |
-| **Debug** | Shows the **Debug** pane if if it currently not visible. |
-| **Palette** | Shows the [](Engine/Blueprints\Editor/UIComponents\Palette) pane if it is not currently visible. |
-| **My Blueprint** | Shows the [](Engine/Blueprints\Editor/UIComponents\MyBlueprint) pane if it is not currently visible.  |
-| **Compiler Results** | Shows the **Compiler Results** pane if it is not currently visible.  |
-| **Find Results** | Shows the **Find Results** pane if it is not currently visible.  | 
-| **Components** | Shows the **Components panel** if it is not currently visible. |
-| **Viewport** | Shows the **Preview Viewport** pane if it is not currently visible.  |
-| **Content Browser** | Opens a sub-menu with access to all four **Content Browsers**. |
-| **Developer Tools** | [INCLUDE:#devtools] |
-| **Project Launcher** | Shows the Project Launcher which enables you to run your project on any appropriately setup and connected devices. |
-| **Plugins** | Shows the Plugin Tab where you can load/unload plugins. |
-| **Reset Layout...** | Restores the default layout for the entire Unreal Engine 4 Editor. This requires the Editor to restart, but it will re-open the current project. |
-| **Save Layout** | Saves the current interface layout. |
-| **Enable Fullscreen** | Enables fullscreen mode for the Main Editor window. |
-
-
-[/EXCERPT:Main]
-
-[COMMENT:none]
-
 [EXCERPT:devtools]
 [REGION:simpletable]
 | Sub-Menu Item | Description |
diff --git a/Engine/Documentation/Source/Engine/Blueprints/Editor/UIComponents/MyBlueprint/MyBlueprint.INT.udn b/Engine/Documentation/Source/Engine/Blueprints/Editor/UIComponents/MyBlueprint/MyBlueprint.INT.udn
--- a/Engine/Documentation/Source/Engine/Blueprints/Editor/UIComponents/MyBlueprint/MyBlueprint.INT.udn
+++ b/Engine/Documentation/Source/Engine/Blueprints/Editor/UIComponents/MyBlueprint/MyBlueprint.INT.udn
@@ -18,6 +18,23 @@ For example, a normal Blueprint will always have a **ConstructionScript** and an
 
 ## Creation Buttons
 
+[INCLUDE:#MyBlueprintPanel]
+
+## My Blueprint Sections
+
+The **My Blueprint** tab is divided up into 6 sections: The **Add New**, Graphs, Functions, Macros, Variables, and Event Dispatchers sections.
+
+![My Blueprint Pane](MyBlueprintPane.png)
+
+The 5 lower sections will organize elements of your Blueprint into their respective groups. While the top section gives you quick access to create new Graphs, Variables, etc... and search the entire **My Blueprint** panel.
+
+## Searching in My Blueprint
+
+The **My Blueprint** tab contains a text box used for searching the graphs belonging to the Blueprint. This works the same way as the contextual action menu that you use for adding new nodes, but is limited to searching items found within **My Blueprint**.  You can search based on name, comment, and other data.  So, if you know you have a node that is **SetActorHidden**, you can type that in the box, and the explorer will show you all the **SetActorHidden** nodes in your graph.
+
+[/EXCERPT:Main]
+
+[COMMENT:none]
 [EXCERPT:MyBlueprintPanel]
 The **My Blueprint** tab has shortcut buttons (the ![](plus_Button.png) at the end of each heading) for creating new variables, functions, macros, event graphs, and event dispatchers.
 
@@ -38,23 +55,4 @@ Finally, you can use the **Add New** button at the top of the **My Blueprint** P
 ![](addNew.png)
 
 [/EXCERPT:MyBlueprintPanel]
-
-
-## My Blueprint Sections
-
-The **My Blueprint** tab is divided up into 6 sections: The **Add New**, Graphs, Functions, Macros, Variables, and Event Dispatchers sections.
-
-![My Blueprint Pane](MyBlueprintPane.png)
-
-The 5 lower sections will organize elements of your Blueprint into their respective groups. While the top section gives you quick access to create new Graphs, Variables, etc... and search the entire **My Blueprint** panel.
-
-## Searching in My Blueprint
-
-The **My Blueprint** tab contains a text box used for searching the graphs belonging to the Blueprint. This works the same way as the contextual action menu that you use for adding new nodes, but is limited to searching items found within **My Blueprint**.  You can search based on name, comment, and other data.  So, if you know you have a node that is **SetActorHidden**, you can type that in the box, and the explorer will show you all the **SetActorHidden** nodes in your graph.
-
-
-
-
-
-[/EXCERPT:Main]
-
+[/COMMENT]
\ No newline at end of file
diff --git a/Engine/Documentation/Source/Engine/Blueprints/UserGuide/Debugging/KismetDebugging.INT.udn b/Engine/Documentation/Source/Engine/Blueprints/UserGuide/Debugging/KismetDebugging.INT.udn
--- a/Engine/Documentation/Source/Engine/Blueprints/UserGuide/Debugging/KismetDebugging.INT.udn
+++ b/Engine/Documentation/Source/Engine/Blueprints/UserGuide/Debugging/KismetDebugging.INT.udn
@@ -62,12 +62,9 @@ In addition, the viewport will display an overlay to alert the player that the g
 
 ## Debug Tab and Blueprint Debugger
 
-
 [EXCERPT:DebugTab]
 
-[EXCERPT:DebugPanel]
-The **Debug** tab shows information designated as important by the designer in the form of Breakpoints and watch values as well as a trace stack of all nodes belonging to the Blueprint that have been executed. This window also shows controls for controlling execution of the game when using Breakpoints.
-[/EXCERPT:DebugPanel]
+[INCLUDE:#DebugPanel]
 
 ### Watch Values
 
@@ -90,3 +87,9 @@ The **Execution Trace** stack shows a list of the nodes executed with the most r
 This list updates as you step through the graph when debugging.
 
 [/EXCERPT:DebugTab]
+
+[COMMENT:none]
+[EXCERPT:DebugPanel]
+The **Debug** tab shows information designated as important by the designer in the form of Breakpoints and watch values as well as a trace stack of all nodes belonging to the Blueprint that have been executed. This window also shows controls for controlling execution of the game when using Breakpoints.
+[/EXCERPT:DebugPanel]
+[/COMMENT]
\ No newline at end of file
diff --git a/Engine/Documentation/Source/Engine/Composure/HowTo/ComposureHowTo.INT.udn b/Engine/Documentation/Source/Engine/Composure/HowTo/ComposureHowTo.INT.udn
--- a/Engine/Documentation/Source/Engine/Composure/HowTo/ComposureHowTo.INT.udn
+++ b/Engine/Documentation/Source/Engine/Composure/HowTo/ComposureHowTo.INT.udn
@@ -14,8 +14,6 @@ ADD Overview Text
 The Overview text should describe the system, features, or concepts being covered in the How-Tos.
 [/REGION]
 
-[](FANCY:Engine\Composure\HowTo\FakeAo "%Engine\Composure\HowTo\FakeAo:description%")
-
 ADD Beginner Topic List
 [REGION:note]
 The Beginner Topic List should link to all beginner How-To guides using fancy links.
diff --git a/Engine/Documentation/Source/Engine/Landscape/Editing/Splines/Splines.INT.udn b/Engine/Documentation/Source/Engine/Landscape/Editing/Splines/Splines.INT.udn
--- a/Engine/Documentation/Source/Engine/Landscape/Editing/Splines/Splines.INT.udn
+++ b/Engine/Documentation/Source/Engine/Landscape/Editing/Splines/Splines.INT.udn
@@ -30,14 +30,18 @@ The **Spline** tool creates a spline Actor that deforms a terrain mesh along it,
 
 [EXCERPT:Splines001]
 
-[EXCERPT:LandscapeSplines]
-**Landscape Splines** are a flexible system for creating any linear feature that needs to conform to a Landscape, and can even push and pull the terrain to better facilitate building these features. They are created and edited using the spline tool in the Landscape tool.
-[/EXCERPT:LandscapeSplines]
+[INCLUDE:#LandscapeSplines]
 
 The Spline tool is best used to make streets or paths in the environment, but in general, it can be used for any mesh that must conform to the terrain, such as cobblestone walls or long straight rows of corn. 
 
 [/EXCERPT:Splines001]
 
+[COMMENT:none]
+[EXCERPT:LandscapeSplines]
+**Landscape Splines** are a flexible system for creating any linear feature that needs to conform to a Landscape, and can even push and pull the terrain to better facilitate building these features. They are created and edited using the spline tool in the Landscape tool.
+[/EXCERPT:LandscapeSplines]
+[/COMMENT]
+
 ## Creating a Spline
 
 **To create a spline:**
diff --git a/Engine/Documentation/Source/Engine/Landscape/Materials/Materials.INT.udn b/Engine/Documentation/Source/Engine/Landscape/Materials/Materials.INT.udn
--- a/Engine/Documentation/Source/Engine/Landscape/Materials/Materials.INT.udn
+++ b/Engine/Documentation/Source/Engine/Landscape/Materials/Materials.INT.udn
@@ -122,7 +122,7 @@ You can set the properties of the LandscapeLayerSwitch node in the **Details** p
 |2:| **Preview Used**   | If checked, uses a preview. |
 |3:| **LayerUsed**      | The result to use when the layer specified in the node's properties is in use by the current region of the Landscape.|
 |4:| **LayerNotUsed**   | The result to use when the layer is not used by the current region of the Landscape and has a weight of zero.|
-|5:| ** Unlabeled Output**        | Either the **LayerUsed** or **LayerNotUsed** inputs, depending on whether or not the layer contributes to the particular region of the Landscape.|  
+|5:| **Unlabeled Output**        | Either the **LayerUsed** or **LayerNotUsed** inputs, depending on whether or not the layer contributes to the particular region of the Landscape.|  
 
 [/EXCERPT:TerrainLayerSwitch]
 
@@ -295,14 +295,6 @@ You can use the same Material for your Landscape Material and your Landscape Hol
 
 	![](Landscape_BaseNodeOpacity.png)
 
-	[COMMENT:none]
-
-	[REGION:tip]
-	If you are using a separate Material for your Landscape hole you need to make sure that you also have a texture feeding into the Base Color slot or when you try and make parts of your Landscape invisible you will end up turning parts of your Landscape black.
-	[/REGION]
-	
-[/COMMENT]
-
 1. Next, search for and add a **LandscapeVisibilityMask** node to the Material Graph and connect its output to the **Opacity Mask** input.
 
 	![](Landscape_VisibilityMaskNode.png)
diff --git a/Engine/Documentation/Source/Engine/Matinee/HowTo/MHT_1/MHT_1.INT.udn b/Engine/Documentation/Source/Engine/Matinee/HowTo/MHT_1/MHT_1.INT.udn
--- a/Engine/Documentation/Source/Engine/Matinee/HowTo/MHT_1/MHT_1.INT.udn
+++ b/Engine/Documentation/Source/Engine/Matinee/HowTo/MHT_1/MHT_1.INT.udn
@@ -8,6 +8,7 @@ SkillLevel: Beginner
 
 In the following steps, we are going to create a simple door that opens/closes when entering/exiting a trigger volume. 
 
+[EXCERPT:ProjectStart]
 ## Creating the Project
 
 In this section, we will set up the project. 
diff --git a/Engine/Documentation/Source/Engine/MediaFramework/HowTo/MF_HowTo.INT.udn b/Engine/Documentation/Source/Engine/MediaFramework/HowTo/MF_HowTo.INT.udn
--- a/Engine/Documentation/Source/Engine/MediaFramework/HowTo/MF_HowTo.INT.udn
+++ b/Engine/Documentation/Source/Engine/MediaFramework/HowTo/MF_HowTo.INT.udn
@@ -15,31 +15,15 @@ If you are looking to quickly get a video asset up and running inside your level
 Refer to the table below for more information and click each link to view its contents. 
 
 ## Beginner
-[REGION:wrapper transparent]
-[](FANCY:Engine\MediaFramework\HowTo\FileMediaSource "%Engine\MediaFramework\HowTo\FileMediaSource:description%")
-[](FANCY:Engine\MediaFramework\HowTo\StreamMediaSource "%Engine\MediaFramework\HowTo\StreamMediaSource:description%")
-[](FANCY:Engine\MediaFramework\HowTo\ImgMediaSource "%Engine\MediaFramework\HowTo\ImgMediaSource:description%")
-[/REGION]
 
+[DIR(output:"fancy" type:"how-to" skill:"Beginner" parent:"Engine/MediaFramework/HowTo" org:"hierarchy")]
 
-## Intermediate
-[REGION:wrapper transparent]
-[](FANCY:Engine\MediaFramework\HowTo\MediaPlaylists "%Engine\MediaFramework\HowTo\MediaPlaylists:description%")
-[](FANCY:Engine\MediaFramework\HowTo\PlatformMedia "%Engine\MediaFramework\HowTo\PlatformMedia:description%")
-[](FANCY:Engine\MediaFramework\HowTo\ControllingPlayback "%Engine\MediaFramework\HowTo\ControllingPlayback:description%")
-[/REGION]
 
-## Advanced
-[REGION:wrapper transparent]
-[](FANCY:Engine\MediaFramework\HowTo\UsingWebCams "%Engine\MediaFramework\HowTo\UsingWebCams:description%")
-[/REGION]
+## Intermediate
 
+[DIR(output:"fancy" type:"how-to" skill:"Intermediate" parent:"Engine/MediaFramework/HowTo" org:"hierarchy")]
 
-[COMMENT:none]
- To Do
 ## Advanced
-[REGION:wrapper transparent]
-[/REGION]
 
-[/COMMENT]
+[DIR(output:"fancy" type:"how-to" skill:"Advanced" parent:"Engine/MediaFramework/HowTo" org:"hierarchy")]
 
diff --git a/Engine/Documentation/Source/Engine/Physics/Cloth/Reference/ClothReference.INT.udn b/Engine/Documentation/Source/Engine/Physics/Cloth/Reference/ClothReference.INT.udn
--- a/Engine/Documentation/Source/Engine/Physics/Cloth/Reference/ClothReference.INT.udn
+++ b/Engine/Documentation/Source/Engine/Physics/Cloth/Reference/ClothReference.INT.udn
@@ -41,7 +41,6 @@ mesh that you've selected, create a cloth asset for your render mesh and its LOD
 
 
 [COMMENT:none]
--
 
 [EXCERPT:createclothsection]
 ![](CreateClothSection.png)(w:350)
@@ -51,17 +50,6 @@ mesh that you've selected, create a cloth asset for your render mesh and its LOD
 | **Remove from Mesh** | Whether or not to leave this section behind (if driving a mesh with itself). Enable this is driving a high-poly mesh with a low-poly mesh. |
 | [REGION:tablesection] Collision [/REGION] ||
 | **Physics Asset** | Physics Asset to extract collisions from. Note that this will export Sphere and Sphylls, but does support up to 32 convex faces (or 5 boxes). |
-<!---| [REGION:tablesection] Auto Fix [/REGION] ||-
-[/COMMENT]
-[COMMENT:none]
--| **Try Auto Fix** |  |-
-[/COMMENT]
-[COMMENT:none]
--| **Auto Fix Threshold** |  |-
-[/COMMENT]
-[COMMENT:none]
--| **Simulated Particle Max Distance** |  |-
-[/COMMENT]
 [/EXCERPT:createclothsection]
 
 [EXCERPT:createclothsectionLOD]
@@ -82,8 +70,7 @@ mesh that you've selected, create a cloth asset for your render mesh and its LOD
 | **Add LOD** | Use the selected section to add as a new LOD. |
 [/EXCERPT:lodindex]
 
-
---->
+[/COMMENT]
 
 
 ## Cloth Paint Panel
@@ -116,7 +103,7 @@ The **Masks** window displays any created parameter sets for painted cloth value
 
 [REGION:raw]
 | Property | Description |
-| --- | --- | --- |
+| --- | --- |
 | **Name** | The name given to the mask along with its target setting for this parameter set. [INCLUDE:#masksettings]  |
 | **Add Mask** | Adds a new Mask to the available list of Mask parameters. |
 [/REGION]
@@ -124,7 +111,6 @@ The **Masks** window displays any created parameter sets for painted cloth value
 
 
 [COMMENT:none]
--
 [INCLUDE:#masksettings]
 
 [EXCERPT:masksettings]
@@ -141,8 +127,6 @@ The **Masks** window displays any created parameter sets for painted cloth value
 | **Delete** | Deletes the Mask from the list. |
 | **Apply** | Applies the Mask to the physical mesh. |
 [/EXCERPT:masksettings]
-
--
 [/COMMENT]
 
 
@@ -186,7 +170,6 @@ The **Config** properties enables you to adjust how your cloth react so that it
 
 
 [COMMENT:none]
--
 [INCLUDE:#constraintgconfig]
 
 [EXCERPT:constraintgconfig]
@@ -200,11 +183,8 @@ The **Config** properties enables you to adjust how your cloth react so that it
 [EXCERPT:windmethod]
 | --- | --- |
 | **Legacy** | Use the Legacy wind mode, where accelerations are modified directly by the simulation with no regard to drag or lift. |
-<!--- | **Accurate** |  | -
-[/COMMENT]
 [/EXCERPT:windmethod]
-
---->
+[/COMMENT]
 
 
 
diff --git a/Engine/Documentation/Source/Engine/Physics/Constraints/ConstraintsReference/Constraints.INT.udn b/Engine/Documentation/Source/Engine/Physics/Constraints/ConstraintsReference/Constraints.INT.udn
--- a/Engine/Documentation/Source/Engine/Physics/Constraints/ConstraintsReference/Constraints.INT.udn
+++ b/Engine/Documentation/Source/Engine/Physics/Constraints/ConstraintsReference/Constraints.INT.udn
@@ -229,7 +229,7 @@ Whether this angular motor uses SLERP (Spherical Linear Interpolation) or decomp
 
 | Property | Description |
 | -------- | ----------- |
-|[REGION:tablesection]Constraint[/REGION]|
+|[REGION:tablesection]Constraint[/REGION]||
 | **Component Name 1** |  |
 | **Component Name 2** |   |
 | **Joint Name** |  |
@@ -259,7 +259,7 @@ Whether this angular motor uses SLERP (Spherical Linear Interpolation) or decomp
 
 | Property | Description |
 | -------- | ----------- |
-|[REGION:tablesection]Linear Limits[/REGION]|
+|[REGION:tablesection]Linear Limits[/REGION]||
 | **Linear XMotion** |  |
 | **Linear YMotion** |  |
 | **Linear ZMotion** |  |
@@ -270,11 +270,11 @@ Whether this angular motor uses SLERP (Spherical Linear Interpolation) or decomp
 | **Linear Limit Damping** |  |
 | **Linear Breakable** |  |
 | **Linear Break Threshold** |  |
-|[REGION:tablesection]Projection[/REGION]|
+|[REGION:tablesection]Projection[/REGION]||
 | **Enable Projection** |  |
 | **Projection Linear Tolerance** |  |
 | **Projection Angular Tolerance** |  |
-|[REGION:tablesection]Angular Limits[/REGION]|
+|[REGION:tablesection]Angular Limits[/REGION]||
 | **Angular Swing 1Motion** |  |
 | **Angular Swing 2Motion** |  |
 | **Angular Twist Motion** |  |
@@ -290,7 +290,7 @@ Whether this angular motor uses SLERP (Spherical Linear Interpolation) or decomp
 | **Twist Limit Damping** |  |
 | **Angular Breakable** |  |
 | **Angular Break Threshold** |  |
-|[REGION:tablesection]Linear Motor[/REGION]|
+|[REGION:tablesection]Linear Motor[/REGION]||
 | **Linear Position Drive** |  |
 | **Linear XPosition Drive** |  |
 | **Linear YPosition Drive** |  |
@@ -301,17 +301,14 @@ Whether this angular motor uses SLERP (Spherical Linear Interpolation) or decomp
 | **Linear Velocity Target** |  |
 | **Linear Velocity Strength** | Damping to apply to the linear drive. |
 | **Max Linear Force** | Limit to the force the linear drive can apply. |
-|[REGION:tablesection]Angular Motor[/REGION]|
+|[REGION:tablesection]Angular Motor[/REGION]||
 | **Angular Orientation Drive** |  |
 | **Angular Orientation Target** |  |
 | **Angular Position Strength** | Spring value to apply for the angular drive. |
 | **Angular Velocity Drive** |  |
 | **Angular Velocity Target** |  |
 | **Angular Velocity Strength** | Damping value to apply for the angular drive. |
-| **Max Angular Force ** | Limit to the force the angular drive can apply. |
-
-
-<!--
+| **Max Angular Force** | Limit to the force the angular drive can apply. |
 
 [EXCERPT:linearMotionEnum]
 | **Type** | **Description** |
@@ -330,8 +327,6 @@ Whether this angular motor uses SLERP (Spherical Linear Interpolation) or decomp
 [/EXCERPT:angularMotionEnum]
 
 
-[/COMMENT]
-
 ## Further Information
 
 ### Joint Setup
@@ -360,10 +355,6 @@ Under most circumstances, the Linear Motions' freedoms will remain locked. Howev
 
 _From left to right (Physics Constraint scaled up 4x): Linear Motions limited to 15 Units, and Linear Motions set to free._
 
-[COMMENT:none]
- For information on all the properties relating to Linear Freedom and Limits see: [%Engine/Physics\Constraints\ConstraintsReferencetitle%](Engine/Physics\Constraints\ConstraintsReference 
-[/COMMENT]
-
 ### Angular Freedom and Limits
 
 ![](PhysCon_AngularLimitsProps.png)
@@ -384,10 +375,6 @@ Further, since the joint is located where the Physics Constraint Actor is locate
 [![](PhysCon_AngPos1.png)(h:312)](PhysCon_AngPos1.png) [![](PhysCon_AngPos2.png)(h:312)](PhysCon_AngPos2.png) 
 [/REGION]
 
-[COMMENT:none]
- For information on all the properties relating to Angular Freedom and Limits see: [%Engine/Physics\Constraints\ConstraintsReferencetitle%](Engine/Physics\Constraints\ConstraintsReference 
-[/COMMENT]
-
 ### Orientation
 
 The axes of a Physics Constraint Actor are in line with the rotation of the Physics Constraint Actor, and not aligned between the two objects. Given you desired effect, you will need to be aware of how the constraint is oriented.
@@ -404,9 +391,7 @@ In either setup, the Physics Constraint has its Linear XMotion set to free, and
 [![](PhysCon_SimFlat.png)(w:453)](PhysCon_SimFlat.png) [![](PhysCon_SimAngled.png)(w:453)](PhysCon_SimAngled.png)
 [/REGION]
 
-When using Angular limits, it is assumed that the Twist Axis (positive X) will be aligned pointing towards the blue constrained Actor (Constrained Actor 2). [COMMENT:none]
--This is really important because the -
-[/COMMENT] 
+When using Angular limits, it is assumed that the Twist Axis (positive X) will be aligned pointing towards the blue constrained Actor (Constrained Actor 2). 
 
 [REGION:lightbox]
 [![](PhysCon_Angular.png)(w:453)](PhysCon_Angular.png) [![](PhysCon_AngularSim.png)(w:453)](PhysCon_AngularSim.png)
@@ -416,32 +401,7 @@ Also, while simulating, you will notice that the cone that represents the angula
 
 ![](PhysCon_AngularCone_2.png)(h:300 convert:false) ![](PhysCon_Arc.png)(h:300 convert:false)
 
-[COMMENT:none]
-
-
-### Ball And Socket
-A ball-and-socket joint is simply a joint that has all its linear parts locked and all its angular parts free.
-
-
-### Hinge
-A hinge is like a ball-and-socket in that all its linear parts are locked, but its swing parts are also locked, so it can only rotate around its twist axis.
-
-
-### Prismatic
-A prismatic joint has 2 out of its 3 linear degrees of freedom locked, and also has all its angular degrees of freedom locked. This allows is to slide down a 'rail' but cannot rotate at all.
-
-
-### Joint Friction
-You can simulate joint friction using the drive settings. Turn on velocity drive with a target velocity of zero, and use the LinearDriveForceLimit and AngularDriveForceLimit settings to control how strongly the joint will resist motion.
-
-
-### Punch bag
-You could simulate something like a punch-bag (a ball on the top of a vertical pole that can be knocked around) with a joint by first creating a ball-and-socket joint. Then turn on swing position drive and adjust the **AngularDriveSpring** and **AngularDriveDamping** until it springs nicely around the upright position when pushed. You could turn on swing limits as well to ensure that it cannot bend over too far.
-
-
 [/COMMENT]
--->
-
 
 
 
diff --git a/Engine/Documentation/Source/Engine/Physics/PhAT/HowTo/ApplyConstraintProfile/ApplyConstraintProfile.INT.udn b/Engine/Documentation/Source/Engine/Physics/PhAT/HowTo/ApplyConstraintProfile/ApplyConstraintProfile.INT.udn
--- a/Engine/Documentation/Source/Engine/Physics/PhAT/HowTo/ApplyConstraintProfile/ApplyConstraintProfile.INT.udn
+++ b/Engine/Documentation/Source/Engine/Physics/PhAT/HowTo/ApplyConstraintProfile/ApplyConstraintProfile.INT.udn
@@ -87,7 +87,60 @@ Below, you will find the steps for creating a simple graph to enable a **Constra
 			![](fullGraph.png)
 		[/PARAM]
 		[PARAM:snippet]
-			%BPSnip_001%
+				Begin Object Class=K2Node_Event Name="K2Node_Event_8"
+				   EventReference=(MemberParent=Class'/Script/Engine.Actor',MemberName="ReceiveBeginPlay")
+				   bOverrideFunction=True
+				   NodePosX=32
+				   NodePosY=-16
+				   NodeGuid=48E7E2C9458A1028A689B3A249855163
+				   CustomProperties Pin (PinId=D1792C7A4CBF979351A63AB775A61034,PinName="OutputDelegate",Direction="EGPD_Output",PinType.PinCategory="delegate",PinType.PinSubCategoryObject=None,PinType.PinSubCategoryMemberReference=(MemberParent=Class'/Script/Engine.Actor',MemberName="ReceiveBeginPlay"),PinType.bIsArray=False,PinType.bIsReference=False,PinType.bIsConst=False,PinType.bIsWeakPointer=False,PersistentGuid=00000000000000000000000000000000,bHidden=False,bNotConnectable=False,bDefaultValueIsReadOnly=False,bDefaultValueIsIgnored=False,bAdvancedView=False,)
+				   CustomProperties Pin (PinId=3880C4144B19D9B0F602638392FAD6FA,PinName="then",Direction="EGPD_Output",PinType.PinCategory="exec",PinType.PinSubCategoryObject=None,PinType.PinSubCategoryMemberReference=(),PinType.bIsArray=False,PinType.bIsReference=False,PinType.bIsConst=False,PinType.bIsWeakPointer=False,LinkedTo=(K2Node_CallFunction_132 7ED8710C4A46FF83B175828EFBF77DAE,),PersistentGuid=00000000000000000000000000000000,bHidden=False,bNotConnectable=False,bDefaultValueIsReadOnly=False,bDefaultValueIsIgnored=False,bAdvancedView=False,)
+				End Object
+				Begin Object Class=K2Node_VariableGet Name="K2Node_VariableGet_155"
+				   VariableReference=(MemberName="SkeletalMeshComponent",bSelfContext=True)
+				   NodePosX=32
+				   NodePosY=112
+				   NodeGuid=42ADF0644E0964ECBC298DB6863ADA89
+				   CustomProperties Pin (PinId=788C7F5B4B4B5A7362C8AEA404935EEB,PinName="SkeletalMeshComponent",Direction="EGPD_Output",PinType.PinCategory="object",PinType.PinSubCategoryObject=Class'/Script/Engine.SkeletalMeshComponent',PinType.PinSubCategoryMemberReference=(),PinType.bIsArray=False,PinType.bIsReference=False,PinType.bIsConst=False,PinType.bIsWeakPointer=False,LinkedTo=(K2Node_CallFunction_132 E1392084422728C657D6C3A54A4116F2,),PersistentGuid=00000000000000000000000000000000,bHidden=False,bNotConnectable=False,bDefaultValueIsReadOnly=False,bDefaultValueIsIgnored=False,bAdvancedView=False,)
+				   CustomProperties Pin (PinId=D7B4F9F54069CB07496DEE843FDBBA6B,PinName="self",PinFriendlyName=NSLOCTEXT("K2Node", "Target", "Target"),PinType.PinCategory="object",PinType.PinSubCategoryObject=BlueprintGeneratedClass'/Game/PhysAnimPawn.PhysAnimPawn_C',PinType.PinSubCategoryMemberReference=(),PinType.bIsArray=False,PinType.bIsReference=False,PinType.bIsConst=False,PinType.bIsWeakPointer=False,PersistentGuid=00000000000000000000000000000000,bHidden=True,bNotConnectable=False,bDefaultValueIsReadOnly=False,bDefaultValueIsIgnored=False,bAdvancedView=False,)
+				End Object
+				Begin Object Class=K2Node_CallFunction Name="K2Node_CallFunction_132"
+				   FunctionReference=(MemberParent=Class'/Script/Engine.SkeletalMeshComponent',MemberName="SetConstraintProfileForAll")
+				   NodePosX=320
+				   NodeGuid=19366CE04ECF7D2D1B3F3FAB826015E3
+				   CustomProperties Pin (PinId=7ED8710C4A46FF83B175828EFBF77DAE,PinName="execute",PinToolTip="\nExec",PinType.PinCategory="exec",PinType.PinSubCategoryObject=None,PinType.PinSubCategoryMemberReference=(),PinType.bIsArray=False,PinType.bIsReference=False,PinType.bIsConst=False,PinType.bIsWeakPointer=False,LinkedTo=(K2Node_Event_8 3880C4144B19D9B0F602638392FAD6FA,),PersistentGuid=00000000000000000000000000000000,bHidden=False,bNotConnectable=False,bDefaultValueIsReadOnly=False,bDefaultValueIsIgnored=False,bAdvancedView=False,)
+				   CustomProperties Pin (PinId=94ECC7884A2DC860DBB45E8573B644E7,PinName="then",PinToolTip="\nExec",Direction="EGPD_Output",PinType.PinCategory="exec",PinType.PinSubCategoryObject=None,PinType.PinSubCategoryMemberReference=(),PinType.bIsArray=False,PinType.bIsReference=False,PinType.bIsConst=False,PinType.bIsWeakPointer=False,LinkedTo=(K2Node_CallFunction_8 399238E643A83FB3CFC85CB926B5CF9F,),PersistentGuid=00000000000000000000000000000000,bHidden=False,bNotConnectable=False,bDefaultValueIsReadOnly=False,bDefaultValueIsIgnored=False,bAdvancedView=False,)
+				   CustomProperties Pin (PinId=E1392084422728C657D6C3A54A4116F2,PinName="self",PinFriendlyName=NSLOCTEXT("K2Node", "Target", "Target"),PinToolTip="Target\nSkeletal Mesh Component Reference",PinType.PinCategory="object",PinType.PinSubCategoryObject=Class'/Script/Engine.SkeletalMeshComponent',PinType.PinSubCategoryMemberReference=(),PinType.bIsArray=False,PinType.bIsReference=False,PinType.bIsConst=False,PinType.bIsWeakPointer=False,LinkedTo=(K2Node_VariableGet_155 788C7F5B4B4B5A7362C8AEA404935EEB,),PersistentGuid=00000000000000000000000000000000,bHidden=False,bNotConnectable=False,bDefaultValueIsReadOnly=False,bDefaultValueIsIgnored=False,bAdvancedView=False,)
+				   CustomProperties Pin (PinId=28663C674854E90FE78C718CD73FB3DF,PinName="ProfileName",PinToolTip="Profile Name\nName",PinType.PinCategory="name",PinType.PinSubCategoryObject=None,PinType.PinSubCategoryMemberReference=(),PinType.bIsArray=False,PinType.bIsReference=False,PinType.bIsConst=False,PinType.bIsWeakPointer=False,DefaultValue="DocsConstraintProfile",AutogeneratedDefaultValue="None",PersistentGuid=00000000000000000000000000000000,bHidden=False,bNotConnectable=False,bDefaultValueIsReadOnly=False,bDefaultValueIsIgnored=False,bAdvancedView=False,)
+				   CustomProperties Pin (PinId=D960208D49BE62D6FDE20698BB7D27CC,PinName="bDefaultIfNotFound",PinToolTip="Default if Not Found\nBoolean",PinType.PinCategory="bool",PinType.PinSubCategoryObject=None,PinType.PinSubCategoryMemberReference=(),PinType.bIsArray=False,PinType.bIsReference=False,PinType.bIsConst=False,PinType.bIsWeakPointer=False,DefaultValue="true",AutogeneratedDefaultValue="false",PersistentGuid=00000000000000000000000000000000,bHidden=False,bNotConnectable=False,bDefaultValueIsReadOnly=False,bDefaultValueIsIgnored=False,bAdvancedView=False,)
+				End Object
+				Begin Object Class=K2Node_VariableGet Name="K2Node_VariableGet_1"
+				   VariableReference=(MemberName="SkeletalMeshComponent",bSelfContext=True)
+				   NodePosX=624
+				   NodePosY=80
+				   NodeGuid=211663904BCAEE00C6164CBF0871A8B9
+				   CustomProperties Pin (PinId=1EEABDD6494CBF38903C6691F5F5B45C,PinName="SkeletalMeshComponent",Direction="EGPD_Output",PinType.PinCategory="object",PinType.PinSubCategoryObject=Class'/Script/Engine.SkeletalMeshComponent',PinType.PinSubCategoryMemberReference=(),PinType.bIsArray=False,PinType.bIsReference=False,PinType.bIsConst=False,PinType.bIsWeakPointer=False,LinkedTo=(K2Node_CallFunction_8 4DD6FE2846E695E0F3F569B06CED0A46,),PersistentGuid=00000000000000000000000000000000,bHidden=False,bNotConnectable=False,bDefaultValueIsReadOnly=False,bDefaultValueIsIgnored=False,bAdvancedView=False,)
+				   CustomProperties Pin (PinId=D68B011B4E4D641D541E8CA766972993,PinName="self",PinFriendlyName=NSLOCTEXT("K2Node", "Target", "Target"),PinType.PinCategory="object",PinType.PinSubCategoryObject=BlueprintGeneratedClass'/Game/PhysAnimPawn.PhysAnimPawn_C',PinType.PinSubCategoryMemberReference=(),PinType.bIsArray=False,PinType.bIsReference=False,PinType.bIsConst=False,PinType.bIsWeakPointer=False,PersistentGuid=00000000000000000000000000000000,bHidden=True,bNotConnectable=False,bDefaultValueIsReadOnly=False,bDefaultValueIsIgnored=False,bAdvancedView=False,)
+				End Object
+				Begin Object Class=K2Node_CallFunction Name="K2Node_CallFunction_8"
+				   FunctionReference=(MemberParent=Class'/Script/Engine.SkeletalMeshComponent',MemberName="SetAllBodiesBelowSimulatePhysics")
+				   NodePosX=848
+				   NodeGuid=3116B9174BA9B4F48CE4CAA1E47000EC
+				   CustomProperties Pin (PinId=399238E643A83FB3CFC85CB926B5CF9F,PinName="execute",PinToolTip="\nExec",PinType.PinCategory="exec",PinType.PinSubCategoryObject=None,PinType.PinSubCategoryMemberReference=(),PinType.bIsArray=False,PinType.bIsReference=False,PinType.bIsConst=False,PinType.bIsWeakPointer=False,LinkedTo=(K2Node_CallFunction_132 94ECC7884A2DC860DBB45E8573B644E7,),PersistentGuid=00000000000000000000000000000000,bHidden=False,bNotConnectable=False,bDefaultValueIsReadOnly=False,bDefaultValueIsIgnored=False,bAdvancedView=False,)
+				   CustomProperties Pin (PinId=C132D8014BC6A4DE84C437AF07CF838E,PinName="then",PinToolTip="\nExec",Direction="EGPD_Output",PinType.PinCategory="exec",PinType.PinSubCategoryObject=None,PinType.PinSubCategoryMemberReference=(),PinType.bIsArray=False,PinType.bIsReference=False,PinType.bIsConst=False,PinType.bIsWeakPointer=False,PersistentGuid=00000000000000000000000000000000,bHidden=False,bNotConnectable=False,bDefaultValueIsReadOnly=False,bDefaultValueIsIgnored=False,bAdvancedView=False,)
+				   CustomProperties Pin (PinId=4DD6FE2846E695E0F3F569B06CED0A46,PinName="self",PinFriendlyName=NSLOCTEXT("K2Node", "Target", "Target"),PinToolTip="Target\nSkeletal Mesh Component Reference",PinType.PinCategory="object",PinType.PinSubCategoryObject=Class'/Script/Engine.SkeletalMeshComponent',PinType.PinSubCategoryMemberReference=(),PinType.bIsArray=False,PinType.bIsReference=False,PinType.bIsConst=False,PinType.bIsWeakPointer=False,LinkedTo=(K2Node_VariableGet_1 1EEABDD6494CBF38903C6691F5F5B45C,),PersistentGuid=00000000000000000000000000000000,bHidden=False,bNotConnectable=False,bDefaultValueIsReadOnly=False,bDefaultValueIsIgnored=False,bAdvancedView=False,)
+				   CustomProperties Pin (PinId=7D29E35244DEB831C54B1CAEF4A5FE5C,PinName="InBoneName",PinToolTip="In Bone Name\nName (by ref)",PinType.PinCategory="name",PinType.PinSubCategoryObject=None,PinType.PinSubCategoryMemberReference=(),PinType.bIsArray=False,PinType.bIsReference=True,PinType.bIsConst=True,PinType.bIsWeakPointer=False,DefaultValue="None",AutogeneratedDefaultValue="None",LinkedTo=(K2Node_VariableGet_181 4DF93E5141ACE87768E5D8A7543C7B07,),PersistentGuid=00000000000000000000000000000000,bHidden=False,bNotConnectable=False,bDefaultValueIsReadOnly=False,bDefaultValueIsIgnored=True,bAdvancedView=False,)
+				   CustomProperties Pin (PinId=AE4340174DD534FEA5B812AE6DAEA1EB,PinName="bNewSimulate",PinToolTip="New Simulate\nBoolean",PinType.PinCategory="bool",PinType.PinSubCategoryObject=None,PinType.PinSubCategoryMemberReference=(),PinType.bIsArray=False,PinType.bIsReference=False,PinType.bIsConst=False,PinType.bIsWeakPointer=False,DefaultValue="true",AutogeneratedDefaultValue="false",PersistentGuid=00000000000000000000000000000000,bHidden=False,bNotConnectable=False,bDefaultValueIsReadOnly=False,bDefaultValueIsIgnored=False,bAdvancedView=False,)
+				   CustomProperties Pin (PinId=80EC7602465D797C7F0D4D93D6BA9BD5,PinName="bIncludeSelf",PinToolTip="Include Self\nBoolean",PinType.PinCategory="bool",PinType.PinSubCategoryObject=None,PinType.PinSubCategoryMemberReference=(),PinType.bIsArray=False,PinType.bIsReference=False,PinType.bIsConst=False,PinType.bIsWeakPointer=False,DefaultValue="true",AutogeneratedDefaultValue="true",PersistentGuid=00000000000000000000000000000000,bHidden=False,bNotConnectable=False,bDefaultValueIsReadOnly=False,bDefaultValueIsIgnored=False,bAdvancedView=False,)
+				End Object
+				Begin Object Class=K2Node_VariableGet Name="K2Node_VariableGet_181"
+				   VariableReference=(MemberName="BoneName",MemberGuid=A5E270C141CEE32A0269309215CEBAC4,bSelfContext=True)
+				   NodePosX=640
+				   NodePosY=128
+				   NodeGuid=B116D28D4BF4E6625FF02F92D2EA8DFD
+				   CustomProperties Pin (PinId=4DF93E5141ACE87768E5D8A7543C7B07,PinName="BoneName",Direction="EGPD_Output",PinType.PinCategory="name",PinType.PinSubCategoryObject=None,PinType.PinSubCategoryMemberReference=(),PinType.bIsArray=False,PinType.bIsReference=False,PinType.bIsConst=False,PinType.bIsWeakPointer=False,AutogeneratedDefaultValue="None",LinkedTo=(K2Node_CallFunction_8 7D29E35244DEB831C54B1CAEF4A5FE5C,),PersistentGuid=00000000000000000000000000000000,bHidden=False,bNotConnectable=False,bDefaultValueIsReadOnly=False,bDefaultValueIsIgnored=False,bAdvancedView=False,)
+				   CustomProperties Pin (PinId=F259489241A7853EE35268954C561767,PinName="self",PinFriendlyName=NSLOCTEXT("K2Node", "Target", "Target"),PinType.PinCategory="object",PinType.PinSubCategoryObject=BlueprintGeneratedClass'/Game/PhysAnimPawn.PhysAnimPawn_C',PinType.PinSubCategoryMemberReference=(),PinType.bIsArray=False,PinType.bIsReference=False,PinType.bIsConst=False,PinType.bIsWeakPointer=False,PersistentGuid=00000000000000000000000000000000,bHidden=True,bNotConnectable=False,bDefaultValueIsReadOnly=False,bDefaultValueIsIgnored=False,bAdvancedView=False,)
+				End Object
 		[/PARAM]
 	[/OBJECT]
 
@@ -113,67 +166,6 @@ This effect was achieved by setting the **Linear Limits** on the upper arm const
 
 ![](achievedBy.png)
 
-[COMMENT:none]
- Variables 
-[/COMMENT]
-
-[VAR:BPSnip_001]
-	Begin Object Class=K2Node_Event Name="K2Node_Event_8"
-	   EventReference=(MemberParent=Class'/Script/Engine.Actor',MemberName="ReceiveBeginPlay")
-	   bOverrideFunction=True
-	   NodePosX=32
-	   NodePosY=-16
-	   NodeGuid=48E7E2C9458A1028A689B3A249855163
-	   CustomProperties Pin (PinId=D1792C7A4CBF979351A63AB775A61034,PinName="OutputDelegate",Direction="EGPD_Output",PinType.PinCategory="delegate",PinType.PinSubCategoryObject=None,PinType.PinSubCategoryMemberReference=(MemberParent=Class'/Script/Engine.Actor',MemberName="ReceiveBeginPlay"),PinType.bIsArray=False,PinType.bIsReference=False,PinType.bIsConst=False,PinType.bIsWeakPointer=False,PersistentGuid=00000000000000000000000000000000,bHidden=False,bNotConnectable=False,bDefaultValueIsReadOnly=False,bDefaultValueIsIgnored=False,bAdvancedView=False,)
-	   CustomProperties Pin (PinId=3880C4144B19D9B0F602638392FAD6FA,PinName="then",Direction="EGPD_Output",PinType.PinCategory="exec",PinType.PinSubCategoryObject=None,PinType.PinSubCategoryMemberReference=(),PinType.bIsArray=False,PinType.bIsReference=False,PinType.bIsConst=False,PinType.bIsWeakPointer=False,LinkedTo=(K2Node_CallFunction_132 7ED8710C4A46FF83B175828EFBF77DAE,),PersistentGuid=00000000000000000000000000000000,bHidden=False,bNotConnectable=False,bDefaultValueIsReadOnly=False,bDefaultValueIsIgnored=False,bAdvancedView=False,)
-	End Object
-	Begin Object Class=K2Node_VariableGet Name="K2Node_VariableGet_155"
-	   VariableReference=(MemberName="SkeletalMeshComponent",bSelfContext=True)
-	   NodePosX=32
-	   NodePosY=112
-	   NodeGuid=42ADF0644E0964ECBC298DB6863ADA89
-	   CustomProperties Pin (PinId=788C7F5B4B4B5A7362C8AEA404935EEB,PinName="SkeletalMeshComponent",Direction="EGPD_Output",PinType.PinCategory="object",PinType.PinSubCategoryObject=Class'/Script/Engine.SkeletalMeshComponent',PinType.PinSubCategoryMemberReference=(),PinType.bIsArray=False,PinType.bIsReference=False,PinType.bIsConst=False,PinType.bIsWeakPointer=False,LinkedTo=(K2Node_CallFunction_132 E1392084422728C657D6C3A54A4116F2,),PersistentGuid=00000000000000000000000000000000,bHidden=False,bNotConnectable=False,bDefaultValueIsReadOnly=False,bDefaultValueIsIgnored=False,bAdvancedView=False,)
-	   CustomProperties Pin (PinId=D7B4F9F54069CB07496DEE843FDBBA6B,PinName="self",PinFriendlyName=NSLOCTEXT("K2Node", "Target", "Target"),PinType.PinCategory="object",PinType.PinSubCategoryObject=BlueprintGeneratedClass'/Game/PhysAnimPawn.PhysAnimPawn_C',PinType.PinSubCategoryMemberReference=(),PinType.bIsArray=False,PinType.bIsReference=False,PinType.bIsConst=False,PinType.bIsWeakPointer=False,PersistentGuid=00000000000000000000000000000000,bHidden=True,bNotConnectable=False,bDefaultValueIsReadOnly=False,bDefaultValueIsIgnored=False,bAdvancedView=False,)
-	End Object
-	Begin Object Class=K2Node_CallFunction Name="K2Node_CallFunction_132"
-	   FunctionReference=(MemberParent=Class'/Script/Engine.SkeletalMeshComponent',MemberName="SetConstraintProfileForAll")
-	   NodePosX=320
-	   NodeGuid=19366CE04ECF7D2D1B3F3FAB826015E3
-	   CustomProperties Pin (PinId=7ED8710C4A46FF83B175828EFBF77DAE,PinName="execute",PinToolTip="\nExec",PinType.PinCategory="exec",PinType.PinSubCategoryObject=None,PinType.PinSubCategoryMemberReference=(),PinType.bIsArray=False,PinType.bIsReference=False,PinType.bIsConst=False,PinType.bIsWeakPointer=False,LinkedTo=(K2Node_Event_8 3880C4144B19D9B0F602638392FAD6FA,),PersistentGuid=00000000000000000000000000000000,bHidden=False,bNotConnectable=False,bDefaultValueIsReadOnly=False,bDefaultValueIsIgnored=False,bAdvancedView=False,)
-	   CustomProperties Pin (PinId=94ECC7884A2DC860DBB45E8573B644E7,PinName="then",PinToolTip="\nExec",Direction="EGPD_Output",PinType.PinCategory="exec",PinType.PinSubCategoryObject=None,PinType.PinSubCategoryMemberReference=(),PinType.bIsArray=False,PinType.bIsReference=False,PinType.bIsConst=False,PinType.bIsWeakPointer=False,LinkedTo=(K2Node_CallFunction_8 399238E643A83FB3CFC85CB926B5CF9F,),PersistentGuid=00000000000000000000000000000000,bHidden=False,bNotConnectable=False,bDefaultValueIsReadOnly=False,bDefaultValueIsIgnored=False,bAdvancedView=False,)
-	   CustomProperties Pin (PinId=E1392084422728C657D6C3A54A4116F2,PinName="self",PinFriendlyName=NSLOCTEXT("K2Node", "Target", "Target"),PinToolTip="Target\nSkeletal Mesh Component Reference",PinType.PinCategory="object",PinType.PinSubCategoryObject=Class'/Script/Engine.SkeletalMeshComponent',PinType.PinSubCategoryMemberReference=(),PinType.bIsArray=False,PinType.bIsReference=False,PinType.bIsConst=False,PinType.bIsWeakPointer=False,LinkedTo=(K2Node_VariableGet_155 788C7F5B4B4B5A7362C8AEA404935EEB,),PersistentGuid=00000000000000000000000000000000,bHidden=False,bNotConnectable=False,bDefaultValueIsReadOnly=False,bDefaultValueIsIgnored=False,bAdvancedView=False,)
-	   CustomProperties Pin (PinId=28663C674854E90FE78C718CD73FB3DF,PinName="ProfileName",PinToolTip="Profile Name\nName",PinType.PinCategory="name",PinType.PinSubCategoryObject=None,PinType.PinSubCategoryMemberReference=(),PinType.bIsArray=False,PinType.bIsReference=False,PinType.bIsConst=False,PinType.bIsWeakPointer=False,DefaultValue="DocsConstraintProfile",AutogeneratedDefaultValue="None",PersistentGuid=00000000000000000000000000000000,bHidden=False,bNotConnectable=False,bDefaultValueIsReadOnly=False,bDefaultValueIsIgnored=False,bAdvancedView=False,)
-	   CustomProperties Pin (PinId=D960208D49BE62D6FDE20698BB7D27CC,PinName="bDefaultIfNotFound",PinToolTip="Default if Not Found\nBoolean",PinType.PinCategory="bool",PinType.PinSubCategoryObject=None,PinType.PinSubCategoryMemberReference=(),PinType.bIsArray=False,PinType.bIsReference=False,PinType.bIsConst=False,PinType.bIsWeakPointer=False,DefaultValue="true",AutogeneratedDefaultValue="false",PersistentGuid=00000000000000000000000000000000,bHidden=False,bNotConnectable=False,bDefaultValueIsReadOnly=False,bDefaultValueIsIgnored=False,bAdvancedView=False,)
-	End Object
-	Begin Object Class=K2Node_VariableGet Name="K2Node_VariableGet_1"
-	   VariableReference=(MemberName="SkeletalMeshComponent",bSelfContext=True)
-	   NodePosX=624
-	   NodePosY=80
-	   NodeGuid=211663904BCAEE00C6164CBF0871A8B9
-	   CustomProperties Pin (PinId=1EEABDD6494CBF38903C6691F5F5B45C,PinName="SkeletalMeshComponent",Direction="EGPD_Output",PinType.PinCategory="object",PinType.PinSubCategoryObject=Class'/Script/Engine.SkeletalMeshComponent',PinType.PinSubCategoryMemberReference=(),PinType.bIsArray=False,PinType.bIsReference=False,PinType.bIsConst=False,PinType.bIsWeakPointer=False,LinkedTo=(K2Node_CallFunction_8 4DD6FE2846E695E0F3F569B06CED0A46,),PersistentGuid=00000000000000000000000000000000,bHidden=False,bNotConnectable=False,bDefaultValueIsReadOnly=False,bDefaultValueIsIgnored=False,bAdvancedView=False,)
-	   CustomProperties Pin (PinId=D68B011B4E4D641D541E8CA766972993,PinName="self",PinFriendlyName=NSLOCTEXT("K2Node", "Target", "Target"),PinType.PinCategory="object",PinType.PinSubCategoryObject=BlueprintGeneratedClass'/Game/PhysAnimPawn.PhysAnimPawn_C',PinType.PinSubCategoryMemberReference=(),PinType.bIsArray=False,PinType.bIsReference=False,PinType.bIsConst=False,PinType.bIsWeakPointer=False,PersistentGuid=00000000000000000000000000000000,bHidden=True,bNotConnectable=False,bDefaultValueIsReadOnly=False,bDefaultValueIsIgnored=False,bAdvancedView=False,)
-	End Object
-	Begin Object Class=K2Node_CallFunction Name="K2Node_CallFunction_8"
-	   FunctionReference=(MemberParent=Class'/Script/Engine.SkeletalMeshComponent',MemberName="SetAllBodiesBelowSimulatePhysics")
-	   NodePosX=848
-	   NodeGuid=3116B9174BA9B4F48CE4CAA1E47000EC
-	   CustomProperties Pin (PinId=399238E643A83FB3CFC85CB926B5CF9F,PinName="execute",PinToolTip="\nExec",PinType.PinCategory="exec",PinType.PinSubCategoryObject=None,PinType.PinSubCategoryMemberReference=(),PinType.bIsArray=False,PinType.bIsReference=False,PinType.bIsConst=False,PinType.bIsWeakPointer=False,LinkedTo=(K2Node_CallFunction_132 94ECC7884A2DC860DBB45E8573B644E7,),PersistentGuid=00000000000000000000000000000000,bHidden=False,bNotConnectable=False,bDefaultValueIsReadOnly=False,bDefaultValueIsIgnored=False,bAdvancedView=False,)
-	   CustomProperties Pin (PinId=C132D8014BC6A4DE84C437AF07CF838E,PinName="then",PinToolTip="\nExec",Direction="EGPD_Output",PinType.PinCategory="exec",PinType.PinSubCategoryObject=None,PinType.PinSubCategoryMemberReference=(),PinType.bIsArray=False,PinType.bIsReference=False,PinType.bIsConst=False,PinType.bIsWeakPointer=False,PersistentGuid=00000000000000000000000000000000,bHidden=False,bNotConnectable=False,bDefaultValueIsReadOnly=False,bDefaultValueIsIgnored=False,bAdvancedView=False,)
-	   CustomProperties Pin (PinId=4DD6FE2846E695E0F3F569B06CED0A46,PinName="self",PinFriendlyName=NSLOCTEXT("K2Node", "Target", "Target"),PinToolTip="Target\nSkeletal Mesh Component Reference",PinType.PinCategory="object",PinType.PinSubCategoryObject=Class'/Script/Engine.SkeletalMeshComponent',PinType.PinSubCategoryMemberReference=(),PinType.bIsArray=False,PinType.bIsReference=False,PinType.bIsConst=False,PinType.bIsWeakPointer=False,LinkedTo=(K2Node_VariableGet_1 1EEABDD6494CBF38903C6691F5F5B45C,),PersistentGuid=00000000000000000000000000000000,bHidden=False,bNotConnectable=False,bDefaultValueIsReadOnly=False,bDefaultValueIsIgnored=False,bAdvancedView=False,)
-	   CustomProperties Pin (PinId=7D29E35244DEB831C54B1CAEF4A5FE5C,PinName="InBoneName",PinToolTip="In Bone Name\nName (by ref)",PinType.PinCategory="name",PinType.PinSubCategoryObject=None,PinType.PinSubCategoryMemberReference=(),PinType.bIsArray=False,PinType.bIsReference=True,PinType.bIsConst=True,PinType.bIsWeakPointer=False,DefaultValue="None",AutogeneratedDefaultValue="None",LinkedTo=(K2Node_VariableGet_181 4DF93E5141ACE87768E5D8A7543C7B07,),PersistentGuid=00000000000000000000000000000000,bHidden=False,bNotConnectable=False,bDefaultValueIsReadOnly=False,bDefaultValueIsIgnored=True,bAdvancedView=False,)
-	   CustomProperties Pin (PinId=AE4340174DD534FEA5B812AE6DAEA1EB,PinName="bNewSimulate",PinToolTip="New Simulate\nBoolean",PinType.PinCategory="bool",PinType.PinSubCategoryObject=None,PinType.PinSubCategoryMemberReference=(),PinType.bIsArray=False,PinType.bIsReference=False,PinType.bIsConst=False,PinType.bIsWeakPointer=False,DefaultValue="true",AutogeneratedDefaultValue="false",PersistentGuid=00000000000000000000000000000000,bHidden=False,bNotConnectable=False,bDefaultValueIsReadOnly=False,bDefaultValueIsIgnored=False,bAdvancedView=False,)
-	   CustomProperties Pin (PinId=80EC7602465D797C7F0D4D93D6BA9BD5,PinName="bIncludeSelf",PinToolTip="Include Self\nBoolean",PinType.PinCategory="bool",PinType.PinSubCategoryObject=None,PinType.PinSubCategoryMemberReference=(),PinType.bIsArray=False,PinType.bIsReference=False,PinType.bIsConst=False,PinType.bIsWeakPointer=False,DefaultValue="true",AutogeneratedDefaultValue="true",PersistentGuid=00000000000000000000000000000000,bHidden=False,bNotConnectable=False,bDefaultValueIsReadOnly=False,bDefaultValueIsIgnored=False,bAdvancedView=False,)
-	End Object
-	Begin Object Class=K2Node_VariableGet Name="K2Node_VariableGet_181"
-	   VariableReference=(MemberName="BoneName",MemberGuid=A5E270C141CEE32A0269309215CEBAC4,bSelfContext=True)
-	   NodePosX=640
-	   NodePosY=128
-	   NodeGuid=B116D28D4BF4E6625FF02F92D2EA8DFD
-	   CustomProperties Pin (PinId=4DF93E5141ACE87768E5D8A7543C7B07,PinName="BoneName",Direction="EGPD_Output",PinType.PinCategory="name",PinType.PinSubCategoryObject=None,PinType.PinSubCategoryMemberReference=(),PinType.bIsArray=False,PinType.bIsReference=False,PinType.bIsConst=False,PinType.bIsWeakPointer=False,AutogeneratedDefaultValue="None",LinkedTo=(K2Node_CallFunction_8 7D29E35244DEB831C54B1CAEF4A5FE5C,),PersistentGuid=00000000000000000000000000000000,bHidden=False,bNotConnectable=False,bDefaultValueIsReadOnly=False,bDefaultValueIsIgnored=False,bAdvancedView=False,)
-	   CustomProperties Pin (PinId=F259489241A7853EE35268954C561767,PinName="self",PinFriendlyName=NSLOCTEXT("K2Node", "Target", "Target"),PinType.PinCategory="object",PinType.PinSubCategoryObject=BlueprintGeneratedClass'/Game/PhysAnimPawn.PhysAnimPawn_C',PinType.PinSubCategoryMemberReference=(),PinType.bIsArray=False,PinType.bIsReference=False,PinType.bIsConst=False,PinType.bIsWeakPointer=False,PersistentGuid=00000000000000000000000000000000,bHidden=True,bNotConnectable=False,bDefaultValueIsReadOnly=False,bDefaultValueIsIgnored=False,bAdvancedView=False,)
-	End Object
-[/VAR]
-
 
 ## Additional Resources
 
