Copying //UE4/Dev-Documentation to //UE4/Samples-Main (Source: //UE4/Dev-Documentation @ 4091886)

diff --git a/Engine/Documentation/Source/Engine/Physics/PhAT/HowTo/ApplyPhysicalAnimationProfile/ApplyPhysicalAnimProfile.INT.udn b/Engine/Documentation/Source/Engine/Physics/PhAT/HowTo/ApplyPhysicalAnimationProfile/ApplyPhysicalAnimProfile.INT.udn
--- a/Engine/Documentation/Source/Engine/Physics/PhAT/HowTo/ApplyPhysicalAnimationProfile/ApplyPhysicalAnimProfile.INT.udn
+++ b/Engine/Documentation/Source/Engine/Physics/PhAT/HowTo/ApplyPhysicalAnimationProfile/ApplyPhysicalAnimProfile.INT.udn
@@ -86,7 +86,108 @@ Below you will find the steps for creating a simple graph to enable a **Physical
 			![](ApplyPhysicialAnimProf_BP.png)
 		[/PARAM]
 		[PARAM:snippet]
-			%BPSnip_001%
+				Begin Object Class=K2Node_Event Name="K2Node_Event_0"
+					EventReference=(MemberParent=Class'/Script/Engine.Actor',MemberName="ReceiveBeginPlay")
+					bOverrideFunction=True
+					bCommentBubblePinned=True
+					NodeGuid=5B240C834C80C2A814C96C90E55C2B1D
+					CustomProperties Pin (PinId=971FD445473845FC61CB348BF812743F,PinName="OutputDelegate",Direction="EGPD_Output",PinType.PinCategory="delegate",PinType.PinSubCategoryObject=None,PinType.PinSubCategoryMemberReference=(MemberParent=Class'/Script/Engine.Actor',MemberName="ReceiveBeginPlay"),PinType.bIsArray=False,PinType.bIsReference=False,PinType.bIsConst=False,PinType.bIsWeakPointer=False,PersistentGuid=00000000000000000000000000000000,bHidden=False,bNotConnectable=False,bDefaultValueIsReadOnly=False,bDefaultValueIsIgnored=False,bAdvancedView=False,)
+					CustomProperties Pin (PinId=A695AE4A4D3440230EDBA897398BEE45,PinName="then",Direction="EGPD_Output",PinType.PinCategory="exec",PinType.PinSubCategoryObject=None,PinType.PinSubCategoryMemberReference=(),PinType.bIsArray=False,PinType.bIsReference=False,PinType.bIsConst=False,PinType.bIsWeakPointer=False,LinkedTo=(K2Node_CallFunction_0 193A453A43AF46E1ADBD5FA4F37E154B,),PersistentGuid=00000000000000000000000000000000,bHidden=False,bNotConnectable=False,bDefaultValueIsReadOnly=False,bDefaultValueIsIgnored=False,bAdvancedView=False,)
+				End Object
+				Begin Object Class=K2Node_VariableGet Name="K2Node_VariableGet_0"
+					VariableReference=(MemberName="PhysicalAnimationComponent",bSelfContext=True)
+					NodePosY=80
+					NodeGuid=BD6F0554431375C8D13E808D4DFFD28B
+					CustomProperties Pin (PinId=5F0735EC4A38FE885F7175B29F0832CA,PinName="PhysicalAnimationComponent",Direction="EGPD_Output",PinType.PinCategory="object",PinType.PinSubCategoryObject=Class'/Script/Engine.PhysicalAnimationComponent',PinType.PinSubCategoryMemberReference=(),PinType.bIsArray=False,PinType.bIsReference=False,PinType.bIsConst=False,PinType.bIsWeakPointer=False,LinkedTo=(K2Node_CallFunction_0 7766A5D54AEB0ACB1F95B49EF70BCB55,K2Node_Knot_1 36EFEC904DDD665AA5BC8EA40F722815,),PersistentGuid=00000000000000000000000000000000,bHidden=False,bNotConnectable=False,bDefaultValueIsReadOnly=False,bDefaultValueIsIgnored=False,bAdvancedView=False,)
+					CustomProperties Pin (PinId=7142470242FB2133424E24AD697AC6C0,PinName="self",PinFriendlyName=NSLOCTEXT("K2Node", "Target", "Target"),PinType.PinCategory="object",PinType.PinSubCategoryObject=BlueprintGeneratedClass'/Game/PhysAnimPawn.PhysAnimPawn_C',PinType.PinSubCategoryMemberReference=(),PinType.bIsArray=False,PinType.bIsReference=False,PinType.bIsConst=False,PinType.bIsWeakPointer=False,PersistentGuid=00000000000000000000000000000000,bHidden=True,bNotConnectable=False,bDefaultValueIsReadOnly=False,bDefaultValueIsIgnored=False,bAdvancedView=False,)
+				End Object
+				Begin Object Class=K2Node_VariableGet Name="K2Node_VariableGet_5"
+					VariableReference=(MemberName="SkeletalMeshComponent",bSelfContext=True)
+					NodePosX=16
+					NodePosY=144
+					NodeGuid=5F5D73144D96E94DCDBB92A9F9803A95
+					CustomProperties Pin (PinId=1EEABDD6494CBF38903C6691F5F5B45C,PinName="SkeletalMeshComponent",Direction="EGPD_Output",PinType.PinCategory="object",PinType.PinSubCategoryObject=Class'/Script/Engine.SkeletalMeshComponent',PinType.PinSubCategoryMemberReference=(),PinType.bIsArray=False,PinType.bIsReference=False,PinType.bIsConst=False,PinType.bIsWeakPointer=False,LinkedTo=(K2Node_CallFunction_0 37845433463046F2DB1841A4B947D094,),PersistentGuid=00000000000000000000000000000000,bHidden=False,bNotConnectable=False,bDefaultValueIsReadOnly=False,bDefaultValueIsIgnored=False,bAdvancedView=False,)
+					CustomProperties Pin (PinId=D68B011B4E4D641D541E8CA766972993,PinName="self",PinFriendlyName=NSLOCTEXT("K2Node", "Target", "Target"),PinType.PinCategory="object",PinType.PinSubCategoryObject=BlueprintGeneratedClass'/Game/PhysAnimPawn.PhysAnimPawn_C',PinType.PinSubCategoryMemberReference=(),PinType.bIsArray=False,PinType.bIsReference=False,PinType.bIsConst=False,PinType.bIsWeakPointer=False,PersistentGuid=00000000000000000000000000000000,bHidden=True,bNotConnectable=False,bDefaultValueIsReadOnly=False,bDefaultValueIsIgnored=False,bAdvancedView=False,)
+				End Object
+				Begin Object Class=K2Node_CallFunction Name="K2Node_CallFunction_0"
+					FunctionReference=(MemberParent=Class'/Script/Engine.PhysicalAnimationComponent',MemberName="SetSkeletalMeshComponent")
+					NodePosX=256
+					NodePosY=-16
+					NodeGuid=DEDD52104BE4DC7743D0068EAEF62DD1
+					CustomProperties Pin (PinId=193A453A43AF46E1ADBD5FA4F37E154B,PinName="execute",PinToolTip="\nExec",PinType.PinCategory="exec",PinType.PinSubCategoryObject=None,PinType.PinSubCategoryMemberReference=(),PinType.bIsArray=False,PinType.bIsReference=False,PinType.bIsConst=False,PinType.bIsWeakPointer=False,LinkedTo=(K2Node_Event_0 A695AE4A4D3440230EDBA897398BEE45,),PersistentGuid=00000000000000000000000000000000,bHidden=False,bNotConnectable=False,bDefaultValueIsReadOnly=False,bDefaultValueIsIgnored=False,bAdvancedView=False,)
+					CustomProperties Pin (PinId=BE952CF64A6E491E6C0B8F8B5CEB1000,PinName="then",PinToolTip="\nExec",Direction="EGPD_Output",PinType.PinCategory="exec",PinType.PinSubCategoryObject=None,PinType.PinSubCategoryMemberReference=(),PinType.bIsArray=False,PinType.bIsReference=False,PinType.bIsConst=False,PinType.bIsWeakPointer=False,LinkedTo=(K2Node_CallFunction_20 B96C42C240335FB66D65AA84ADD838AC,),PersistentGuid=00000000000000000000000000000000,bHidden=False,bNotConnectable=False,bDefaultValueIsReadOnly=False,bDefaultValueIsIgnored=False,bAdvancedView=False,)
+					CustomProperties Pin (PinId=7766A5D54AEB0ACB1F95B49EF70BCB55,PinName="self",PinFriendlyName=NSLOCTEXT("K2Node", "Target", "Target"),PinToolTip="Target\nPhysical Animation Component Reference",PinType.PinCategory="object",PinType.PinSubCategoryObject=Class'/Script/Engine.PhysicalAnimationComponent',PinType.PinSubCategoryMemberReference=(),PinType.bIsArray=False,PinType.bIsReference=False,PinType.bIsConst=False,PinType.bIsWeakPointer=False,LinkedTo=(K2Node_VariableGet_0 5F0735EC4A38FE885F7175B29F0832CA,),PersistentGuid=00000000000000000000000000000000,bHidden=False,bNotConnectable=False,bDefaultValueIsReadOnly=False,bDefaultValueIsIgnored=False,bAdvancedView=False,)
+					CustomProperties Pin (PinId=37845433463046F2DB1841A4B947D094,PinName="InSkeletalMeshComponent",PinToolTip="In Skeletal Mesh Component\nSkeletal Mesh Component Reference",PinType.PinCategory="object",PinType.PinSubCategoryObject=Class'/Script/Engine.SkeletalMeshComponent',PinType.PinSubCategoryMemberReference=(),PinType.bIsArray=False,PinType.bIsReference=False,PinType.bIsConst=False,PinType.bIsWeakPointer=False,LinkedTo=(K2Node_VariableGet_5 1EEABDD6494CBF38903C6691F5F5B45C,),PersistentGuid=00000000000000000000000000000000,bHidden=False,bNotConnectable=False,bDefaultValueIsReadOnly=False,bDefaultValueIsIgnored=False,bAdvancedView=False,)
+				End Object
+				Begin Object Class=K2Node_Knot Name="K2Node_Knot_1"
+					NodePosX=256
+					NodePosY=128
+					NodeGuid=E4A8165E4E49B3FDC3D36AB3198BE2A7
+					CustomProperties Pin (PinId=36EFEC904DDD665AA5BC8EA40F722815,PinName="InputPin",PinType.PinCategory="object",PinType.PinSubCategoryObject=Class'/Script/Engine.PhysicalAnimationComponent',PinType.PinSubCategoryMemberReference=(),PinType.bIsArray=False,PinType.bIsReference=False,PinType.bIsConst=False,PinType.bIsWeakPointer=False,LinkedTo=(K2Node_VariableGet_0 5F0735EC4A38FE885F7175B29F0832CA,),PersistentGuid=00000000000000000000000000000000,bHidden=False,bNotConnectable=False,bDefaultValueIsReadOnly=False,bDefaultValueIsIgnored=True,bAdvancedView=False,)
+					CustomProperties Pin (PinId=E54B8E124E80A63FF92ACF8FA75AACAA,PinName="OutputPin",Direction="EGPD_Output",PinType.PinCategory="object",PinType.PinSubCategoryObject=Class'/Script/Engine.PhysicalAnimationComponent',PinType.PinSubCategoryMemberReference=(),PinType.bIsArray=False,PinType.bIsReference=False,PinType.bIsConst=False,PinType.bIsWeakPointer=False,LinkedTo=(K2Node_Knot_2 A246331D43359B8F04A7EBBEDB2EDC22,),PersistentGuid=00000000000000000000000000000000,bHidden=False,bNotConnectable=False,bDefaultValueIsReadOnly=False,bDefaultValueIsIgnored=False,bAdvancedView=False,)
+				End Object
+				Begin Object Class=K2Node_Knot Name="K2Node_Knot_2"
+					NodePosX=496
+					NodePosY=128
+					NodeGuid=59DD0FBD4738F7C8ADCAB787A410908C
+					CustomProperties Pin (PinId=A246331D43359B8F04A7EBBEDB2EDC22,PinName="InputPin",PinType.PinCategory="object",PinType.PinSubCategoryObject=Class'/Script/Engine.PhysicalAnimationComponent',PinType.PinSubCategoryMemberReference=(),PinType.bIsArray=False,PinType.bIsReference=False,PinType.bIsConst=False,PinType.bIsWeakPointer=False,LinkedTo=(K2Node_Knot_1 E54B8E124E80A63FF92ACF8FA75AACAA,),PersistentGuid=00000000000000000000000000000000,bHidden=False,bNotConnectable=False,bDefaultValueIsReadOnly=False,bDefaultValueIsIgnored=True,bAdvancedView=False,)
+					CustomProperties Pin (PinId=CC66C9584A8CC28A92B6D4AD7CD40472,PinName="OutputPin",Direction="EGPD_Output",PinType.PinCategory="object",PinType.PinSubCategoryObject=Class'/Script/Engine.PhysicalAnimationComponent',PinType.PinSubCategoryMemberReference=(),PinType.bIsArray=False,PinType.bIsReference=False,PinType.bIsConst=False,PinType.bIsWeakPointer=False,LinkedTo=(K2Node_CallFunction_20 96A091DE4DF2762BDFA90CB50CAE9BFB,),PersistentGuid=00000000000000000000000000000000,bHidden=False,bNotConnectable=False,bDefaultValueIsReadOnly=False,bDefaultValueIsIgnored=False,bAdvancedView=False,)
+				End Object
+				Begin Object Class=K2Node_VariableGet Name="K2Node_VariableGet_37"
+					VariableReference=(MemberName="BoneName",bSelfContext=True)
+					NodePosX=400
+					NodePosY=160
+					NodeGuid=7B320511484EBE8F7DDB7B8C7044CD27
+					CustomProperties Pin (PinId=62FD25D94D3A5CAE32C5E685EDA461B1,PinName="BoneName",Direction="EGPD_Output",PinType.PinCategory="name",PinType.PinSubCategoryObject=None,PinType.PinSubCategoryMemberReference=(),PinType.bIsArray=False,PinType.bIsReference=False,PinType.bIsConst=False,PinType.bIsWeakPointer=False,AutogeneratedDefaultValue="None",LinkedTo=(K2Node_Knot_4 4BF4C00042C37A0CDB25AD926D106B22,K2Node_CallFunction_20 EB57682C40C2F631928F3D88FF0B2117,),PersistentGuid=00000000000000000000000000000000,bHidden=False,bNotConnectable=False,bDefaultValueIsReadOnly=False,bDefaultValueIsIgnored=False,bAdvancedView=False,)
+					CustomProperties Pin (PinId=CC9F75D84295F4E8568F03A68C65A446,PinName="self",PinFriendlyName=NSLOCTEXT("K2Node", "Target", "Target"),PinType.PinCategory="object",PinType.PinSubCategoryObject=BlueprintGeneratedClass'/Game/PhysAnimPawn.PhysAnimPawn_C',PinType.PinSubCategoryMemberReference=(),PinType.bIsArray=False,PinType.bIsReference=False,PinType.bIsConst=False,PinType.bIsWeakPointer=False,PersistentGuid=00000000000000000000000000000000,bHidden=True,bNotConnectable=False,bDefaultValueIsReadOnly=False,bDefaultValueIsIgnored=False,bAdvancedView=False,)
+				End Object
+				Begin Object Class=K2Node_VariableGet Name="K2Node_VariableGet_1"
+					VariableReference=(MemberName="SkeletalMeshComponent",bSelfContext=True)
+					NodePosX=928
+					NodePosY=64
+					NodeGuid=5E9CBBED44802CB75DD821972402DEFD
+					CustomProperties Pin (PinId=1EEABDD6494CBF38903C6691F5F5B45C,PinName="SkeletalMeshComponent",Direction="EGPD_Output",PinType.PinCategory="object",PinType.PinSubCategoryObject=Class'/Script/Engine.SkeletalMeshComponent',PinType.PinSubCategoryMemberReference=(),PinType.bIsArray=False,PinType.bIsReference=False,PinType.bIsConst=False,PinType.bIsWeakPointer=False,LinkedTo=(K2Node_CallFunction_8 4DD6FE2846E695E0F3F569B06CED0A46,),PersistentGuid=00000000000000000000000000000000,bHidden=False,bNotConnectable=False,bDefaultValueIsReadOnly=False,bDefaultValueIsIgnored=False,bAdvancedView=False,)
+					CustomProperties Pin (PinId=D68B011B4E4D641D541E8CA766972993,PinName="self",PinFriendlyName=NSLOCTEXT("K2Node", "Target", "Target"),PinType.PinCategory="object",PinType.PinSubCategoryObject=BlueprintGeneratedClass'/Game/PhysAnimPawn.PhysAnimPawn_C',PinType.PinSubCategoryMemberReference=(),PinType.bIsArray=False,PinType.bIsReference=False,PinType.bIsConst=False,PinType.bIsWeakPointer=False,PersistentGuid=00000000000000000000000000000000,bHidden=True,bNotConnectable=False,bDefaultValueIsReadOnly=False,bDefaultValueIsIgnored=False,bAdvancedView=False,)
+				End Object
+				Begin Object Class=K2Node_CallFunction Name="K2Node_CallFunction_8"
+					FunctionReference=(MemberParent=Class'/Script/Engine.SkeletalMeshComponent',MemberName="SetAllBodiesBelowSimulatePhysics")
+					NodePosX=1152
+					NodePosY=-16
+					NodeGuid=C58AFBDD484D31A196381BA69976AB10
+					CustomProperties Pin (PinId=399238E643A83FB3CFC85CB926B5CF9F,PinName="execute",PinToolTip="\nExec",PinType.PinCategory="exec",PinType.PinSubCategoryObject=None,PinType.PinSubCategoryMemberReference=(),PinType.bIsArray=False,PinType.bIsReference=False,PinType.bIsConst=False,PinType.bIsWeakPointer=False,LinkedTo=(K2Node_CallFunction_20 B38EF5724C795C379DDB809F78D77BA1,),PersistentGuid=00000000000000000000000000000000,bHidden=False,bNotConnectable=False,bDefaultValueIsReadOnly=False,bDefaultValueIsIgnored=False,bAdvancedView=False,)
+					CustomProperties Pin (PinId=C132D8014BC6A4DE84C437AF07CF838E,PinName="then",PinToolTip="\nExec",Direction="EGPD_Output",PinType.PinCategory="exec",PinType.PinSubCategoryObject=None,PinType.PinSubCategoryMemberReference=(),PinType.bIsArray=False,PinType.bIsReference=False,PinType.bIsConst=False,PinType.bIsWeakPointer=False,PersistentGuid=00000000000000000000000000000000,bHidden=False,bNotConnectable=False,bDefaultValueIsReadOnly=False,bDefaultValueIsIgnored=False,bAdvancedView=False,)
+					CustomProperties Pin (PinId=4DD6FE2846E695E0F3F569B06CED0A46,PinName="self",PinFriendlyName=NSLOCTEXT("K2Node", "Target", "Target"),PinToolTip="Target\nSkeletal Mesh Component Reference",PinType.PinCategory="object",PinType.PinSubCategoryObject=Class'/Script/Engine.SkeletalMeshComponent',PinType.PinSubCategoryMemberReference=(),PinType.bIsArray=False,PinType.bIsReference=False,PinType.bIsConst=False,PinType.bIsWeakPointer=False,LinkedTo=(K2Node_VariableGet_1 1EEABDD6494CBF38903C6691F5F5B45C,),PersistentGuid=00000000000000000000000000000000,bHidden=False,bNotConnectable=False,bDefaultValueIsReadOnly=False,bDefaultValueIsIgnored=False,bAdvancedView=False,)
+					CustomProperties Pin (PinId=7D29E35244DEB831C54B1CAEF4A5FE5C,PinName="InBoneName",PinToolTip="In Bone Name\nName (by ref)",PinType.PinCategory="name",PinType.PinSubCategoryObject=None,PinType.PinSubCategoryMemberReference=(),PinType.bIsArray=False,PinType.bIsReference=True,PinType.bIsConst=True,PinType.bIsWeakPointer=False,DefaultValue="None",AutogeneratedDefaultValue="None",LinkedTo=(K2Node_Knot_3 336A876240769F5A4170CF8126A5CF36,),PersistentGuid=00000000000000000000000000000000,bHidden=False,bNotConnectable=False,bDefaultValueIsReadOnly=False,bDefaultValueIsIgnored=True,bAdvancedView=False,)
+					CustomProperties Pin (PinId=AE4340174DD534FEA5B812AE6DAEA1EB,PinName="bNewSimulate",PinToolTip="New Simulate\nBoolean",PinType.PinCategory="bool",PinType.PinSubCategoryObject=None,PinType.PinSubCategoryMemberReference=(),PinType.bIsArray=False,PinType.bIsReference=False,PinType.bIsConst=False,PinType.bIsWeakPointer=False,DefaultValue="true",AutogeneratedDefaultValue="false",PersistentGuid=00000000000000000000000000000000,bHidden=False,bNotConnectable=False,bDefaultValueIsReadOnly=False,bDefaultValueIsIgnored=False,bAdvancedView=False,)
+					CustomProperties Pin (PinId=80EC7602465D797C7F0D4D93D6BA9BD5,PinName="bIncludeSelf",PinToolTip="Include Self\nBoolean",PinType.PinCategory="bool",PinType.PinSubCategoryObject=None,PinType.PinSubCategoryMemberReference=(),PinType.bIsArray=False,PinType.bIsReference=False,PinType.bIsConst=False,PinType.bIsWeakPointer=False,DefaultValue="true",AutogeneratedDefaultValue="true",PersistentGuid=00000000000000000000000000000000,bHidden=False,bNotConnectable=False,bDefaultValueIsReadOnly=False,bDefaultValueIsIgnored=False,bAdvancedView=False,)
+				End Object
+				Begin Object Class=K2Node_Knot Name="K2Node_Knot_3"
+					NodePosX=976
+					NodePosY=224
+					NodeGuid=55465C7448CB280BB4657AA36C94D451
+					CustomProperties Pin (PinId=D6618E7C42476ADDD628FF855FD6E7C9,PinName="InputPin",PinType.PinCategory="name",PinType.PinSubCategoryObject=None,PinType.PinSubCategoryMemberReference=(),PinType.bIsArray=False,PinType.bIsReference=False,PinType.bIsConst=False,PinType.bIsWeakPointer=False,LinkedTo=(K2Node_Knot_4 A0911D3740ABD40031008799EB009AE0,),PersistentGuid=00000000000000000000000000000000,bHidden=False,bNotConnectable=False,bDefaultValueIsReadOnly=False,bDefaultValueIsIgnored=True,bAdvancedView=False,)
+					CustomProperties Pin (PinId=336A876240769F5A4170CF8126A5CF36,PinName="OutputPin",Direction="EGPD_Output",PinType.PinCategory="name",PinType.PinSubCategoryObject=None,PinType.PinSubCategoryMemberReference=(),PinType.bIsArray=False,PinType.bIsReference=False,PinType.bIsConst=False,PinType.bIsWeakPointer=False,LinkedTo=(K2Node_CallFunction_8 7D29E35244DEB831C54B1CAEF4A5FE5C,),PersistentGuid=00000000000000000000000000000000,bHidden=False,bNotConnectable=False,bDefaultValueIsReadOnly=False,bDefaultValueIsIgnored=False,bAdvancedView=False,)
+				End Object
+				Begin Object Class=K2Node_Knot Name="K2Node_Knot_4"
+					NodePosX=608
+					NodePosY=224
+					NodeGuid=230B64A94C55B4E5636829B6DC0D21E7
+					CustomProperties Pin (PinId=4BF4C00042C37A0CDB25AD926D106B22,PinName="InputPin",PinType.PinCategory="name",PinType.PinSubCategoryObject=None,PinType.PinSubCategoryMemberReference=(),PinType.bIsArray=False,PinType.bIsReference=False,PinType.bIsConst=False,PinType.bIsWeakPointer=False,LinkedTo=(K2Node_VariableGet_37 62FD25D94D3A5CAE32C5E685EDA461B1,),PersistentGuid=00000000000000000000000000000000,bHidden=False,bNotConnectable=False,bDefaultValueIsReadOnly=False,bDefaultValueIsIgnored=True,bAdvancedView=False,)
+					CustomProperties Pin (PinId=A0911D3740ABD40031008799EB009AE0,PinName="OutputPin",Direction="EGPD_Output",PinType.PinCategory="name",PinType.PinSubCategoryObject=None,PinType.PinSubCategoryMemberReference=(),PinType.bIsArray=False,PinType.bIsReference=False,PinType.bIsConst=False,PinType.bIsWeakPointer=False,LinkedTo=(K2Node_Knot_3 D6618E7C42476ADDD628FF855FD6E7C9,),PersistentGuid=00000000000000000000000000000000,bHidden=False,bNotConnectable=False,bDefaultValueIsReadOnly=False,bDefaultValueIsIgnored=False,bAdvancedView=False,)
+				End Object
+				Begin Object Class=K2Node_CallFunction Name="K2Node_CallFunction_20"
+					FunctionReference=(MemberParent=Class'/Script/Engine.PhysicalAnimationComponent',MemberName="ApplyPhysicalAnimationProfileBelow")
+					NodePosX=624
+					NodePosY=-16
+					NodeGuid=D1A218604F6806C6A38DF0B1C6812215
+					CustomProperties Pin (PinId=B96C42C240335FB66D65AA84ADD838AC,PinName="execute",PinToolTip="\nExec",PinType.PinCategory="exec",PinType.PinSubCategoryObject=None,PinType.PinSubCategoryMemberReference=(),PinType.bIsArray=False,PinType.bIsReference=False,PinType.bIsConst=False,PinType.bIsWeakPointer=False,LinkedTo=(K2Node_CallFunction_0 BE952CF64A6E491E6C0B8F8B5CEB1000,),PersistentGuid=00000000000000000000000000000000,bHidden=False,bNotConnectable=False,bDefaultValueIsReadOnly=False,bDefaultValueIsIgnored=False,bAdvancedView=False,)
+					CustomProperties Pin (PinId=B38EF5724C795C379DDB809F78D77BA1,PinName="then",PinToolTip="\nExec",Direction="EGPD_Output",PinType.PinCategory="exec",PinType.PinSubCategoryObject=None,PinType.PinSubCategoryMemberReference=(),PinType.bIsArray=False,PinType.bIsReference=False,PinType.bIsConst=False,PinType.bIsWeakPointer=False,LinkedTo=(K2Node_CallFunction_8 399238E643A83FB3CFC85CB926B5CF9F,),PersistentGuid=00000000000000000000000000000000,bHidden=False,bNotConnectable=False,bDefaultValueIsReadOnly=False,bDefaultValueIsIgnored=False,bAdvancedView=False,)
+					CustomProperties Pin (PinId=96A091DE4DF2762BDFA90CB50CAE9BFB,PinName="self",PinFriendlyName=NSLOCTEXT("K2Node", "Target", "Target"),PinToolTip="Target\nPhysical Animation Component Reference",PinType.PinCategory="object",PinType.PinSubCategoryObject=Class'/Script/Engine.PhysicalAnimationComponent',PinType.PinSubCategoryMemberReference=(),PinType.bIsArray=False,PinType.bIsReference=False,PinType.bIsConst=False,PinType.bIsWeakPointer=False,LinkedTo=(K2Node_Knot_2 CC66C9584A8CC28A92B6D4AD7CD40472,),PersistentGuid=00000000000000000000000000000000,bHidden=False,bNotConnectable=False,bDefaultValueIsReadOnly=False,bDefaultValueIsIgnored=False,bAdvancedView=False,)
+					CustomProperties Pin (PinId=EB57682C40C2F631928F3D88FF0B2117,PinName="BodyName",PinToolTip="Body Name\nName\n\nThe body from which we\'d like to start applying the physical animation profile. Finds all bodies below in the skeleton hierarchy. None implies all bodies",PinType.PinCategory="name",PinType.PinSubCategoryObject=None,PinType.PinSubCategoryMemberReference=(),PinType.bIsArray=False,PinType.bIsReference=False,PinType.bIsConst=False,PinType.bIsWeakPointer=False,DefaultValue="None",AutogeneratedDefaultValue="None",LinkedTo=(K2Node_VariableGet_37 62FD25D94D3A5CAE32C5E685EDA461B1,),PersistentGuid=00000000000000000000000000000000,bHidden=False,bNotConnectable=False,bDefaultValueIsReadOnly=False,bDefaultValueIsIgnored=False,bAdvancedView=False,)
+					CustomProperties Pin (PinId=D9E9561844B25F69A79A2C9FE98806C6,PinName="ProfileName",PinToolTip="Profile Name\nName\n\nThe physical animation profile we\'d like to apply. For each body in the physics asset we search for physical animation settings with this name.",PinType.PinCategory="name",PinType.PinSubCategoryObject=None,PinType.PinSubCategoryMemberReference=(),PinType.bIsArray=False,PinType.bIsReference=False,PinType.bIsConst=False,PinType.bIsWeakPointer=False,DefaultValue="DocProfile",AutogeneratedDefaultValue="None",PersistentGuid=00000000000000000000000000000000,bHidden=False,bNotConnectable=False,bDefaultValueIsReadOnly=False,bDefaultValueIsIgnored=False,bAdvancedView=False,)
+					CustomProperties Pin (PinId=7D946B3045556970B2A7519A2614C109,PinName="bIncludeSelf",PinToolTip="Include Self\nBoolean\n\nWhether to include the provided body name in the list of bodies we act on (useful to ignore for cases where a root has multiple children)",PinType.PinCategory="bool",PinType.PinSubCategoryObject=None,PinType.PinSubCategoryMemberReference=(),PinType.bIsArray=False,PinType.bIsReference=False,PinType.bIsConst=False,PinType.bIsWeakPointer=False,DefaultValue="true",AutogeneratedDefaultValue="true",PersistentGuid=00000000000000000000000000000000,bHidden=False,bNotConnectable=False,bDefaultValueIsReadOnly=False,bDefaultValueIsIgnored=False,bAdvancedView=False,)
+					CustomProperties Pin (PinId=41EFD7554F080B09E40F7199FA2BAF6F,PinName="bClearNotFound",PinToolTip="Clear Not Found\nBoolean\n\nIf true, bodies without the given profile name will have any existing physical animation settings cleared. If false, bodies without the given profile name are left untouched.",PinType.PinCategory="bool",PinType.PinSubCategoryObject=None,PinType.PinSubCategoryMemberReference=(),PinType.bIsArray=False,PinType.bIsReference=False,PinType.bIsConst=False,PinType.bIsWeakPointer=False,DefaultValue="true",AutogeneratedDefaultValue="false",PersistentGuid=00000000000000000000000000000000,bHidden=False,bNotConnectable=False,bDefaultValueIsReadOnly=False,bDefaultValueIsIgnored=False,bAdvancedView=False,)
+				End Object
 		[/PARAM]
 	[/OBJECT]
 
@@ -109,115 +210,6 @@ When finished, your **Skeletal Mesh Component** will be simulating physics using
 
 ![](resultApplyPAP.png)
 
-[COMMENT:none]
- Variables 
-[/COMMENT]
-
-[VAR:BPSnip_001]
-	Begin Object Class=K2Node_Event Name="K2Node_Event_0"
-		EventReference=(MemberParent=Class'/Script/Engine.Actor',MemberName="ReceiveBeginPlay")
-		bOverrideFunction=True
-		bCommentBubblePinned=True
-		NodeGuid=5B240C834C80C2A814C96C90E55C2B1D
-		CustomProperties Pin (PinId=971FD445473845FC61CB348BF812743F,PinName="OutputDelegate",Direction="EGPD_Output",PinType.PinCategory="delegate",PinType.PinSubCategoryObject=None,PinType.PinSubCategoryMemberReference=(MemberParent=Class'/Script/Engine.Actor',MemberName="ReceiveBeginPlay"),PinType.bIsArray=False,PinType.bIsReference=False,PinType.bIsConst=False,PinType.bIsWeakPointer=False,PersistentGuid=00000000000000000000000000000000,bHidden=False,bNotConnectable=False,bDefaultValueIsReadOnly=False,bDefaultValueIsIgnored=False,bAdvancedView=False,)
-		CustomProperties Pin (PinId=A695AE4A4D3440230EDBA897398BEE45,PinName="then",Direction="EGPD_Output",PinType.PinCategory="exec",PinType.PinSubCategoryObject=None,PinType.PinSubCategoryMemberReference=(),PinType.bIsArray=False,PinType.bIsReference=False,PinType.bIsConst=False,PinType.bIsWeakPointer=False,LinkedTo=(K2Node_CallFunction_0 193A453A43AF46E1ADBD5FA4F37E154B,),PersistentGuid=00000000000000000000000000000000,bHidden=False,bNotConnectable=False,bDefaultValueIsReadOnly=False,bDefaultValueIsIgnored=False,bAdvancedView=False,)
-	End Object
-	Begin Object Class=K2Node_VariableGet Name="K2Node_VariableGet_0"
-		VariableReference=(MemberName="PhysicalAnimationComponent",bSelfContext=True)
-		NodePosY=80
-		NodeGuid=BD6F0554431375C8D13E808D4DFFD28B
-		CustomProperties Pin (PinId=5F0735EC4A38FE885F7175B29F0832CA,PinName="PhysicalAnimationComponent",Direction="EGPD_Output",PinType.PinCategory="object",PinType.PinSubCategoryObject=Class'/Script/Engine.PhysicalAnimationComponent',PinType.PinSubCategoryMemberReference=(),PinType.bIsArray=False,PinType.bIsReference=False,PinType.bIsConst=False,PinType.bIsWeakPointer=False,LinkedTo=(K2Node_CallFunction_0 7766A5D54AEB0ACB1F95B49EF70BCB55,K2Node_Knot_1 36EFEC904DDD665AA5BC8EA40F722815,),PersistentGuid=00000000000000000000000000000000,bHidden=False,bNotConnectable=False,bDefaultValueIsReadOnly=False,bDefaultValueIsIgnored=False,bAdvancedView=False,)
-		CustomProperties Pin (PinId=7142470242FB2133424E24AD697AC6C0,PinName="self",PinFriendlyName=NSLOCTEXT("K2Node", "Target", "Target"),PinType.PinCategory="object",PinType.PinSubCategoryObject=BlueprintGeneratedClass'/Game/PhysAnimPawn.PhysAnimPawn_C',PinType.PinSubCategoryMemberReference=(),PinType.bIsArray=False,PinType.bIsReference=False,PinType.bIsConst=False,PinType.bIsWeakPointer=False,PersistentGuid=00000000000000000000000000000000,bHidden=True,bNotConnectable=False,bDefaultValueIsReadOnly=False,bDefaultValueIsIgnored=False,bAdvancedView=False,)
-	End Object
-	Begin Object Class=K2Node_VariableGet Name="K2Node_VariableGet_5"
-		VariableReference=(MemberName="SkeletalMeshComponent",bSelfContext=True)
-		NodePosX=16
-		NodePosY=144
-		NodeGuid=5F5D73144D96E94DCDBB92A9F9803A95
-		CustomProperties Pin (PinId=1EEABDD6494CBF38903C6691F5F5B45C,PinName="SkeletalMeshComponent",Direction="EGPD_Output",PinType.PinCategory="object",PinType.PinSubCategoryObject=Class'/Script/Engine.SkeletalMeshComponent',PinType.PinSubCategoryMemberReference=(),PinType.bIsArray=False,PinType.bIsReference=False,PinType.bIsConst=False,PinType.bIsWeakPointer=False,LinkedTo=(K2Node_CallFunction_0 37845433463046F2DB1841A4B947D094,),PersistentGuid=00000000000000000000000000000000,bHidden=False,bNotConnectable=False,bDefaultValueIsReadOnly=False,bDefaultValueIsIgnored=False,bAdvancedView=False,)
-		CustomProperties Pin (PinId=D68B011B4E4D641D541E8CA766972993,PinName="self",PinFriendlyName=NSLOCTEXT("K2Node", "Target", "Target"),PinType.PinCategory="object",PinType.PinSubCategoryObject=BlueprintGeneratedClass'/Game/PhysAnimPawn.PhysAnimPawn_C',PinType.PinSubCategoryMemberReference=(),PinType.bIsArray=False,PinType.bIsReference=False,PinType.bIsConst=False,PinType.bIsWeakPointer=False,PersistentGuid=00000000000000000000000000000000,bHidden=True,bNotConnectable=False,bDefaultValueIsReadOnly=False,bDefaultValueIsIgnored=False,bAdvancedView=False,)
-	End Object
-	Begin Object Class=K2Node_CallFunction Name="K2Node_CallFunction_0"
-		FunctionReference=(MemberParent=Class'/Script/Engine.PhysicalAnimationComponent',MemberName="SetSkeletalMeshComponent")
-		NodePosX=256
-		NodePosY=-16
-		NodeGuid=DEDD52104BE4DC7743D0068EAEF62DD1
-		CustomProperties Pin (PinId=193A453A43AF46E1ADBD5FA4F37E154B,PinName="execute",PinToolTip="\nExec",PinType.PinCategory="exec",PinType.PinSubCategoryObject=None,PinType.PinSubCategoryMemberReference=(),PinType.bIsArray=False,PinType.bIsReference=False,PinType.bIsConst=False,PinType.bIsWeakPointer=False,LinkedTo=(K2Node_Event_0 A695AE4A4D3440230EDBA897398BEE45,),PersistentGuid=00000000000000000000000000000000,bHidden=False,bNotConnectable=False,bDefaultValueIsReadOnly=False,bDefaultValueIsIgnored=False,bAdvancedView=False,)
-		CustomProperties Pin (PinId=BE952CF64A6E491E6C0B8F8B5CEB1000,PinName="then",PinToolTip="\nExec",Direction="EGPD_Output",PinType.PinCategory="exec",PinType.PinSubCategoryObject=None,PinType.PinSubCategoryMemberReference=(),PinType.bIsArray=False,PinType.bIsReference=False,PinType.bIsConst=False,PinType.bIsWeakPointer=False,LinkedTo=(K2Node_CallFunction_20 B96C42C240335FB66D65AA84ADD838AC,),PersistentGuid=00000000000000000000000000000000,bHidden=False,bNotConnectable=False,bDefaultValueIsReadOnly=False,bDefaultValueIsIgnored=False,bAdvancedView=False,)
-		CustomProperties Pin (PinId=7766A5D54AEB0ACB1F95B49EF70BCB55,PinName="self",PinFriendlyName=NSLOCTEXT("K2Node", "Target", "Target"),PinToolTip="Target\nPhysical Animation Component Reference",PinType.PinCategory="object",PinType.PinSubCategoryObject=Class'/Script/Engine.PhysicalAnimationComponent',PinType.PinSubCategoryMemberReference=(),PinType.bIsArray=False,PinType.bIsReference=False,PinType.bIsConst=False,PinType.bIsWeakPointer=False,LinkedTo=(K2Node_VariableGet_0 5F0735EC4A38FE885F7175B29F0832CA,),PersistentGuid=00000000000000000000000000000000,bHidden=False,bNotConnectable=False,bDefaultValueIsReadOnly=False,bDefaultValueIsIgnored=False,bAdvancedView=False,)
-		CustomProperties Pin (PinId=37845433463046F2DB1841A4B947D094,PinName="InSkeletalMeshComponent",PinToolTip="In Skeletal Mesh Component\nSkeletal Mesh Component Reference",PinType.PinCategory="object",PinType.PinSubCategoryObject=Class'/Script/Engine.SkeletalMeshComponent',PinType.PinSubCategoryMemberReference=(),PinType.bIsArray=False,PinType.bIsReference=False,PinType.bIsConst=False,PinType.bIsWeakPointer=False,LinkedTo=(K2Node_VariableGet_5 1EEABDD6494CBF38903C6691F5F5B45C,),PersistentGuid=00000000000000000000000000000000,bHidden=False,bNotConnectable=False,bDefaultValueIsReadOnly=False,bDefaultValueIsIgnored=False,bAdvancedView=False,)
-	End Object
-	Begin Object Class=K2Node_Knot Name="K2Node_Knot_1"
-		NodePosX=256
-		NodePosY=128
-		NodeGuid=E4A8165E4E49B3FDC3D36AB3198BE2A7
-		CustomProperties Pin (PinId=36EFEC904DDD665AA5BC8EA40F722815,PinName="InputPin",PinType.PinCategory="object",PinType.PinSubCategoryObject=Class'/Script/Engine.PhysicalAnimationComponent',PinType.PinSubCategoryMemberReference=(),PinType.bIsArray=False,PinType.bIsReference=False,PinType.bIsConst=False,PinType.bIsWeakPointer=False,LinkedTo=(K2Node_VariableGet_0 5F0735EC4A38FE885F7175B29F0832CA,),PersistentGuid=00000000000000000000000000000000,bHidden=False,bNotConnectable=False,bDefaultValueIsReadOnly=False,bDefaultValueIsIgnored=True,bAdvancedView=False,)
-		CustomProperties Pin (PinId=E54B8E124E80A63FF92ACF8FA75AACAA,PinName="OutputPin",Direction="EGPD_Output",PinType.PinCategory="object",PinType.PinSubCategoryObject=Class'/Script/Engine.PhysicalAnimationComponent',PinType.PinSubCategoryMemberReference=(),PinType.bIsArray=False,PinType.bIsReference=False,PinType.bIsConst=False,PinType.bIsWeakPointer=False,LinkedTo=(K2Node_Knot_2 A246331D43359B8F04A7EBBEDB2EDC22,),PersistentGuid=00000000000000000000000000000000,bHidden=False,bNotConnectable=False,bDefaultValueIsReadOnly=False,bDefaultValueIsIgnored=False,bAdvancedView=False,)
-	End Object
-	Begin Object Class=K2Node_Knot Name="K2Node_Knot_2"
-		NodePosX=496
-		NodePosY=128
-		NodeGuid=59DD0FBD4738F7C8ADCAB787A410908C
-		CustomProperties Pin (PinId=A246331D43359B8F04A7EBBEDB2EDC22,PinName="InputPin",PinType.PinCategory="object",PinType.PinSubCategoryObject=Class'/Script/Engine.PhysicalAnimationComponent',PinType.PinSubCategoryMemberReference=(),PinType.bIsArray=False,PinType.bIsReference=False,PinType.bIsConst=False,PinType.bIsWeakPointer=False,LinkedTo=(K2Node_Knot_1 E54B8E124E80A63FF92ACF8FA75AACAA,),PersistentGuid=00000000000000000000000000000000,bHidden=False,bNotConnectable=False,bDefaultValueIsReadOnly=False,bDefaultValueIsIgnored=True,bAdvancedView=False,)
-		CustomProperties Pin (PinId=CC66C9584A8CC28A92B6D4AD7CD40472,PinName="OutputPin",Direction="EGPD_Output",PinType.PinCategory="object",PinType.PinSubCategoryObject=Class'/Script/Engine.PhysicalAnimationComponent',PinType.PinSubCategoryMemberReference=(),PinType.bIsArray=False,PinType.bIsReference=False,PinType.bIsConst=False,PinType.bIsWeakPointer=False,LinkedTo=(K2Node_CallFunction_20 96A091DE4DF2762BDFA90CB50CAE9BFB,),PersistentGuid=00000000000000000000000000000000,bHidden=False,bNotConnectable=False,bDefaultValueIsReadOnly=False,bDefaultValueIsIgnored=False,bAdvancedView=False,)
-	End Object
-	Begin Object Class=K2Node_VariableGet Name="K2Node_VariableGet_37"
-		VariableReference=(MemberName="BoneName",bSelfContext=True)
-		NodePosX=400
-		NodePosY=160
-		NodeGuid=7B320511484EBE8F7DDB7B8C7044CD27
-		CustomProperties Pin (PinId=62FD25D94D3A5CAE32C5E685EDA461B1,PinName="BoneName",Direction="EGPD_Output",PinType.PinCategory="name",PinType.PinSubCategoryObject=None,PinType.PinSubCategoryMemberReference=(),PinType.bIsArray=False,PinType.bIsReference=False,PinType.bIsConst=False,PinType.bIsWeakPointer=False,AutogeneratedDefaultValue="None",LinkedTo=(K2Node_Knot_4 4BF4C00042C37A0CDB25AD926D106B22,K2Node_CallFunction_20 EB57682C40C2F631928F3D88FF0B2117,),PersistentGuid=00000000000000000000000000000000,bHidden=False,bNotConnectable=False,bDefaultValueIsReadOnly=False,bDefaultValueIsIgnored=False,bAdvancedView=False,)
-		CustomProperties Pin (PinId=CC9F75D84295F4E8568F03A68C65A446,PinName="self",PinFriendlyName=NSLOCTEXT("K2Node", "Target", "Target"),PinType.PinCategory="object",PinType.PinSubCategoryObject=BlueprintGeneratedClass'/Game/PhysAnimPawn.PhysAnimPawn_C',PinType.PinSubCategoryMemberReference=(),PinType.bIsArray=False,PinType.bIsReference=False,PinType.bIsConst=False,PinType.bIsWeakPointer=False,PersistentGuid=00000000000000000000000000000000,bHidden=True,bNotConnectable=False,bDefaultValueIsReadOnly=False,bDefaultValueIsIgnored=False,bAdvancedView=False,)
-	End Object
-	Begin Object Class=K2Node_VariableGet Name="K2Node_VariableGet_1"
-		VariableReference=(MemberName="SkeletalMeshComponent",bSelfContext=True)
-		NodePosX=928
-		NodePosY=64
-		NodeGuid=5E9CBBED44802CB75DD821972402DEFD
-		CustomProperties Pin (PinId=1EEABDD6494CBF38903C6691F5F5B45C,PinName="SkeletalMeshComponent",Direction="EGPD_Output",PinType.PinCategory="object",PinType.PinSubCategoryObject=Class'/Script/Engine.SkeletalMeshComponent',PinType.PinSubCategoryMemberReference=(),PinType.bIsArray=False,PinType.bIsReference=False,PinType.bIsConst=False,PinType.bIsWeakPointer=False,LinkedTo=(K2Node_CallFunction_8 4DD6FE2846E695E0F3F569B06CED0A46,),PersistentGuid=00000000000000000000000000000000,bHidden=False,bNotConnectable=False,bDefaultValueIsReadOnly=False,bDefaultValueIsIgnored=False,bAdvancedView=False,)
-		CustomProperties Pin (PinId=D68B011B4E4D641D541E8CA766972993,PinName="self",PinFriendlyName=NSLOCTEXT("K2Node", "Target", "Target"),PinType.PinCategory="object",PinType.PinSubCategoryObject=BlueprintGeneratedClass'/Game/PhysAnimPawn.PhysAnimPawn_C',PinType.PinSubCategoryMemberReference=(),PinType.bIsArray=False,PinType.bIsReference=False,PinType.bIsConst=False,PinType.bIsWeakPointer=False,PersistentGuid=00000000000000000000000000000000,bHidden=True,bNotConnectable=False,bDefaultValueIsReadOnly=False,bDefaultValueIsIgnored=False,bAdvancedView=False,)
-	End Object
-	Begin Object Class=K2Node_CallFunction Name="K2Node_CallFunction_8"
-		FunctionReference=(MemberParent=Class'/Script/Engine.SkeletalMeshComponent',MemberName="SetAllBodiesBelowSimulatePhysics")
-		NodePosX=1152
-		NodePosY=-16
-		NodeGuid=C58AFBDD484D31A196381BA69976AB10
-		CustomProperties Pin (PinId=399238E643A83FB3CFC85CB926B5CF9F,PinName="execute",PinToolTip="\nExec",PinType.PinCategory="exec",PinType.PinSubCategoryObject=None,PinType.PinSubCategoryMemberReference=(),PinType.bIsArray=False,PinType.bIsReference=False,PinType.bIsConst=False,PinType.bIsWeakPointer=False,LinkedTo=(K2Node_CallFunction_20 B38EF5724C795C379DDB809F78D77BA1,),PersistentGuid=00000000000000000000000000000000,bHidden=False,bNotConnectable=False,bDefaultValueIsReadOnly=False,bDefaultValueIsIgnored=False,bAdvancedView=False,)
-		CustomProperties Pin (PinId=C132D8014BC6A4DE84C437AF07CF838E,PinName="then",PinToolTip="\nExec",Direction="EGPD_Output",PinType.PinCategory="exec",PinType.PinSubCategoryObject=None,PinType.PinSubCategoryMemberReference=(),PinType.bIsArray=False,PinType.bIsReference=False,PinType.bIsConst=False,PinType.bIsWeakPointer=False,PersistentGuid=00000000000000000000000000000000,bHidden=False,bNotConnectable=False,bDefaultValueIsReadOnly=False,bDefaultValueIsIgnored=False,bAdvancedView=False,)
-		CustomProperties Pin (PinId=4DD6FE2846E695E0F3F569B06CED0A46,PinName="self",PinFriendlyName=NSLOCTEXT("K2Node", "Target", "Target"),PinToolTip="Target\nSkeletal Mesh Component Reference",PinType.PinCategory="object",PinType.PinSubCategoryObject=Class'/Script/Engine.SkeletalMeshComponent',PinType.PinSubCategoryMemberReference=(),PinType.bIsArray=False,PinType.bIsReference=False,PinType.bIsConst=False,PinType.bIsWeakPointer=False,LinkedTo=(K2Node_VariableGet_1 1EEABDD6494CBF38903C6691F5F5B45C,),PersistentGuid=00000000000000000000000000000000,bHidden=False,bNotConnectable=False,bDefaultValueIsReadOnly=False,bDefaultValueIsIgnored=False,bAdvancedView=False,)
-		CustomProperties Pin (PinId=7D29E35244DEB831C54B1CAEF4A5FE5C,PinName="InBoneName",PinToolTip="In Bone Name\nName (by ref)",PinType.PinCategory="name",PinType.PinSubCategoryObject=None,PinType.PinSubCategoryMemberReference=(),PinType.bIsArray=False,PinType.bIsReference=True,PinType.bIsConst=True,PinType.bIsWeakPointer=False,DefaultValue="None",AutogeneratedDefaultValue="None",LinkedTo=(K2Node_Knot_3 336A876240769F5A4170CF8126A5CF36,),PersistentGuid=00000000000000000000000000000000,bHidden=False,bNotConnectable=False,bDefaultValueIsReadOnly=False,bDefaultValueIsIgnored=True,bAdvancedView=False,)
-		CustomProperties Pin (PinId=AE4340174DD534FEA5B812AE6DAEA1EB,PinName="bNewSimulate",PinToolTip="New Simulate\nBoolean",PinType.PinCategory="bool",PinType.PinSubCategoryObject=None,PinType.PinSubCategoryMemberReference=(),PinType.bIsArray=False,PinType.bIsReference=False,PinType.bIsConst=False,PinType.bIsWeakPointer=False,DefaultValue="true",AutogeneratedDefaultValue="false",PersistentGuid=00000000000000000000000000000000,bHidden=False,bNotConnectable=False,bDefaultValueIsReadOnly=False,bDefaultValueIsIgnored=False,bAdvancedView=False,)
-		CustomProperties Pin (PinId=80EC7602465D797C7F0D4D93D6BA9BD5,PinName="bIncludeSelf",PinToolTip="Include Self\nBoolean",PinType.PinCategory="bool",PinType.PinSubCategoryObject=None,PinType.PinSubCategoryMemberReference=(),PinType.bIsArray=False,PinType.bIsReference=False,PinType.bIsConst=False,PinType.bIsWeakPointer=False,DefaultValue="true",AutogeneratedDefaultValue="true",PersistentGuid=00000000000000000000000000000000,bHidden=False,bNotConnectable=False,bDefaultValueIsReadOnly=False,bDefaultValueIsIgnored=False,bAdvancedView=False,)
-	End Object
-	Begin Object Class=K2Node_Knot Name="K2Node_Knot_3"
-		NodePosX=976
-		NodePosY=224
-		NodeGuid=55465C7448CB280BB4657AA36C94D451
-		CustomProperties Pin (PinId=D6618E7C42476ADDD628FF855FD6E7C9,PinName="InputPin",PinType.PinCategory="name",PinType.PinSubCategoryObject=None,PinType.PinSubCategoryMemberReference=(),PinType.bIsArray=False,PinType.bIsReference=False,PinType.bIsConst=False,PinType.bIsWeakPointer=False,LinkedTo=(K2Node_Knot_4 A0911D3740ABD40031008799EB009AE0,),PersistentGuid=00000000000000000000000000000000,bHidden=False,bNotConnectable=False,bDefaultValueIsReadOnly=False,bDefaultValueIsIgnored=True,bAdvancedView=False,)
-		CustomProperties Pin (PinId=336A876240769F5A4170CF8126A5CF36,PinName="OutputPin",Direction="EGPD_Output",PinType.PinCategory="name",PinType.PinSubCategoryObject=None,PinType.PinSubCategoryMemberReference=(),PinType.bIsArray=False,PinType.bIsReference=False,PinType.bIsConst=False,PinType.bIsWeakPointer=False,LinkedTo=(K2Node_CallFunction_8 7D29E35244DEB831C54B1CAEF4A5FE5C,),PersistentGuid=00000000000000000000000000000000,bHidden=False,bNotConnectable=False,bDefaultValueIsReadOnly=False,bDefaultValueIsIgnored=False,bAdvancedView=False,)
-	End Object
-	Begin Object Class=K2Node_Knot Name="K2Node_Knot_4"
-		NodePosX=608
-		NodePosY=224
-		NodeGuid=230B64A94C55B4E5636829B6DC0D21E7
-		CustomProperties Pin (PinId=4BF4C00042C37A0CDB25AD926D106B22,PinName="InputPin",PinType.PinCategory="name",PinType.PinSubCategoryObject=None,PinType.PinSubCategoryMemberReference=(),PinType.bIsArray=False,PinType.bIsReference=False,PinType.bIsConst=False,PinType.bIsWeakPointer=False,LinkedTo=(K2Node_VariableGet_37 62FD25D94D3A5CAE32C5E685EDA461B1,),PersistentGuid=00000000000000000000000000000000,bHidden=False,bNotConnectable=False,bDefaultValueIsReadOnly=False,bDefaultValueIsIgnored=True,bAdvancedView=False,)
-		CustomProperties Pin (PinId=A0911D3740ABD40031008799EB009AE0,PinName="OutputPin",Direction="EGPD_Output",PinType.PinCategory="name",PinType.PinSubCategoryObject=None,PinType.PinSubCategoryMemberReference=(),PinType.bIsArray=False,PinType.bIsReference=False,PinType.bIsConst=False,PinType.bIsWeakPointer=False,LinkedTo=(K2Node_Knot_3 D6618E7C42476ADDD628FF855FD6E7C9,),PersistentGuid=00000000000000000000000000000000,bHidden=False,bNotConnectable=False,bDefaultValueIsReadOnly=False,bDefaultValueIsIgnored=False,bAdvancedView=False,)
-	End Object
-	Begin Object Class=K2Node_CallFunction Name="K2Node_CallFunction_20"
-		FunctionReference=(MemberParent=Class'/Script/Engine.PhysicalAnimationComponent',MemberName="ApplyPhysicalAnimationProfileBelow")
-		NodePosX=624
-		NodePosY=-16
-		NodeGuid=D1A218604F6806C6A38DF0B1C6812215
-		CustomProperties Pin (PinId=B96C42C240335FB66D65AA84ADD838AC,PinName="execute",PinToolTip="\nExec",PinType.PinCategory="exec",PinType.PinSubCategoryObject=None,PinType.PinSubCategoryMemberReference=(),PinType.bIsArray=False,PinType.bIsReference=False,PinType.bIsConst=False,PinType.bIsWeakPointer=False,LinkedTo=(K2Node_CallFunction_0 BE952CF64A6E491E6C0B8F8B5CEB1000,),PersistentGuid=00000000000000000000000000000000,bHidden=False,bNotConnectable=False,bDefaultValueIsReadOnly=False,bDefaultValueIsIgnored=False,bAdvancedView=False,)
-		CustomProperties Pin (PinId=B38EF5724C795C379DDB809F78D77BA1,PinName="then",PinToolTip="\nExec",Direction="EGPD_Output",PinType.PinCategory="exec",PinType.PinSubCategoryObject=None,PinType.PinSubCategoryMemberReference=(),PinType.bIsArray=False,PinType.bIsReference=False,PinType.bIsConst=False,PinType.bIsWeakPointer=False,LinkedTo=(K2Node_CallFunction_8 399238E643A83FB3CFC85CB926B5CF9F,),PersistentGuid=00000000000000000000000000000000,bHidden=False,bNotConnectable=False,bDefaultValueIsReadOnly=False,bDefaultValueIsIgnored=False,bAdvancedView=False,)
-		CustomProperties Pin (PinId=96A091DE4DF2762BDFA90CB50CAE9BFB,PinName="self",PinFriendlyName=NSLOCTEXT("K2Node", "Target", "Target"),PinToolTip="Target\nPhysical Animation Component Reference",PinType.PinCategory="object",PinType.PinSubCategoryObject=Class'/Script/Engine.PhysicalAnimationComponent',PinType.PinSubCategoryMemberReference=(),PinType.bIsArray=False,PinType.bIsReference=False,PinType.bIsConst=False,PinType.bIsWeakPointer=False,LinkedTo=(K2Node_Knot_2 CC66C9584A8CC28A92B6D4AD7CD40472,),PersistentGuid=00000000000000000000000000000000,bHidden=False,bNotConnectable=False,bDefaultValueIsReadOnly=False,bDefaultValueIsIgnored=False,bAdvancedView=False,)
-		CustomProperties Pin (PinId=EB57682C40C2F631928F3D88FF0B2117,PinName="BodyName",PinToolTip="Body Name\nName\n\nThe body from which we\'d like to start applying the physical animation profile. Finds all bodies below in the skeleton hierarchy. None implies all bodies",PinType.PinCategory="name",PinType.PinSubCategoryObject=None,PinType.PinSubCategoryMemberReference=(),PinType.bIsArray=False,PinType.bIsReference=False,PinType.bIsConst=False,PinType.bIsWeakPointer=False,DefaultValue="None",AutogeneratedDefaultValue="None",LinkedTo=(K2Node_VariableGet_37 62FD25D94D3A5CAE32C5E685EDA461B1,),PersistentGuid=00000000000000000000000000000000,bHidden=False,bNotConnectable=False,bDefaultValueIsReadOnly=False,bDefaultValueIsIgnored=False,bAdvancedView=False,)
-		CustomProperties Pin (PinId=D9E9561844B25F69A79A2C9FE98806C6,PinName="ProfileName",PinToolTip="Profile Name\nName\n\nThe physical animation profile we\'d like to apply. For each body in the physics asset we search for physical animation settings with this name.",PinType.PinCategory="name",PinType.PinSubCategoryObject=None,PinType.PinSubCategoryMemberReference=(),PinType.bIsArray=False,PinType.bIsReference=False,PinType.bIsConst=False,PinType.bIsWeakPointer=False,DefaultValue="DocProfile",AutogeneratedDefaultValue="None",PersistentGuid=00000000000000000000000000000000,bHidden=False,bNotConnectable=False,bDefaultValueIsReadOnly=False,bDefaultValueIsIgnored=False,bAdvancedView=False,)
-		CustomProperties Pin (PinId=7D946B3045556970B2A7519A2614C109,PinName="bIncludeSelf",PinToolTip="Include Self\nBoolean\n\nWhether to include the provided body name in the list of bodies we act on (useful to ignore for cases where a root has multiple children)",PinType.PinCategory="bool",PinType.PinSubCategoryObject=None,PinType.PinSubCategoryMemberReference=(),PinType.bIsArray=False,PinType.bIsReference=False,PinType.bIsConst=False,PinType.bIsWeakPointer=False,DefaultValue="true",AutogeneratedDefaultValue="true",PersistentGuid=00000000000000000000000000000000,bHidden=False,bNotConnectable=False,bDefaultValueIsReadOnly=False,bDefaultValueIsIgnored=False,bAdvancedView=False,)
-		CustomProperties Pin (PinId=41EFD7554F080B09E40F7199FA2BAF6F,PinName="bClearNotFound",PinToolTip="Clear Not Found\nBoolean\n\nIf true, bodies without the given profile name will have any existing physical animation settings cleared. If false, bodies without the given profile name are left untouched.",PinType.PinCategory="bool",PinType.PinSubCategoryObject=None,PinType.PinSubCategoryMemberReference=(),PinType.bIsArray=False,PinType.bIsReference=False,PinType.bIsConst=False,PinType.bIsWeakPointer=False,DefaultValue="true",AutogeneratedDefaultValue="false",PersistentGuid=00000000000000000000000000000000,bHidden=False,bNotConnectable=False,bDefaultValueIsReadOnly=False,bDefaultValueIsIgnored=False,bAdvancedView=False,)
-	End Object
-[/VAR]
-
 ## Additional Resources
 
 [OBJECT:EmbeddedVideo]
diff --git a/Engine/Documentation/Source/Engine/Physics/PhAT/SkeletonTree/PhysicsAsssetSkeletonTree.INT.udn b/Engine/Documentation/Source/Engine/Physics/PhAT/SkeletonTree/PhysicsAsssetSkeletonTree.INT.udn
--- a/Engine/Documentation/Source/Engine/Physics/PhAT/SkeletonTree/PhysicsAsssetSkeletonTree.INT.udn
+++ b/Engine/Documentation/Source/Engine/Physics/PhAT/SkeletonTree/PhysicsAsssetSkeletonTree.INT.udn
@@ -117,7 +117,7 @@ When right-clicking on a Body, you will see the following context menu:
 | **Enable Collision All** | Enable collision between the currently selected bodies and all bodies. |
 | **Disable Collision** | Disable collision between the currently selected bodies. |
 | **Disable Collision All** | Disable collision between the currently selected bodies and all bodies. |
-| %spacer2% ||
+
 [/EXCERPT:collisionbody]
 
 -
@@ -172,7 +172,7 @@ When right-clicking on a Constraint, you will see the following context menu:
 | **Cycle Active Constraint** | Cycles whether each constraint axis is active in isolation. |
 | **Lock Swing** | Set swing 1 or 2 to be locked or limited. |
 | **Lock Twist** | Set twist to be locked or limited. |
-| %spacer2% ||
+
 [/EXCERPT:constraint]
 
 [EXCERPT:convert]
diff --git a/Engine/Documentation/Source/Engine/Physics/PhysicalMaterials/Reference/PhysMatReference.INT.udn b/Engine/Documentation/Source/Engine/Physics/PhysicalMaterials/Reference/PhysMatReference.INT.udn
--- a/Engine/Documentation/Source/Engine/Physics/PhysicalMaterials/Reference/PhysMatReference.INT.udn
+++ b/Engine/Documentation/Source/Engine/Physics/PhysicalMaterials/Reference/PhysMatReference.INT.udn
@@ -19,7 +19,7 @@ This category contains the core properties of the Physics Material: Friction, Re
 
 [REGION:raw]
 | Property | Description |
-| -------- | --- | :---------- |
+| -------- | :---------- |
 | **Friction** | This is the friction value of the surface, which controls how easily things can slide on this surface. |
 | **Friction Combine Mode** | [INCLUDE:#FcombineMode]  |
 | **Override Friction Combine Mode** | By default, the Friction Combine Mode is set to **Average**, by enabling this property you can change how Frictions are combined between touching Physical Materials.|
@@ -37,7 +37,7 @@ This category contains a property to change how scale affects the mass of the Ac
 
 [REGION:raw]
 | Property | Description |
-| -------- | --- | :---------- |
+| -------- | :---------- |
 | **Raise Mass To Power** | Used to adjust the way that mass increases as objects get larger. This is applied to the mass as calculated based on a "solid" object. In actuality, larger objects do not tend to be solid, and become more like "shells" (e.g. a car is not a solid piece of metal). Values are clamped to 1 or less. |
 [/REGION]
 
@@ -61,7 +61,7 @@ Gameplay related properties for the Physical Material.
 
 [REGION:raw]
 | Property | Description |
-| -------- | --- | :---------- |
+| -------- | :---------- |
 | **Surface Type** | [INCLUDE:#one]|
 [/REGION]
 
diff --git a/Engine/Documentation/Source/Engine/Physics/Substepping/Substepping.INT.udn b/Engine/Documentation/Source/Engine/Physics/Substepping/Substepping.INT.udn
--- a/Engine/Documentation/Source/Engine/Physics/Substepping/Substepping.INT.udn
+++ b/Engine/Documentation/Source/Engine/Physics/Substepping/Substepping.INT.udn
@@ -22,7 +22,7 @@ This feature is still a work in progress, and is not fully supported by APEX Des
 [REGION:raw]
 | Property | Description |
 | --- | :--- |
-| **Substepping** || This enables or disables sub-stepping in your current project. |
+| **Substepping** | This enables or disables sub-stepping in your current project. |
 | **Substepping Async** | Whether to substep the async phsyics simulation. |
 | **Max Substep Delta Time** | This is the maximum time, in seconds a sub-step is allowed to take. So, if a full step takes .05 seconds, and **Max Substep Delta Time** is set to .025, the full step will be split into 2 sub-steps. Now, if the full step were to take a substantial amount of time, say 2 full seconds, then the step will be evenly split based on the value in **Max Substeps**, and not **80** `(2.0/0.025=80)` sub-steps. It is worth noting that **Max Physics Delta Time** limits how long a physics step is permitted to take. |
 | **Max Substeps** | The maximum number of sub-steps a full step is permitted to be broken into. |
diff --git a/Engine/Documentation/Source/Engine/Physics/Vehicles/VehicleUserGuide/5/TestingFunctionality.INT.udn b/Engine/Documentation/Source/Engine/Physics/Vehicles/VehicleUserGuide/5/TestingFunctionality.INT.udn
--- a/Engine/Documentation/Source/Engine/Physics/Vehicles/VehicleUserGuide/5/TestingFunctionality.INT.udn
+++ b/Engine/Documentation/Source/Engine/Physics/Vehicles/VehicleUserGuide/5/TestingFunctionality.INT.udn
@@ -24,9 +24,7 @@ the vehicle. If you started with an existing template, some (or all) of these ax
 
 	![](input.png)(w:950)
 	
-1. In the Bindings menu, set up the following controls (if you do not currently have them listed). If you do have them listed, it would be wise to double-check that they are set up like the ones listed here, so that you know your vehicle will work correctly. [COMMENT:none]
- TODO, add images for this steps and all of the sub-steps contained herein. 
-[/COMMENT]
+1. In the Bindings menu, set up the following controls (if you do not currently have them listed). If you do have them listed, it would be wise to double-check that they are set up like the ones listed here, so that you know your vehicle will work correctly. 
 
 	1. First, we'll setup an **Action Mappings** by clicking the **+** sign next to the Action Mappings property.
 
diff --git a/Engine/Documentation/Source/Engine/Rendering/LightingAndShadows/AmbientCubemap/AmbientCubemap.INT.udn b/Engine/Documentation/Source/Engine/Rendering/LightingAndShadows/AmbientCubemap/AmbientCubemap.INT.udn
--- a/Engine/Documentation/Source/Engine/Rendering/LightingAndShadows/AmbientCubemap/AmbientCubemap.INT.udn
+++ b/Engine/Documentation/Source/Engine/Rendering/LightingAndShadows/AmbientCubemap/AmbientCubemap.INT.udn
@@ -6,15 +6,7 @@ Version: 4.9
 topic-image:cubemap2.png
 
 [excerpt:AmbientCubemap]
-[EXCERPT:AmbientCubemap_Desc]
-The **Ambient Cubemap** lights the scene from a provided image. The image is mapped to a sphere in the far distance (implemented as a cubemap texture with the mip maps storing preblurred versions of the image).
-The blurry versions are computed in a way that they can be used for specular lighting with varying glossiness (sharp vs. blurry reflections) and they can be used for diffuse lighting as well. 
-This effect is independent of the position a material is lit from. Viewer position, material roughness (for specular effects), and the material surface normal are all taken into account. 
-
-[REGION:note]
-The Ambient Cubemap has very limited local shadowing (SSAO only). Use a [Sky Light](Engine/Rendering/LightingAndShadows/LightTypes/SkyLight) instead when representing the sky's lighting.
-[/REGION]
-[/EXCERPT:AmbientCubemap_Desc]
+[INCLUDE:#AmbientCubemap_Desc]
 
 The Ambient Cubemap is only intended to be used for a subtle directional ambient term (useful in games with completely dynamic lighting), or for model viewer applications.
 
@@ -36,3 +28,15 @@ In this image, a cubemap image has been automatically laid out into a "longlat"
 
 The cubemap in this image is derived from a light probe by [REDACTED_PERSON]. [http://gl.ict.usc.edu/Data/HighResProbes/](http://gl.ict.usc.edu/Data/HighResProbes/)
 [/excerpt:AmbientCubemap]
+
+[COMMENT:none]
+[EXCERPT:AmbientCubemap_Desc]
+The **Ambient Cubemap** lights the scene from a provided image. The image is mapped to a sphere in the far distance (implemented as a cubemap texture with the mip maps storing preblurred versions of the image).
+The blurry versions are computed in a way that they can be used for specular lighting with varying glossiness (sharp vs. blurry reflections) and they can be used for diffuse lighting as well. 
+This effect is independent of the position a material is lit from. Viewer position, material roughness (for specular effects), and the material surface normal are all taken into account. 
+
+[REGION:note]
+The Ambient Cubemap has very limited local shadowing (SSAO only). Use a [Sky Light](Engine/Rendering/LightingAndShadows/LightTypes/SkyLight) instead when representing the sky's lighting.
+[/REGION]
+[/EXCERPT:AmbientCubemap_Desc]
+[/COMMENT]
\ No newline at end of file
diff --git a/Engine/Documentation/Source/Engine/Rendering/LightingAndShadows/AmbientOcclusion/AmbientOcclusion.INT.udn b/Engine/Documentation/Source/Engine/Rendering/LightingAndShadows/AmbientOcclusion/AmbientOcclusion.INT.udn
--- a/Engine/Documentation/Source/Engine/Rendering/LightingAndShadows/AmbientOcclusion/AmbientOcclusion.INT.udn
+++ b/Engine/Documentation/Source/Engine/Rendering/LightingAndShadows/AmbientOcclusion/AmbientOcclusion.INT.udn
@@ -30,9 +30,7 @@ As the AO is part of the GBuffer, it also can be output by the material. The SSA
 | **Radius in WorldSpace** | If enabled, the value of the **AO Radius** property is assumed to be in world-space. Otherwise, it is assumed to be in view-space.  |
 | **Distance** | A tweak parameter that is needed to avoid dark halos around objects. As the SSAO technique only has the zbuffer, it cannot know about the thickness of objects. [INCLUDE:Engine/Rendering/LightingAndShadows/AmbientOcclusion#Distance] |
 | **Power** | Controls the strength of the darkening effect that is computed. [INCLUDE:Engine/Rendering/LightingAndShadows/AmbientOcclusion#Power] |
-| **Bias** | Allows to tweak AO details. Some bias is needed to avoid precision artifacts in the distance. >0, in unreal units, default (3.0) works well for flat surfaces but can reduce details. [COMMENT:none]
-TODO: images
-[/COMMENT] |
+| **Bias** | Allows to tweak AO details. Some bias is needed to avoid precision artifacts in the distance. >0, in unreal units, default (3.0) works well for flat surfaces but can reduce details. |
 
 [COMMENT:none]
 
diff --git a/Engine/Documentation/Source/Engine/Rendering/LightingAndShadows/DistanceFieldAmbientOcclusion/DistanceFieldAmbientOcclusion.INT.udn b/Engine/Documentation/Source/Engine/Rendering/LightingAndShadows/DistanceFieldAmbientOcclusion/DistanceFieldAmbientOcclusion.INT.udn
--- a/Engine/Documentation/Source/Engine/Rendering/LightingAndShadows/DistanceFieldAmbientOcclusion/DistanceFieldAmbientOcclusion.INT.udn
+++ b/Engine/Documentation/Source/Engine/Rendering/LightingAndShadows/DistanceFieldAmbientOcclusion/DistanceFieldAmbientOcclusion.INT.udn
@@ -272,9 +272,7 @@ With Unreal Engine 4.16, these are some improvements for DFAO with test results
 | Optimization | Orignal Cost (ms) | New Cost (ms) | Savings (ms) |
 | ------------ | ------------ | -------- | ------- |
 | The culling algorithm has changed to produce a list of intersecting screen tiles for each object, instead of the other way around. Each tile/object intersection gets its cone tracing thread group, so wavefronts are much smaller and scheduled better.  | 3.63 | 3.48 | 0.15 |
-| Slow instructions in the inner loop have been replaced with fast approximations [COMMENT:none]
--(exp2 -> sqr + 1, rcpFast, lengthFast)-
-[/COMMENT]. | 3.25 | 3.09 | 0.16 |
+| Slow instructions in the inner loop have been replaced with fast approximations | 3.25 | 3.09 | 0.16 |
 | The transform was moved out of the inner loop from world to local space (sample position constructed from local space position + direction). | 3.09 | 3.04 | 0.05 |
 | Compute shader for ClearUAV.  | 3.04 | 2.62 | 0.42 |
 [/REGION]
@@ -299,7 +297,6 @@ the history faster. Overall, this should help reduce ghosting that can happen wh
 
 
 [COMMENT:none]
---
 
 [INCLUDE:#comp2]
 
@@ -329,15 +326,11 @@ the history faster. Overall, this should help reduce ghosting that can happen wh
 [/REGION]
 [/EXCERPT:comp2]
 
--
 [/COMMENT]
 
 
 #### Global Distance Fields
 
-[COMMENT:none]
--- EXCERPT taken from MeshDistanceFields.INT.udn --
-[/COMMENT]
 [INCLUDE:Engine/Rendering/LightingAndShadows/MeshDistanceFields#visualizeglobaldf1]
 
 The clipmap is separated into four slices around the camera that are only updated where needed, which can happen when camera movement reveals a new slice or an object is never moved causing 
@@ -348,9 +341,6 @@ a teleport.
 [REGION:caption] Clipmap visualization of the texel size of the clipmaps with each clipmap being represented by a different color. [/REGION]
 
 
-[COMMENT:none]
--- EXCERPT taken from MeshDistanceFields.INT.udn --
-[/COMMENT]
 [INCLUDE:Engine/Rendering/LightingAndShadows/MeshDistanceFields#visualizeglobaldf2]
 
 The Global Distance Field is a lower resolution than its object Distance Fields counterpart, making Global Distance Fields inaccurate near surfaces. When cone traces happen, the object Distance Field
diff --git a/Engine/Documentation/Source/Engine/Rendering/LightingAndShadows/LightTypes/Spot/Spot.INT.udn b/Engine/Documentation/Source/Engine/Rendering/LightingAndShadows/LightTypes/Spot/Spot.INT.udn
--- a/Engine/Documentation/Source/Engine/Rendering/LightingAndShadows/LightTypes/Spot/Spot.INT.udn
+++ b/Engine/Documentation/Source/Engine/Rendering/LightingAndShadows/LightTypes/Spot/Spot.INT.udn
@@ -17,11 +17,10 @@ A **Spot Light** emits light from a single point in a cone shape. Users are give
 
 Like other lights, a Spot Light can be set to one of three Mobility settings: 
 
-
-
 * **Static** - (pictured left) which means that the light cannot be changed in game. This is the fastest method for rendering and allows for baked lighting.
 * **Stationary** - (also pictured left) which means the light will only have its shadowing and bounced lighting from static geometry baked by Lightmass, all other lighting will be dynamic. This setting also allows for the light to change color and intensity in game, but, it does not move and allows partial baked lighting. 
 * **Moveable** - (pictured left) which means the light is totally dynamic and allows for dynamic shadowing. This is the slowest in terms of rendering but allows for the most flexibility during gameplay. 
+[/EXCERPT:Intro]
 
 Two examples of a **Spot Light** placed inside a level are shown below demonstrating how to determine the light radius and cone effectors:
 
diff --git a/Engine/Documentation/Source/Engine/Rendering/LightingAndShadows/MeshDistanceFields/MeshDistanceFields.INT.udn b/Engine/Documentation/Source/Engine/Rendering/LightingAndShadows/MeshDistanceFields/MeshDistanceFields.INT.udn
--- a/Engine/Documentation/Source/Engine/Rendering/LightingAndShadows/MeshDistanceFields/MeshDistanceFields.INT.udn
+++ b/Engine/Documentation/Source/Engine/Rendering/LightingAndShadows/MeshDistanceFields/MeshDistanceFields.INT.udn
@@ -207,7 +207,7 @@ the opaque one and will come at a higher cost.
 [/EXCERPT:twosided]
 [/COMMENT]
 
-##### Folaige Tool Settings
+##### Foliage Tool Settings
 
 In the [Foliage Tool](Engine/Foliage), you must enable each foliage type that you want to use Distance Field lighting features for ambient occlusion and shadowing. By default, this setting is disabled because
 some foliage assets (like grass) where you would have thousands or tens of thousands of instances would overflow the tile culling buffer. When this happens, you can unsightly artifacts. So, only enable 
diff --git a/Engine/Documentation/Source/Engine/Rendering/LightingAndShadows/MeshDistanceFields/Reference/DistanceFieldsReference.INT.udn b/Engine/Documentation/Source/Engine/Rendering/LightingAndShadows/MeshDistanceFields/Reference/DistanceFieldsReference.INT.udn
--- a/Engine/Documentation/Source/Engine/Rendering/LightingAndShadows/MeshDistanceFields/Reference/DistanceFieldsReference.INT.udn
+++ b/Engine/Documentation/Source/Engine/Rendering/LightingAndShadows/MeshDistanceFields/Reference/DistanceFieldsReference.INT.udn
@@ -397,9 +397,7 @@ be enabled for the entire project. You can also control the self-shadowing that
 [REGION:raw]
 | Property | Description |
 | -------- | ----------- |
-| **Generate Mesh Distance Field** | Whether to generate a Distance Field for this mesh, which can be used with [Distance Field Indirect Shadows](Engine/Rendering/LightingAndShadows/MeshDistanceFields). This is ignored if the Project Setting for **Generate Mesh Distance Fields** is enabled.  | [COMMENT:none]
-TODO FINISH THE LINK 
-[/COMMENT]
+| **Generate Mesh Distance Field** | Whether to generate a Distance Field for this mesh, which can be used with [Distance Field Indirect Shadows](Engine/Rendering/LightingAndShadows/MeshDistanceFields). This is ignored if the Project Setting for **Generate Mesh Distance Fields** is enabled.  | 
 | **Distance Field Self Shadow Bias** | Used for reducing self-shadowing from Distance Field methods when using World Position Offset to animate the mesh' vertices. |
 [/REGION]
 
diff --git a/Engine/Documentation/Source/Engine/Rendering/Materials/HowTo/MeshDecals/MeshDecalsHowTo.INT.udn b/Engine/Documentation/Source/Engine/Rendering/Materials/HowTo/MeshDecals/MeshDecalsHowTo.INT.udn
--- a/Engine/Documentation/Source/Engine/Rendering/Materials/HowTo/MeshDecals/MeshDecalsHowTo.INT.udn
+++ b/Engine/Documentation/Source/Engine/Rendering/Materials/HowTo/MeshDecals/MeshDecalsHowTo.INT.udn
@@ -10,10 +10,7 @@ Order:
 SkillLevel: Intermediate
 tags:Materials
 
-
-<div align= "center">
 ![](MeshDecalVsMasked.png)(w:700)
-</div>
 
 [TOC (start:2 end:2)]
 
diff --git a/Engine/Documentation/Source/Engine/Rendering/Materials/MaterialInputs/MaterialInputs.INT.udn b/Engine/Documentation/Source/Engine/Rendering/Materials/MaterialInputs/MaterialInputs.INT.udn
--- a/Engine/Documentation/Source/Engine/Rendering/Materials/MaterialInputs/MaterialInputs.INT.udn
+++ b/Engine/Documentation/Source/Engine/Rendering/Materials/MaterialInputs/MaterialInputs.INT.udn
@@ -51,13 +51,11 @@ Fortunately, Unreal Engine 4 removes the guess work for which inputs you need fo
 
 [INCLUDE:Engine/Rendering/Materials\PhysicallyBased#RoughnessMap]
 
-[EXCERPT:lightbox]
 **Click for full size**
 
 [REGION:lightbox]
 [![](Engine/Rendering/Materials\PhysicallyBased/RoughBias_12.png)(w:935)](Engine/Rendering/Materials\PhysicallyBased/RoughBias_12.png)
 [/REGION]
-[/EXCERPT:lightbox]
 
 ## Emissive Color
 
@@ -66,13 +64,11 @@ The **Emissive Color** input controls which parts of your Material will appear t
 
 Values greater than 1 are allowed as HDR lighting is supported.
 
-[EXCERPT:lightbox]
 **Click for full size**
 
 [REGION:lightbox]
 [![](Emissive_21.png)(w:935)](Emissive_21.png)
 [/REGION]
-[/EXCERPT:lightbox]
 [/EXCERPT:Emissive]
 
 ## Opacity
@@ -85,13 +81,11 @@ The **Opacity** input is used when using the [Translucent Blend Mode](Engine/Ren
 
 Opaque and masked blend modes also use Opacity when using one of the subsurface shading models.
 
-[EXCERPT:lightbox]
 **Click for full size**
 
 [REGION:lightbox]
 [![](TranslucencyImage.png)(w:935)](TranslucencyImage.png)
 [/REGION]
-[/EXCERPT:lightbox]
 
 Opacity is mostly used for _Translucent_, _Additive_, and _Modulated Materials_.
 [/EXCERPT:Opacity]
@@ -104,13 +98,11 @@ Opacity is mostly used for _Translucent_, _Additive_, and _Modulated Materials_.
 
 You can use the **Opacity Mask Clip Value** property on the base Material node to control the point where clipping takes place. See the [Masked Blend Mode documentation](Engine/Rendering/Materials\MaterialProperties\BlendModes#Masked) for more details.
 
-[EXCERPT:lightbox]
 **Click for full size**
 
 [REGION:lightbox]
 [![](MaskedChainlink.png)(w:935)](Engine/Rendering/Materials\MaterialInputs/MaskedChainlink.png)
 [/REGION]
-[/EXCERPT:lightbox]
 
 [/EXCERPT:OpacityMask]
 
diff --git a/Engine/Documentation/Source/Engine/Rendering/Materials/MaterialProperties/BlendModes/BlendModes.INT.udn b/Engine/Documentation/Source/Engine/Rendering/Materials/MaterialProperties/BlendModes/BlendModes.INT.udn
--- a/Engine/Documentation/Source/Engine/Rendering/Materials/MaterialProperties/BlendModes/BlendModes.INT.udn
+++ b/Engine/Documentation/Source/Engine/Rendering/Materials/MaterialProperties/BlendModes/BlendModes.INT.udn
@@ -130,8 +130,16 @@ In the graph above, the material has been set to a Blend Mode of _Modulate_ with
 The Modulate blend mode is best suited for certain particle effects. But care must be taken as it does not support lighting or Separate Translucency.
 [/REGION]
 
-<!------------------- 
+| Mode | Description |
+| ---- | ----------- |
+| **_BLEND_Opaque_** | Final color = Source color.  This means that the Material will draw on top of the background. This blend mode is compatible with lighting.
+| **_BLEND_Masked_** | Final color = Source color if OpacityMask > OpacityMaskClipValue, otherwise the pixel is discarded.  This blend mode is compatible with lighting.
+| **_BLEND_Translucent_** | Final color = Source color * Opacity + Dest color * (1 - Opacity).  This blend mode is **NOT** compatible with dynamic lighting.
+| **_BLEND_Additive_** | Final color = Source color + Dest color.  This blend mode is **NOT** compatible with dynamic lighting.
+| **_BLEND_Modulate_** | Final color = Source color x Dest color.  This blend mode is **NOT** compatible with dynamic lighting, or fog, unless this is a decal material.
+[/EXCERPT:BlendModes]
 
+[COMMENT:none]
 [EXCERPT:OpacityMaskClipValue]
 [OBJECT:SliderSequence]
 	[PARAMLITERAL:max]
@@ -169,16 +177,4 @@ The Modulate blend mode is best suited for certain particle effects. But care mu
 	[/PARAM]
 [/OBJECT]
 [/EXCERPT:OpacityMaskClipValue]
-
-
-
-
-| Mode | Description |
-| ---- | ----------- |
-| **_BLEND_Opaque_** | Final color = Source color.  This means that the Material will draw on top of the background. This blend mode is compatible with lighting.
-| **_BLEND_Masked_** | Final color = Source color if OpacityMask > OpacityMaskClipValue, otherwise the pixel is discarded.  This blend mode is compatible with lighting.
-| **_BLEND_Translucent_** | Final color = Source color * Opacity + Dest color * (1 - Opacity).  This blend mode is **NOT** compatible with dynamic lighting.
-| **_BLEND_Additive_** | Final color = Source color + Dest color.  This blend mode is **NOT** compatible with dynamic lighting.
-| **_BLEND_Modulate_** | Final color = Source color x Dest color.  This blend mode is **NOT** compatible with dynamic lighting, or fog, unless this is a decal material.
-[/EXCERPT:BlendModes]
-
+[/COMMENT]
\ No newline at end of file
diff --git a/Engine/Documentation/Source/Engine/Rendering/Materials/MaterialProperties/MaterialProperties.INT.udn b/Engine/Documentation/Source/Engine/Rendering/Materials/MaterialProperties/MaterialProperties.INT.udn
--- a/Engine/Documentation/Source/Engine/Rendering/Materials/MaterialProperties/MaterialProperties.INT.udn
+++ b/Engine/Documentation/Source/Engine/Rendering/Materials/MaterialProperties/MaterialProperties.INT.udn
@@ -15,28 +15,11 @@ tags:Materials
 
 This document breaks down the properties available on a Material. These properties are accessible by selecting the Main Material node in the Material Editor.
 
-[COMMENT:none]
--, 
-or by **Right-Clicking** the Material in the **Content Browser** then selecting **Asset Actions** then [Property Matrix](Engine/UI/PropertyMatrix).
--
-[/COMMENT]
 
 * Properties inside Material:
 
 ![](PropertiesPanelMaterialEditor_Mat.png)(h:445) 
 
-[COMMENT:none]
--
-* Accessing the Property Matrix:
-
-![](PropertiesContextMenu_Mat.png)(h:445) 
-
-* The Property Matrix:
-
-![](PropertyMatrix.png)(h:445)
--
-[/COMMENT]
-
 
 The base material node has a number of properties which affect how the material behaves. These properties are described below, with each category divided into its own document section, and presented in the order in which they appear in the **Properties** panel.
 
@@ -83,9 +66,7 @@ The base material node has a number of properties which affect how the material
 | **Responsive AA (Ant aliasing)** | Small moving objects, especially particles, can sometimes get obliterated by ant aliasing; by setting this property to _true_, a different AA algorithm is used which provides more definition. Put another way, if you create a blizzard or similar particle system and you feel like you cannot really see your snowflakes, turn this on - it will help. However, this should _only_ be used on small moving objects, because it creates aliasing artifacts on the background. |
 | **Translucency Lighting Mode** | [INCLUDE:#TransLightingMode]  See [LitTranslucency](Engine/Rendering/LightingAndShadows/LitTranslucency). |
 | **Translucency Directional Lighting Intensity** | Useful for artificially increasing the influence of the normal on the lighting result for translucency. A value larger than 1 increases the influence for the normal, a value smaller than 1 makes the lighting more ambient. |
-| **Disable Depth Test** |  Allows the Material to disable depth tests, which is only meaningful with translucent blend modes.  Disabling depth tests will make rendering significantly slower since no occluded pixels can get Z-culled.[COMMENT:none]
-- In short, you'll be able to-
-[/COMMENT]  |
+| **Disable Depth Test** |  Allows the Material to disable depth tests, which is only meaningful with translucent blend modes.  Disabling depth tests will make rendering significantly slower since no occluded pixels can get Z-culled. |
 | **Use Translucency Vertex Fog** | When _true_, translucent materials are fogged. Defaults to _true_. |
 
 
diff --git a/Engine/Documentation/Source/Engine/Rendering/Materials/MeshDecals/MeshDecals.INT.udn b/Engine/Documentation/Source/Engine/Rendering/Materials/MeshDecals/MeshDecals.INT.udn
--- a/Engine/Documentation/Source/Engine/Rendering/Materials/MeshDecals/MeshDecals.INT.udn
+++ b/Engine/Documentation/Source/Engine/Rendering/Materials/MeshDecals/MeshDecals.INT.udn
@@ -19,14 +19,6 @@ details that shear and distort when not aligned with the surface it’s projecti
 simple projection and instead can be used with geometry that wraps around edges, being used with spline meshes, ultimately enhancing the 
 look of your characters. 
 
-[COMMENT:none]
--
-<div align= "center">
-![](MeshDecalVsMasked.png)(w:700)
-</div>
--
-[/COMMENT]
-
 ## Mesh Decal versus Masked Material
 
 Mesh Decals are a mix of translucent blend mode and deferred decals in that they do not render to depth but do update the GBuffer or DBuffer 
@@ -56,14 +48,11 @@ needs only to stay in front of the surface you want to affect. With this in mind
 because there is already a small depth bias applied to account for this small offset. Also, using a bit of geometry that feathers out from the decal can be good to create the soft transitions
 for the decal, which is something you cannot attain with a masked material. 
 
-
-<div align= "left">
 [REGION:imagetable]
 | ![](MeshDecalAuthoring.png)(h:450) | ![](MeshDecalAuthoringComposite.png)(h:450) |
 | --- | --- |
 | Base Mesh and Decal Geometry (Isolated) | Composited Mesh |
 [/REGION]
-</div>
 
 Another thing to keep in mind when developing your content, is that Mesh Decals can be a challenge with Levels of Detail (LOD), or long view distances when there is limited depth buffer
 precision, because geometry will intersect (or not match-up) as intended. However, you can either change the mesh to account for this (in most cases), or you can use **World Position Offset** in 
diff --git a/Engine/Documentation/Source/Engine/Rendering/PostProcessEffects/DepthOfField/DepthOfField.INT.udn b/Engine/Documentation/Source/Engine/Rendering/PostProcessEffects/DepthOfField/DepthOfField.INT.udn
--- a/Engine/Documentation/Source/Engine/Rendering/PostProcessEffects/DepthOfField/DepthOfField.INT.udn
+++ b/Engine/Documentation/Source/Engine/Rendering/PostProcessEffects/DepthOfField/DepthOfField.INT.udn
@@ -106,9 +106,7 @@ Note that the circle size might not always match to the actual effect as the imp
 | **Max Bokeh Size** | The maximum size (in percent of the view width) of the blur for Bokeh-based DoF (**Note:** performance cost scales with size*size). An approximate representation of the function used to calculate the size of the Bokeh at any distance is shown below: [INCLUDE:#MaxBokehSize] |
 | **Near Blur Size** | The maximum size (in percent of the view width) of the near blur for Gaussian-based DoF. (**Note:** performance cost scales with size) [INCLUDE:#NearBlurSize] |
 | **Far Blur Size** | The maximum size (in percent of the view width) of the far blur for Gaussian-based DoF. (**Note:** performance cost scales with size) [INCLUDE:#FarBlurSize] |
-| **Shape** | A texture that defines the shape of the Bokeh when objects get out of focus (cannot be blended). [COMMENT:none]
-TODO: show a few example images
-[/COMMENT] |
+| **Shape** | A texture that defines the shape of the Bokeh when objects get out of focus (cannot be blended). |
 | **Occlusion** | Controls how much the blurred geometry extends beyond its usual silhouette and its opacity. A value of 0.18 gives a fairly natural occlusion result. A value of 0.4 may be necessary to solve layer color leaking issues. Very small values (less than 0.18) generally result in negating the blur effect, though can work well when objects are very near to the camera. [INCLUDE:#Occlusion] |
 | **Color Threshold** | Controls the threshold where adaptive DoF switches to using full-resolution based on color. Smaller values cause more of the scene to use full-resolution processing. |
 | **Size Threshold** | Controls the threshold where adaptive DoF switches to using full-resolution based on size. Larger values cause more of the scene to use full-resolution processing. |
diff --git a/Engine/Documentation/Source/Engine/UI/Editor.INT.udn b/Engine/Documentation/Source/Engine/UI/Editor.INT.udn
--- a/Engine/Documentation/Source/Engine/UI/Editor.INT.udn
+++ b/Engine/Documentation/Source/Engine/UI/Editor.INT.udn
@@ -92,6 +92,6 @@ mac
 
 [DIR(output:"playlist" type:"playlist" playlist_id:"PLZlv_N0_O1gasd4IcOe9Cx9wHoBB7rxFl")]
 
-
+[/EXCERPT:LevelEditor]
 
 
diff --git a/Engine/Documentation/Source/Engine/UI/InterfaceOverview/InterfaceOverview.INT.udn b/Engine/Documentation/Source/Engine/UI/InterfaceOverview/InterfaceOverview.INT.udn
--- a/Engine/Documentation/Source/Engine/UI/InterfaceOverview/InterfaceOverview.INT.udn
+++ b/Engine/Documentation/Source/Engine/UI/InterfaceOverview/InterfaceOverview.INT.udn
@@ -4,7 +4,7 @@ Crumbs: %ROOT%, Engine, Engine/UI
 Description:A look at how to customize the look and feel of the user interface by rearranging and docking various panels and tabs.
 Navigation:topic
 Version: 4.9
-topic-image:(Engine/UI/CustomizingUI.png
+topic-image:Engine/UI/CustomizingUI.png
 
 
 [REGION:banner]
diff --git a/Engine/Documentation/Source/Engine/UI/LevelEditor/InEditorTesting/Settings/PIE_SIE_Settings.INT.udn b/Engine/Documentation/Source/Engine/UI/LevelEditor/InEditorTesting/Settings/PIE_SIE_Settings.INT.udn
--- a/Engine/Documentation/Source/Engine/UI/LevelEditor/InEditorTesting/Settings/PIE_SIE_Settings.INT.udn
+++ b/Engine/Documentation/Source/Engine/UI/LevelEditor/InEditorTesting/Settings/PIE_SIE_Settings.INT.udn
@@ -146,6 +146,7 @@ from the **Play In** dropdown menu, there are settings for URL parameters, gamep
 | Run Dedicated Server |[INCLUDE:#PlayNetDedicated] |
 | Route 1st Gamepad to 2nd Client | [INCLUDE:#RouteGamepadToSecondWindow]|
 | Use Single Process |[INCLUDE:#RunUnderOneProcess] |
+[/EXCERPT:NetworkPIESettings]
 
 [COMMENT:none]
 
@@ -173,6 +174,4 @@ from the **Play In** dropdown menu, there are settings for URL parameters, gamep
 [/EXCERPT:RunUnderOneProcess]
 
 
-[/COMMENT]
-
-[/EXCERPT:NetworkPIESettings]
+[/COMMENT]
\ No newline at end of file
diff --git a/Engine/Documentation/Source/Engine/UI/LevelEditor/MenuBar/MainMenuBar.INT.udn b/Engine/Documentation/Source/Engine/UI/LevelEditor/MenuBar/MainMenuBar.INT.udn
--- a/Engine/Documentation/Source/Engine/UI/LevelEditor/MenuBar/MainMenuBar.INT.udn
+++ b/Engine/Documentation/Source/Engine/UI/LevelEditor/MenuBar/MainMenuBar.INT.udn
@@ -116,7 +116,7 @@ mac
 %commonwindowmenu%
 
 [VAR:CommonWindowMenu]
-| [REGION:tablesection]** **[/REGION] ||
+| [REGION:tablesection] [/REGION] ||
 | **Content Browser** | Opens any of up to 4 separate [**Content Browsers**](Engine/Content/Browser) for working with content assets. |
 | **Developer Tools** | Provides access to various **[](Engine\Tools\DevTools)** that are useful for programmers working with the engine and editor. |
 | **Blueprint Debugger** | Opens the **[](Engine/Blueprints/UserGuide/Debugging)** panel. |
diff --git a/Engine/Documentation/Source/Engine/UMG/HowTo/UMGHowTo.INT.udn b/Engine/Documentation/Source/Engine/UMG/HowTo/UMGHowTo.INT.udn
--- a/Engine/Documentation/Source/Engine/UMG/HowTo/UMGHowTo.INT.udn
+++ b/Engine/Documentation/Source/Engine/UMG/HowTo/UMGHowTo.INT.udn
@@ -11,23 +11,16 @@ The **Unreal Motion Graphics UI Designer (UMG) - How To...** pages will show you
 Refer to the table below for more information and click each link to view its contents. 
 
 ## Beginner
-[REGION:wrapper transparent]
-[](FANCY:Engine\UMG\HowTo\CreatingWidgets "%Engine\UMG\HowTo\CreatingWidgets:description%")
-[/REGION]
+
+[DIR(output:"fancy" parent:"Engine/UMG/HowTo" org:"hierarchy" skill_Level:"Beginner")]
 
 ## Intermediate
-[REGION:wrapper transparent]
-[](FANCY:Engine\UMG\HowTo\WidgetTemplates "%Engine\UMG\HowTo\WidgetTemplates:description%")
-[](FANCY:Engine\UMG\HowTo\Create3DWidgets "%Engine\UMG\HowTo\Create3DWidgets:description%")
-[](FANCY:Engine\UMG\HowTo\DragAndDrop "%Engine\UMG\HowTo\DragAndDrop:description%")
-[](FANCY:Engine\UMG\HowTo\ScalingUI "%Engine\UMG\HowTo\ScalingUI:description%")
-[/REGION]
+
+[DIR(output:"fancy" parent:"Engine/UMG/HowTo" org:"hierarchy" skill_Level:"Intermediate")]
 
 ## Advanced
-[REGION:wrapper transparent]
-[](FANCY:Engine\UMG\HowTo\EventBasedUI "%Engine\UMG\HowTo\EventBasedUI:description%")
-[](FANCY:Engine\UMG\HowTo\VirtualKeyboards "%Engine\UMG\HowTo\VirtualKeyboards:description%")
-[/REGION]
+
+[DIR(output:"fancy" parent:"Engine/UMG/HowTo" org:"hierarchy" skill_Level:"Advanced")]
 
 
 
diff --git a/Engine/Documentation/Source/Enterprise/Automotive/Automotive.INT.udn b/Engine/Documentation/Source/Enterprise/Automotive/Automotive.INT.udn
--- a/Engine/Documentation/Source/Enterprise/Automotive/Automotive.INT.udn
+++ b/Engine/Documentation/Source/Enterprise/Automotive/Automotive.INT.udn
@@ -1,4 +1,4 @@
-Availability:Public
+Availability:NoPublish
 Title: Automotive Applications
 Description:Using Unreal Engine for automotive visualization applications
 Crumbs:%ROOT%,
diff --git a/Engine/Documentation/Source/Enterprise/Automotive/AutomotiveMaterialPack/AutomotiveMaterialPack.INT.udn b/Engine/Documentation/Source/Enterprise/Automotive/AutomotiveMaterialPack/AutomotiveMaterialPack.INT.udn
--- a/Engine/Documentation/Source/Enterprise/Automotive/AutomotiveMaterialPack/AutomotiveMaterialPack.INT.udn
+++ b/Engine/Documentation/Source/Enterprise/Automotive/AutomotiveMaterialPack/AutomotiveMaterialPack.INT.udn
@@ -1,10 +1,10 @@
 Availability:Public
 Title: Automotive Material Pack
-Description:A breakdown of what assets the Automotive Material pack contains and how to use them.
+Description:A breakdown of what assets the Automotive Material pack contains and how to use them
 Crumbs:%ROOT%,
 version:4.13
-Parent:Enterprise/Automotive
-tags:Enterprise
+Parent:Enterprise
+tags:Studio
 tags:Automotive
 tags:Materials
 type:overview
@@ -135,7 +135,7 @@ By adjusting just these four different parameters you can achieve vastly differe
 The following image sequence shows off the variety of paint colors and surface looks that can be achieved by using just  the parameters listed above.
 
 [REGION:fullwidth]
-[INCLUDE:Enterprise\Automotive\AutomotiveMaterialPack\#CP]
+[INCLUDE:#CP]
 [/REGION]
 
 ### Orange Peel Material Parameter
@@ -143,7 +143,7 @@ The following image sequence shows off the variety of paint colors and surface l
 The Clearcoat lacquer that is applied on top of automotive paint often has a slightly bumpy surface similar to the peel of an Orange. Because of this resemblance, this dimpled or uneven part of the Clearcoat is referred to as the Orange Peel. You can control the intensity of the Orange Peel using the **Orange Peel Strength** and **Orange Peel Scale** parameters. To disable the Orange Peel effect, set the **Orange Peel Strength** to **0**.
 
 [REGION:fullwidth]
-[INCLUDE:Enterprise\Automotive\AutomotiveMaterialPack\#OP]
+[INCLUDE:#OP]
 [/REGION]
 
 ### Carbon Fiber
@@ -196,7 +196,7 @@ If you are unfamiliar with how to use Vertex Paint in UE4 please check out the [
 [/region]
 
 [REGION:fullwidth]
-[INCLUDE:Enterprise\Automotive\AutomotiveMaterialPack\#RM]
+[INCLUDE:#RM]
 [/REGION]
 
 ### Brake Rotor
@@ -332,7 +332,7 @@ The reflector Material is used for the plastic  reflectors you find next to turn
 * **Triangle**
 
 	[REGION:fullwidth]
-	[INCLUDE:Enterprise\Automotive\AutomotiveMaterialPack\#RF]
+	[INCLUDE:#RF]
 	[/REGION]
 
 To change which reflector pattern is used,  open up the **Reflector** Material Instance and change the **Reflector(Bottom Layer) Normal** to the reflector pattern you want to use.
diff --git a/Engine/Documentation/Source/Enterprise/Enterprise.INT.udn b/Engine/Documentation/Source/Enterprise/Enterprise.INT.udn
--- a/Engine/Documentation/Source/Enterprise/Enterprise.INT.udn
+++ b/Engine/Documentation/Source/Enterprise/Enterprise.INT.udn
@@ -1,10 +1,10 @@
 Availability:Public
-Title: Enterprise Applications
+Title: Unreal Studio
 Description:Using Unreal Engine for automotive, aviation, architecture, consumer electronics, and complex data visualization applications
 Crumbs:%ROOT%,
-version:4.13
+version:4.19
 Parent:%ROOT%
-tags:Enterprise
+tags:Studio
 type:landing
 order:8
 topic-image:Enterprise_Topic.png
@@ -13,5 +13,240 @@ topic-icon:%ROOT%/enterprise_icon.png
 ![](Enterprise_Hero.png)
 
 Unreal Engine and its real-time interaction and rendering capabilities are well-suited for a huge range of projects across all Enterprise applications, including automotive, aviation, architecture, consumer electronics and complex data visualization.
+ 
+## Getting Started
 
-[DIR(output:"topic" org:"hierarchy" parent:"Enterprise" end:"1")]
+[DIR(output:"topic" org:"hierarchy" parent:"Enterprise" tags:"Getting Started" end:"1")]
+[OBJECT:Topic]
+	[PARAM:image]
+	![](Download_topic.png)
+	[/PARAM]
+	[PARAM:title]
+	Downloading and Installing Datasmith
+	[/PARAM]
+	[PARAM:description]
+	Learn now to get set up and start working With Datasmith for Unreal Engine 4 (UE4)
+	[/PARAM]
+	[PARAM:path]
+	[RELATIVE:Enterprise/GettingStarted/DownloadingUnrealEngineandInstallingDatasmith.pdf]
+	[/PARAM]
+[/OBJECT]
+
+## Unreal Studio Tools
+
+[OBJECT:Topic]
+	[PARAM:image]
+	![](Datasmith_topic.png)
+	[/PARAM]
+	[PARAM:title]
+	Datasmith
+	[/PARAM]
+	[PARAM:description]
+	A collection of tools and plugins designed to import scenes and assets created in offline rendering applications, like Autodesk 3DS Max and CAD software, into Unreal Engine 4
+	[/PARAM]
+	[PARAM:path]
+	[RELATIVE:Enterprise/GettingStarted/DatasmithOverview.pdf]
+	[/PARAM]
+[/OBJECT]
+[OBJECT:Topic]
+	[PARAM:image]
+	![](Import_topic.png)
+	[/PARAM]
+	[PARAM:title]
+	Datasmith Import Options
+	[/PARAM]
+	[PARAM:description]
+	Options are available when importing Unreal Datasmith or native CAD content using tne Datasmith Importer plugin for Unreal Engine 4
+	[/PARAM]
+	[PARAM:path]
+	[RELATIVE:Enterprise/GettingStarted/DatasmithImportOptions.pdf]
+	[/PARAM]
+[/OBJECT]
+[OBJECT:Topic]
+	[PARAM:image]
+	![](Filetypes_topic.png)
+	[/PARAM]
+	[PARAM:title]
+	Datasmith Supported Software and File Types
+	[/PARAM]
+	[PARAM:description]
+	Overview of the software packages and file types supported by Datasmith for Unreal Engine 4
+	[/PARAM]
+	[PARAM:path]
+	[RELATIVE:Enterprise/GettingStarted/DatasmithSupportedSoftwareandFileTypes.pdf]
+	[/PARAM]
+[/OBJECT]
+
+## Unreal Studio Guides
+
+[OBJECT:Fancy]
+	[PARAMLITERAL:skillevel]
+	Beginner
+	[/PARAMLITERAL]
+	[PARAMLITERAL:linkText]
+	Preparing a 3D Scene for Datasmith
+	[/PARAMLITERAL]
+	[PARAMLITERAL:linkTitle]
+	General guidelines covering how your content is translated to Unreal through Datasmith
+	[/PARAMLITERAL]
+	[PARAM:linkUrl]
+	[RELATIVE:Enterprise/GettingStarted/Preparinga3DSceneforDatasmith.pdf]
+	[/PARAM]
+	[PARAMLITERAL:version]
+	4.19
+	[/PARAMLITERAL]
+	[PARAMLITERAL:icon]
+	https://docs.unrealengine.com/latest/Include/Images/beginner_icon.png
+	[/PARAMLITERAL]
+[/OBJECT]
+[OBJECT:Fancy]
+	[PARAMLITERAL:skillevel]
+	Beginner
+	[/PARAMLITERAL]
+	[PARAMLITERAL:linkText]
+	Exporting Datasmith Content for Unreal Engine 4
+	[/PARAMLITERAL]
+	[PARAMLITERAL:linkTitle]
+	How to export content from your software package using file types are supported by the Datasmith Importer plugins for Unreal Engine 4 (UE4)
+	[/PARAMLITERAL]
+	[PARAM:linkUrl]
+	[RELATIVE:Enterprise/GettingStarted/ExportingDatasmithContentforUnrealEngine4.pdf]
+	[/PARAM]
+	[PARAMLITERAL:version]
+	4.19
+	[/PARAMLITERAL]
+	[PARAMLITERAL:icon]
+	https://docs.unrealengine.com/latest/Include/Images/beginner_icon.png
+	[/PARAMLITERAL]
+[/OBJECT]
+[OBJECT:Fancy]
+	[PARAMLITERAL:skillevel]
+	Beginner
+	[/PARAMLITERAL]
+	[PARAMLITERAL:linkText]
+	Importing Datasmith Content into Unreal Engine 4
+	[/PARAMLITERAL]
+	[PARAMLITERAL:linkTitle]
+	How to import and reimport your content from your software application using Datasmith Importer for Unreal Engine 4 (UE4)
+	[/PARAMLITERAL]
+	[PARAM:linkUrl]
+	[RELATIVE:Enterprise/GettingStarted/ImportingDatasmithContentintoUnrealEngine4.pdf]
+	[/PARAM]
+	[PARAMLITERAL:version]
+	4.19
+	[/PARAMLITERAL]
+	[PARAMLITERAL:icon]
+	https://docs.unrealengine.com/latest/Include/Images/beginner_icon.png
+	[/PARAMLITERAL]
+[/OBJECT]
+
+
+## Additional Content Resources
+
+[DIR(output:"topic" org:"hierarchy" parent:"Enterprise" tags:"!Getting Started" end:"1")]
+[OBJECT:Topic]
+	[PARAM:image]
+	![](Template_topic.png)
+	[/PARAM]
+	[PARAMLITERAL:title]
+	Unreal Studio Project Template - DRAFT
+	[/PARAMLITERAL]
+	[PARAMLITERAL:description]
+	Guide to using the "Product Viewer" project template designed to enable you to use Unreal as a design review tool
+	[/PARAMLITERAL]
+	[PARAM:path]
+	[RELATIVE:Enterprise/GettingStarted/UnrealStudioProjectTemplate_DRAFT.pdf]
+	[/PARAM]
+[/OBJECT]
+
+[COMMENT:none]
+[OBJECT:Topic]
+	[PARAM:image]
+	![](%ROOT%/placeholder_topic.png)
+	[/PARAM]
+	[PARAM:title]
+	Datasmith
+	[/PARAM]
+	[PARAM:description]
+	A collection of tools and plugins designed to import scenes and assets created in offline rendering applications, like Autodesk 3DS Max and CAD software, into a real-time engine
+	[/PARAM]
+	[PARAM:path]
+	[RELATIVE:Enterprise/GettingStarted/DatasmithOverview.pdf]
+	[/PARAM]
+[/OBJECT]
+
+[OBJECT:Topic]
+	[PARAM:image]
+	![](%ROOT%/placeholder_topic.png)
+	[/PARAM]
+	[PARAM:title]
+	Downloading and Installing Datasmith
+	[/PARAM]
+	[PARAM:description]
+	Learn now to get set up and start working With Datasmith for Unreal Engine 4 (UE4)
+	[/PARAM]
+	[PARAM:path]
+	[RELATIVE:Enterprise/GettingStarted/DownloadingUnrealEngineandInstallingDatasmith.pdf]
+	[/PARAM]
+[/OBJECT]
+
+[OBJECT:Topic]
+	[PARAM:image]
+	![](%ROOT%/placeholder_topic.png)
+	[/PARAM]
+	[PARAM:title]
+	Preparing a 3D Scene for Datasmith
+	[/PARAM]
+	[PARAM:description]
+	General guidelines covering how your content is translated to Unreal through Datasmith
+	[/PARAM]
+	[PARAM:path]
+	[RELATIVE:Enterprise/GettingStarted/Preparinga3DSceneforDatasmith.pdf]
+	[/PARAM]
+[/OBJECT]
+
+[OBJECT:Topic]
+	[PARAM:image]
+	![](%ROOT%/placeholder_topic.png)
+	[/PARAM]
+	[PARAM:title]
+	Exporting Datasmith Content
+	[/PARAM]
+	[PARAM:description]
+	How to export content from your software package using file types are supported by tne Datasmith Importer plugins for Unreal Engine 4 (UE4)
+	[/PARAM]
+	[PARAM:path]
+	[RELATIVE:Enterprise/GettingStarted/ExportingDatasmithContentforUnrealEngine4.pdf]
+	[/PARAM]
+[/OBJECT]
+
+[OBJECT:Topic]
+	[PARAM:image]
+	![](%ROOT%/placeholder_topic.png)
+	[/PARAM]
+	[PARAM:title]
+	Importing Datasmith Content
+	[/PARAM]
+	[PARAM:description]
+	How to import and reimport your content from your software application using Datasmith Importer for Unreal Engine 4 (UE4)
+	[/PARAM]
+	[PARAM:path]
+	[RELATIVE:Enterprise/GettingStarted/ImportingDatasmithContentintoUnrealEngine4.pdf]
+	[/PARAM]
+[/OBJECT]
+
+[OBJECT:Topic]
+	[PARAM:image]
+	![](%ROOT%/placeholder_topic.png)
+	[/PARAM]
+	[PARAM:title]
+	Datasmith Import Options
+	[/PARAM]
+	[PARAM:description]
+	Options are available when importing Unreal Datasmith or native CAD content using the Datasmith Importer plugin for Unreal Engine 4
+	[/PARAM]
+	[PARAM:path]
+	[RELATIVE:Enterprise/GettingStarted/DatasmithImportOptions.pdf]
+	[/PARAM]
+[/OBJECT]
+[/COMMENT]
\ No newline at end of file
