Copying //UE4/Dev-Documentation to //UE4/Samples-Main (Source: //UE4/Dev-Documentation @ 4091886)

diff --git a/Engine/Documentation/Source/Enterprise/GettingStarted/UnrealStudio.INT.udn b/Engine/Documentation/Source/Enterprise/GettingStarted/UnrealStudio.INT.udn
new file mode 100644
--- /dev/null
+++ b/Engine/Documentation/Source/Enterprise/GettingStarted/UnrealStudio.INT.udn
@@ -0,0 +1,590 @@
+title:Getting Started for Unreal Studio Users 
+description:An overview of the Unreal Engine 4 (UE4) with the intent to help Unreal Studio users translate their existing experience into the world of Unreal Engine 4
+crumbs:
+availability:Public
+parent:Enterprise
+order:1
+type:overview
+tags:Getting Started
+version:4.19
+topic-image:Enterprise/Start_topic.png
+
+[TOC]
+
+This guide provides an overview of the Unreal Engine 4 (UE4) with the intent to help Unreal Studios users translate their existing experience into the world of Unreal Engine 4. The information contained in this page points out common areas of the Unreal Editor and its tools, points out existing documentation that can best help you adjust to the tools UE4 offers, workflows that best fit real-time engine development, and much more. 
+
+As part of the Unreal Studio offering, [Datasmith](DatasmithOverview.pdf) is a collection of tools and plugins designed to import scenes and assets created in offline rendering applications into a real-time engine. Whether you're an Unreal Studio user using Datasmith for the first time or have experience with it, the information contained within this page can help you improve your workflows while working in a real-time engine. Even if you're not using Datasmith currently, the best practices and suggestions referenced here can benefit anyone using UE4 coming from an Enterprise background.
+
+## Creating Projects for Unreal Engine 4
+
+When you create a project for Unreal Engine 4, you will use the [Project Browser](Engine/Basics/Projects/Browser/) to do so when you launch the Engine after installation. The Project Browser enables you to select an existing project or create a new one that can be a Blueprint, C++, or Unreal Studio project. 
+
+![image alt text](image_0.png)(w:617 h:429 convert:false)
+
+In the New Project tab, you'll select the type of project you'll want to create. For Unreal Studio users, you'll want to use the **Unreal Studio** tab so that plugins available as part of the Unreal Studio offering (such as Datasmith) are automatically enabled for use. 
+
+[REGION:note]
+If you are not able to see the Unreal Studio tab because your Unreal Studio Entitlements have not yet been activated or you're having trouble with them, that's OK. You can still create a project from one of the **Blueprint** or **C++** tabs and later convert it to an Unreal Studio Project. For Entitlement issues, see the **Entitlements** section of the [Datasmith](DatasmithOverview.pdf) page. To learn how you can convert your existing project or one you've previously created that is not an Unreal Studio project into one, see the **Converting a Project to an Unreal Studio Project** section of the [Download and Installing Datasmith](DownloadingUnrealEngineandInstallingDatasmith.pdf) page.
+[/REGION]
+
+![image alt text](image_1.png)(w:616 h:429 convert:false)
+
+In the Unreal Studio tab, you'll select from the available templates to create your Project from. Currently, Unreal Studio offers two types of templates:
+
+* **Blank** - A clean empty project with no additional content added. 
+
+* **Product Viewer** - A template that enables you to view individual assets with pre-built interaction and navigation to display objects.
+
+For additional information about templates, see the of the [Unreal Studio Project Templates](UnrealStudioProjectTemplate_DRAFT.pdf) page.
+
+For additional information about creating projects for Unreal Studio, see the [Download and Installing Datasmith](DownloadingUnrealEngineandInstallingDatasmith.pdf) page.
+
+## The Editor
+
+[REGION:note]
+Autodesk 3ds Max has been chosen for the interface comparison since the majority of current users are using it.
+[/REGION]
+
+Below are pictures of Autodesk 3ds Max and the Unreal Editor, color-coded to indicate common functionality. Each block is labeled to show the equivalent UE4 terminology. Unreal Editor's layout is fully customizable by dragging and dropping any of the tabs in the Editor around. 
+
+[REGION:imagetable]
+| [INCLUDE:#image1] | [INCLUDE:#image2] |
+| --- | --- |
+|Click image for full size.| Click image for full size. |
+[/REGION]
+
+<!---
+[INCLUDE:#image2]
+[EXCERPT:image1]
+[REGION:lightbox]
+[![image alt text](image_2.png)(w:575 convert:false)](image_2.png)
+[/REGION]
+[/EXCERPT]
+
+[EXCERPT:image2]
+[REGION:lightbox]
+[![image alt text](image_3.png)(w:575 convert:false)](image_3.png)
+[/REGION]
+[/EXCERPT]
+--->
+
+| --- | --- | --- |
+| 1 | Placement Panel | Create Panel |
+| 2 | Main Toolbar | Main Toolbar |
+| 3 | World Outliner | Scene Explorer |
+| 4 | Details Panel | Command Panel |
+| 5 | Viewport | Viewport |
+| 6 | Content Browser | -- |
+
+
+### Editing Assets
+
+In Autodesk 3ds Max the Command Panel exposes properties of selected objects. In UE4, the **Details** panel exposes the properties of selected objects, while larger editing tasks lead to a dedicated Editor that has its own window or tab. Examples of these can be found while opening a Texture, Materials, or Static Mesh. New tabbed windows will open for each edited asset, similar to a web browser.
+
+![image alt text](image_4.jpg)(w:1013 h:570 convert:false)
+
+*The Material Editor docked in the main Unreal Editor tabs. *
+
+When Editor windows are docked, you can select the tab to bring it to the foreground, or when another tab is selected, it will move to the background (like the example above). All tabs, except the Main Editor tab, can be dragged as a freestanding window or docked with other freestanding windows or tabs inside the Editor. 
+
+### The Main Toolbar
+
+The Main Toolbar (or Level Editor Toolbar), like in most applications, is a group of commands providing quick access to commonly used tools and operations. 
+
+![image alt text](image_5.png)(w:952 h:64 convert:false)
+
+### Scene Navigation
+
+Unreal Engine 4 has a few ways that enable you to navigate your scene, select and manipulate Actors, and ones that enable you to change how content is displayed in the viewport (even while working in virtual reality (VR)).
+
+For **Desktop **users using a Mouse and Keyboard, see the [Viewport Controls](Engine/UI/LevelEditor/Viewports/ViewportControls) page for the controls used to navigate the viewport scene. There you will learn about the controls that UE4 uses for Selecting and Transforming Actors, panning, orbiting, and zooming the Camera and navigation. 
+
+For those using **Virtual Reality (VR)** with controllers, see the [VR Mode Controls](Engine/Editor/VR/Controls) and the [Navigating the World in VR Mode](Engine/Editor/VR/WorldInteraction) pages to learn about the controls available. Also learn how you can use them to grab objects to manipulate them, move around the world, and scale yourself in the world to achieve something closer to real-world scale for scene visualization.
+
+## Quick Glossary
+
+The following section contains common Enterprise or Autodesk 3ds Max terms and their UE4 equivalents (or rough equivalent). Unreal Engine 4 keywords link directly to the more in-depth information contained inside of the UE4 Documentation.
+
+
+| Category | Autodesk 3ds Max | UE4 | Notes |
+| --- | --- | --- | --- |
+| [REGION:tablesection]Meshes[/REGION] ||||
+|  | Mesh or mesh object | [Static Mesh](Engine/Actors/StaticMeshActor) |  |
+|   | Rigged and skinned mesh | [Skeletal Mesh](Engine/Content/Types/SkeletalMeshes) |  |
+|  | Object/Node | [Level Actor](Engine/Actors/StaticMeshActor) |  |
+|  | Modifier Stack Result/Base Object | [Static Mesh](Engine/Actors/Types) |  |
+| [REGION:tablesection]Editor UI[/REGION] ||||
+|  | Scene Explorer | [World Outliner](Engine/UI/LevelEditor#worldoutliner) |  |
+|   | Modify Panel | [Details Panel](Engine/UI/LevelEditor#details) |  |
+|   | Layers Panel | [Layers Panel](Engine/UI/LevelEditor/Layers) |  |
+|   | Asset Browser Utility | [Content Browser](Engine/Content/Browser) |  |
+|  | Viewport | [Viewport](Engine/UI/LevelEditor/Viewports) |  |
+| [REGION:tablesection]Materials[/REGION] ||||
+|  | Material | [Material](Engine/Rendering/Materials) |  |
+|   | Material with instanced textures | [Material Instance](Engine/Rendering/Materials/MaterialInstances) |  |
+|   | Material Editor | [Material Editor](Engine/Rendering/Materials/Editor) |  |
+| [REGION:tablesection]Effects[/REGION] ||||
+|  | Particle Systems | [Effect, Particle System, Cascade](Engine/Rendering/ParticleSystems) |  |
+|   | Particle View | [Cascade](Engine/Rendering/ParticleSystems) |  |
+| [REGION:tablesection]Animation[/REGION] ||||
+|  | Biped animation | [Skeletal Animation System](Engine/Animation) | These are not exactly the same but can be a good resource to see comparisons and how Unreal works with imported animations. |
+|   | Animation at the vertex sub-object level | [Vertex Animation](Engine/Animation/Tools/VertexAnimationTool/) | There is not a dedicated tool for applying vertex animation to an asset, but you can use this plugin for Autodesk 3ds Max to bring them into UE4 and apply them using the Material Editor.  |
+| [REGION:tablesection]Programming[/REGION] ||||
+|  | MAXScript | Python |  |
+|   | MAXScript scripts | [Blueprint](Engine/Blueprints) |  |
+| [REGION:tablesection]Physics[/REGION] ||||
+|  | -- | [Line Trace, Shape Trace](Engine/Physics/Tracing) |  |
+|   | MassFX | [Collision, Physics](Engine/Physics) |  |
+| [REGION:tablesection]Lights[/REGION] ||||
+|  | Directional Light | [Directional Light](Engine/Rendering/LightingAndShadows/LightTypes/Directional) |  |
+|   |  Point Light | [Point Light](Engine/Rendering/LightingAndShadows/LightTypes/Point) |  |
+|   |  Spot Light | [Spot Light](Engine/Rendering/LightingAndShadows/LightTypes/Spot) |  |
+|   |  Sky Light | [Sky Light](Engine/Rendering/LightingAndShadows/LightTypes/SkyLight) |  |
+|   | Light Portal | Lightmass Portals, [Lightmass Importance Volume](Engine/Rendering/LightingAndShadows/Lightmass/Basics/#lightmassimportancevolume) | |
+|   | Photometric Lights | [IES Profile](Engine/Rendering/LightingAndShadows/IESLightProfiles/), Physical Light Units |  |
+| [REGION:tablesection]Post Processing[/REGION] ||||
+|  | Environment Exposure Control, Camera Exposure Control | [Post Process Volume](Engine/Rendering/PostProcessEffects) |  |
+| [REGION:tablesection]Rendering[/REGION] ||||
+|  | Global Illumination | [Lightmass](Engine/Rendering/LightingAndShadows/Lightmass/) |  |
+| [REGION:tablesection]Cinematics[/REGION] ||||
+|  | Track View | [Sequence Editor](Engine/Sequencer) |  |
+| [REGION:tablesection]Textures[/REGION] ||||
+|  | UV Map Channel | UV Channel of a Static/Skeletal Mesh |  |
+| [REGION:tablesection]Unreal Engine Architecture[/REGION] ||||
+|  | -- | [Level (Map)](Engine/Levels/) | Each scene file is self-contained, whereas in Unreal Engine 4, you can have multiple levels (or scenes) that can be opened individually or combined all within the Unreal Editor. |
+
+
+
+## Project and Files
+
+**So what are all these directories and files?**
+
+For projects created in the Enterprise space, you may have to maintain your own folder to manage your source files for textures and meshes that make up your scene project file. In Unreal, projects always exist in their own directory and have their own project file that references a specific version of the Unreal Engine Editor and looks at a set of folders where its content is stored. 
+
+![image alt text](image_6.jpg)(w:693 h:349 convert:false)
+
+*An example of the Content Browser in Unreal and the Folder Structure maintained in Windows Explorer.*
+
+To open your project, you can use the Epic Games Launcher to search your **Library** for the project. 
+
+![image alt text](image_7.png)(w:941 h:609 convert:false)   
+
+Alternatively, navigate to its folder in Windows and double-click on the ***.uproject** file to load the project into the Unreal Editor. 
+
+![image alt text](image_8.png)(w:613 h:293 convert:false)
+
+The project folder has various sub-folders that contain your project's content and source code as well as various configuration files and binaries when deploying to multiple platforms. The two most important folders here are the **Content**, which holds all your assets that can be referenced and opened inside of the Content Browser, and the **Source** folder, which contains any C++ files or edits that you have made if you're using a source version of the Engine. 
+
+**Where do I put my assets?**
+
+In UE4, each project has a **Content **folder, which is the same as the Content Browser when you open the Unreal Editor. This is where all your project's assets are stored. 
+
+![image alt text](image_9.png)(w:1026 h:462 convert:false)
+
+The assets contained within these folders are ***.uasset** files. These are files that Unreal creates when you import any supported file types. The source file is no longer required to maintain the asset imported into Unreal, but it is recommended to save these files in the event you need them later for editing purposes. 
+
+[REGION:warning]
+You should not attempt to reorganize your assets found in the Content folder of your project outside of the Unreal Editor. This has the potentital to break your content or worse corrupt your project. Moving content like this will break references to the assets in the Unreal Editor. For example, if you moved a texture asset through the project's Windows folder, the Material that references that texture will be lost meaning it will not display correctly for the Material that is assigned to your Static Mesh in your level. 
+
+Make sure to use the Unreal Editor when moving assets since it maintains references without breaking your content.
+[/REGION]
+
+**Importing Assets into Unreal Engine 4**
+
+To import assets into your project, there are a couple of ways to go about doing that: 
+
+1. [Using Datasmith to import Unreal Datasmith and support CAD file types.](ImportingDatasmithContentintoUnrealEngine4.pdf)
+
+2. [Using the Content Browser to import support asset types not associated with Datasmith.](Engine/Content/ImportingContent/ContentBrowserImport)
+
+When assets are imported using either of these methods, they will be stored in your Content Browser, which mirrors the project's folder structure in Windows. Use the Content Browser to open and edit your assets in their own editors through the Content Browser. Any edits you make will automatically update the asset in real-time in the Editor and the level if it's referenced there.
+
+![image alt text](image_10.jpg)(w:582 h:293 convert:false)
+
+[REGION:note]
+
+When importing content that is not part of Datasmith Import process, like when using the Content Browser import button, assets will be imported directly into the Content folder for the selected folder. When using the Datasmith Importer, it will bring in a lot of different assets types (meshes, materials, and textures) in a collection of subdirectories to keep your scene clean. The Datasmith scene and its sub-directories will be imported to a folder named for the Datasmith file that you import. 
+[/REGION]
+
+**What common file formats are supported?**
+
+Other software applications support a wide variety of file formats. Unreal Engine 4 supports the most common file types natively. 
+
+
+| Native Asset Type | Supported Formats |
+| --- | --- |
+| 3D | .fbx, .obj, .abc, .srt |
+| Texture | .png, .jpeg, .jpg, .bmp ,.tga, .dds, .exr, .psd, .hdr, .ies, .float, .pcx |
+| Sound | .wav, .wave, .wmv, .caf, .aif, .aiff, .au, .adts, .as, .ac3, .amr, .bwf, .cdda, .acc, .mp3, .snd |
+| Fonts | .ttf, .otf |
+| Videos | .mov, .mp4, .m4v, .m4a, .wmv, .avi, .asf, .3gp, .3g2, .sdv |
+
+
+Through plugins, like [Datasmith](DatasmithOverview.pdf) for Unreal Engine 4, support for additional file types can be added. For a comprehensive list of supported file types, see the [Datasmith Supported Software and File Types](DatasmithSupportedSoftwareandFileTypes.pdf) page.
+
+**How is my Scene stored?**
+
+In your software applications, the assets that make up your scene make up a single scene file. Unreal uses a different concept by using what we call a Level or Map file. These are individual "scenes" that have their own assets placed in them. They can be opened individually or even combined to break a larger level up into smaller ones. Map files store data about your Level and the objects that it contains as well as certain level-specific settings. 
+
+If you'd like to learn more about Levels and how Unreal Editor works and stores these, see the [Working with Level Assets](Engine/Levels/HowTo/WorkWithLevelAssets) page.
+
+<!---
+**How do I change my Editor's and Project's settings?**
+
+All the settings can be found in the main menu under **Edit > Editor Preferences** or **Project Settings**. Like the preferences or settings in your software application, these allow you to customize and change parts of how your project works. 
+
+In Unreal Engine 4, the [Editor Preferences](Engine/UI/EditorSettings) contains settings that enable you to alter how the Editor functions by setting things like color correction for common color blindness deficiencies, enabling experimental features in the Editor, changing the units of measure from centimeters to something more familiar. 
+
+The [Project Settings](Engine/UI/ProjectSettings) enables you to specify information about your project (such as project name and icons), configure game input bindings, and define how the Engine behaves when running your project. 
+
+--->
+
+##Art Source Auto-Reimport
+
+[REGION:note]
+Currently with Datasmith, it is recommended not to use Auto-Reimport as part of your workflow. However, for any other assets you are creating and importing that are natively supported, you can set up auto-reimport monitoring.
+[/REGION]
+
+For projects, where you'll be using additional content that you import aside from just that used with the [Datasmith](DatasmithOverview.pdf) plugin, you can set up Unreal to use [Auto Reimport](Engine/Content/AutoReImport). It monitors a source assets folder where you store assets like FBX, OBJ, PSD, Texture, Audio, and any other native file that UE4 supports for import. When the source asset is updated, the monitored folder will automatically reimport or prompt you to reimport the file and reflect the changes inside of the Engine. Unlike, Datasmith, these assets are not placed directly in your level and takes an approach to get content in the Editor but not set it up for you. 
+
+**From Objects (Scene Entities) to Actors**
+
+**Where is my Object?**
+
+In software applications, like Autodesk 3ds Max, an asset in your scene may be referred to as a scene entity, geometry object, or some other familiar term for all the objects that make up your scene. In UE4, the equivalent is an Actor. Actors are objects added to a level, referencing other assets such as Static Meshes and Materials. 
+
+[OBJECT:EmbeddedVideo]
+	[PARAMLITERAL:videoid]
+	UXGFtZ71mfQ
+	[/PARAMLITERAL]
+	[PARAMLITERAL:width]
+	50
+	[/PARAMLITERAL]
+	[PARAMLITERAL:height]
+	65
+	[/PARAMLITERAL]
+	[PARAMLITERAL:units]
+	%
+	[/PARAMLITERAL]
+[/OBJECT]
+
+You could build a project out of Empty Actors (like the one in the video), but UE4 also includes special types of Actors with built-in features, such as a [Pawn](Gameplay/Framework/Pawn) (for players or Artificial Intelligence (AI) objects), an Exponential Height Fog Volume, or Lights. These Actors in your level can have their properties customized and even have additional components added to them using their Details panel. The same applies to assets from the Content Browser that you reference from Actors, Materials, or Blueprints. They, too, can have their own properties set and components added to them. 
+
+**So, is everything an Actor?**
+
+No, not everything. **Actors **are the most common class used for gameplay in UE4 and the only type that can be Spawned into the Level. So everything you place in your level will be an Actor. 
+
+The other important type to know about is **Object**. An Object is the base class of all Unreal classes, including Actor and many others. This is a much lower-level construct than Actor but still has features that you will expect from an Unreal class, such as Reflection and Serialization. An Object is a class we use when we need to define a new type that just does not fit the mold of Actor. For example,** Actor Component** is the base class for all Components and derives from Object instead of Actor. 
+
+All that said, you should not have to worry about how this functions for you, but it can help clarify why you may sometimes hear the terms Actor and Object mentioned in close proximity or potentially in place of Actor. 
+
+**What are Components?**
+
+In software applications, like Autodesk 3ds Max, you add functionality or extend properties to objects by using the Modifier stack.
+
+PF Notes: This can be effectively become a mini-scene within a scene and it's where a lot o fpower and felxibility of Unreal is. 
+Give an example of typical application: a stove actor would have a door and a drawer component and a Light or something along those lines. 
+
+
+In UE4, [Components](Engine/Components/) are a piece of functionality that can be added to an Actor in your level or in a Blueprint. These components cannot exist by themselves, however, when added to an Actor will have access to and can use functionality provided by the Component. You can think of these as a way of creating a mini-scene or a grouping of Actors that belong together which are treated as a single object. This process enables you to have a lot of flexibility and control over your scene in a manageable way in Unreal.  
+
+For example, if you have a stove asset, you can add different componenst that make up other parts of that base asset, like a door, a bottom drawer, and maybe a light. 
+
+There are many different types of components that can be added to your Actors that reference many different types of assets, like Static Meshes, Particle Systems, Lights, and so on. To add these to your Actor, use the **Details** panel and choose a component from the list to add. 
+
+[OBJECT:EmbeddedVideo]
+	[PARAMLITERAL:videoid]
+	SknellY4I4g
+	[/PARAMLITERAL]
+	[PARAMLITERAL:width]
+	50
+	[/PARAMLITERAL]
+	[PARAMLITERAL:height]
+	45
+	[/PARAMLITERAL]
+	[PARAMLITERAL:units]
+	%
+	[/PARAMLITERAL]
+[/OBJECT]
+
+In this example, Static Mesh, Point Light, and Particle System components have been added to an Empty Actor placed in the scene. This Actor exists only in this level unless it's converted to a Blueprint Actor, which can exist as a single asset in your Content Browser.
+
+For additional information, see the [Components](Engine/Components/) page. 
+
+**How do I load my last project automatically?**
+
+When opening a scene file in your software application, you immediately get access to your content for that scene. However, in Unreal Engine 4, when you open just the Editor from the Epic Games Launcher, you will be taken to the [Unreal Project Browser](Engine/Basics/Projects/Browser/) where you can select your project (if it's already created). There are two ways you can enable your last project loaded to be used when opening the Unreal Editor:
+
+In the Project Browser, enable **Always load last project on Startup** when selecting a project to open. 
+
+![image alt text](image_28.png)(w:965 h:398 convert:false)
+
+**Where do I Define the Keyboard Inputs or Game Pad Inputs for my Project?**
+
+In UE4, you can add various buttons (actions) and analog controls (axes) to get callbacks to your game's character (Pawn) when input events are triggered. The Input Bindings can be found in the Project Settings under the **Input** section. Here you will give each control a name and a default binding to a key on your keyboard or a button on a controller. The names of these input bindings can be used in Blueprint to perform actions when pressed, and it allows you to set multiple inputs to a single name. For example, if you have a player character that you want to move forward in your level, you could call the Input Binding "Forward" and then bind the "W" key from the keyboard and the "Up" position on a joystick of a controller to this "Forward" event.
+
+For additional information, see the [Input ](Gameplay/Input) page learn more.
+
+**How do I change the starting level (map) in my project?**
+
+You can change the startup map of your project from the **Project Settings**. From the main menu, select **Edit > Project Settings** and then under **Maps & Modes** change the **Startup Map** to the one you'd like in the Content Browser. 
+
+![image alt text](image_29.png)(w:1026 h:400 convert:false)
+
+**How do I run my game?**
+
+The easiest way to run your game is to click the **Play** button on the main toolbar. It will run the game inside of the Editor's viewport with all functionality and the currently selected Gamemode. If you want to run it as a standalone application, click the drop-down next to **Play** and choose **Standalone Game**. Finally, if you want to run on a mobile device or in a web browser, you will use the **Launch** button on the toolbar and select from the available platforms. Some platforms require additional set up and can only be used after installing any prerequisites that are needed. 
+
+![image alt text](image_30.png)(w:506 h:220 convert:false)
+
+For additional information about the various platforms supported, see the [Platform Development](Platforms) page to learn more.
+
+**What units are used in Unreal Engine 4?**
+
+In your software application, you can adjust the primary unit of measure from your preferences or settings. In UE4, the primary unit of measure if one centimeter. For this reason, it is recommended to change your primary units of measure to match those used in Unreal.
+
+**What is the coordinate system? Which way is up?**
+
+Unreal Engine 4 uses a left-handed coordinate system. In UE4, positive X is "Forward," positive Y is "Right," and positive Z is "Up." 
+
+**How do I see the Log Output from my project?**
+
+In the Editor, you can open the **Output Log** by using the file menu to select **Window > Developer Tools**. 
+
+![image alt text](image_31.png)(w:421 h:613 convert:false)
+
+
+
+##Rendering with a Traditional Raytracer versus Real-time
+
+The lighting system in Unreal Engine 4 is designed to be flexible and provide different levels of lighting features that fit your project and a real-time game engine. This means that you can have fully static (or pre-baked) lighting or fully dynamic lighting for your world. For the purpose of Enterprise Architectural Visualization, static lighting will be used and leverage the Lightmass Global Illumination system in UE4. 
+
+![image alt text](Datasmith_HFB.jpg)
+[REGION:caption]
+_Image provided courtesy of HFB Visuals._
+[/REGION]
+
+To get the most out of UE4's lighting system, below are a few concepts to help you get your level lighting built and improve the quality with some minimal effort. First, understanding that you need to build lighting after any changes to see the results is key and additionally that shadowing information is stored in a Lightmap texture that may need its resolution increased to capture beautiful detail.    
+
+###Building Lighting for your Level and what that means
+
+In Unreal, especially with Enterprise Architectural work, there is a focus on beautifually lit scenes that match your renders in your software application. The Unreal [Lightmass Global Illumination](Engine/Rendering/LightingAndShadows/Lightmass) system is used to generate lighting for your levels. Before that can happen, you're asset will need a proper Lightmap UV generated for it to store this light and 
+shadow information properly. Unreal Engine 4 does a lot of the heavy-lifting to create lightmaps for you by auto-generating them on import so that you don't have any overlapping faces and that the UV charts are nicely packed in the 0,1 space of the UV. For Unreal Studio users, Datasmith streamlines this further by actually unwrapping the UVs for the mesh during import before repacking them. This leads to fewer issues since the out-of-the-box version of UE4 is limited to only repacking an existing UV. This is process is further explained in detail in the Lightmap UV section below. 
+
+When you import your Datasmith Scene, it will import directly into the level with lights placed where they were in your application scene file. Here, you will see what is referred to a "Preview" lighting to give you an idea of what the scene will look like. 
+
+![image alt text](image_12.png)(w:569 h:364 convert:false)
+
+When you see the "Preview" label in shadowed areas this indicates a couple of things. First, that the light source you are using requires lighting to be built for the level because you are using a light source that is set to Static or Stationary. Secondly, it means that the shadows shown are using a dynamic shadow to give you an idea of what lighting will look like when you build it. When this happens, you aren't getting accurate lighting from GI bounces or completely accurate shadowing for the level. 
+
+To kick off a light build, you can click the **Build** button in the main toolbar.
+
+![image alt text](image_13.png)(w:248 h:64 convert:false)
+
+The **Build** button will kick off a series of tasks, such as rebuilding lighting, recapturing reflections for reflection capture Actors, and other game-related tasks. 
+For Unreal Studio users, there is a level lighting quality setting that must be set. It can be found under the Build drop-down and select **Lighting Quality**.
+
+![image alt text](image_14.png)(w:489 h:261 convert:false)
+
+The Lighting Quality levels can be set from **Preview** to **Production**. Quality levels determine how accurate the lighting will be when generating it for the level. Preview will be quicker to build for a lot of scene but less accurate. Production quality will take longer but will be much more accurate and representative of your final light build for the level settings. 
+
+[REGION:note]
+For a comparison of quality, visit the [Lighting Quality Comparison](https://wiki.unrealengine.com/LightingTroubleshootingGuide#Lighting_Quality_comparison.2C_or_Production_.3E_Preview) example on the [Lighting Troubleshooting Guide](https://wiki.unrealengine.com/LightingTroubleshootingGuide) created and curated by our community.
+[/REGION] 
+
+When lighting is being built, you will have a percentage pop-up in the corner of the screen indicating that lighitng is currently being built for the level. If the percentage appears to be "stuck" for a while, you may just need to wait for a bit while lighting is generated. Depending on the complexity of your level and the lighting quality level chosen, it could take a little while to generate.
+
+![](BuildLighting1.png)
+
+Once lighting is built and applied to the scene, you'll no longer see the "preview" label. However, you may see that maybe the quality of the shadows generated are not what you expected. Keep in mind that Lighting Quality set in the level only tells Unreal Lightmass how accurate it should be, not how much texture resolution should be used to capture shadowing information for your lightmap UV. For this reason, you'll want to read the next section where we discuss the importance of Lightmap UVs and how you can control their resolution. 
+
+![image alt text](image_15.png)(w:600 h:397 convert:false)
+
+
+###Lightmap UVs
+
+During import, a lightmap UV is generated for each Static Mesh automatically based off a UV channel. By default, UE4 uses Channel 0 for lightmap UV generation. As previously mentioned for Unreal Studio users, Datasmith takes this a step further and streamlines the process by unwrapping and repacking in a single go so that there is no overlapping faces and all the UV charts are contained within the 0,1 UV space. This process is needed for applications, like CAD ones, where you don't have access to a UV Editor to create a proper unwrap or even when you just want to get your scene in UE4 but don't need to worry about the time consuming process of manually creating a proper Lightmap UV. Datasmith handles all this under-the-hood for you. 
+
+![image alt text](image_16.png)(w:1026 h:602 convert:false)
+*Left, an asset brought in from a CAD file using Datasmith; Right, the resulting UV unwrap and repack.*
+
+In this example, this is a complex object that has many different faces that have been broken apart and packed into the lightmap UV so that there is no overlapping faces and they're all contained within the 0,1 UV space. This ensures that when lighting is built for this mesh, there won't be any lightmap errors that necessarily require editing the UV charts outside of the Unreal Editor. 
+
+To inspect your lightmap UV, you can start by opening your Static Mesh by selecting it in the Content Browser and double-clicking. 
+
+![](StaticMeshEditor.png) 
+
+In the toolbar, click the **UV** button and select a UV channel to display. 
+
+![](StaticMeshEditor2.png) 
+
+Next, the selected UV channel will be displayed in the Static Mesh Editor veiwport.
+
+![](StaticMeshEditor3.png)
+
+[REGION:tip]
+If you're not sure which UV channel is your Lightmap UV, you can click through them to see which one looks similar to the example above, or to be absolutely sure, you can use the **Details** panel and see what UV channel is assigned as the **Lightmap Coordinate Index**.
+[/REGION]
+
+One thing to note about this mesh in particular is that since this is a large mesh which has many faces that end up being very tiny within the lightmap UV, you can still get errors related to quality when lighting is built for the level. You can solve these issuses by using one of the following methods:
+
+* Increase the Lightmap Resolution of the Static Mesh.
+* Manually unwrap and Reimport your Mesh
+* In your software application, break the model into smaller chunks.
+
+###Lightmap Resolution
+
+The lightmap resolution increases the texture resolution. It improves the quality of the lightmap generated during the light build. Sometimes, this can help when you've got a lightmap UV that contains really small UV charts but ones that aren't so small that increasing the texture resolution still cause issues with quality. 
+
+You can adjust the **Lightmap Resolution** of a mesh in the Static Mesh Editor's Details panel under the **General** section. Look for **Lightmap Resolution** and set it to a power of two value (64, 128, 256, 512, etc).
+
+![](LightmapResolution.png)
+
+So if we take this single mesh with all its UV charts scaled to fit within a single lightmap UV, the lightmap resolution can greatly impact the quality of the baked light and shadow information. Even using a high lightmap resolution like a 2048 texture resolution can still have issues because some faces in the UV chart will still be too small to store proper lighting information at an acceptable resolution.  
+
+![image alt text](image_17.png)(w:437 h:437 convert:false)![image alt text](image_18.png)(w:470 h:437 convert:false)
+
+Increasing the lightmap resolution can help decrease artifacts from the light build related to quality. However, it also increases the amount of texture memory overall and increases the time it takes to generate lighting for the level.
+
+![image alt text](image_19.png)(w:488 h:453 convert:false)
+
+*Single mesh with an increased lightmap resolution of 2048.*
+
+**Breaking The Mesh Into Smaller Chunks**
+
+Instead of trying to increase the lightmap resolution, it can be more efficient to break the model up into smaller chunks. This affords you the ability to use lower lightmap resolutions, have parts that are removed from memory when they are not within the camera's view, and will ultimately provide higher quality lighting since the UV charts can use more of the lightmap UV space efficiently. 
+
+Take, for example, the original UV layout for the single mesh (left). It had so much geometry crammed into a single lightmap UV that lighting detail can easily get lost and cause your light build to not look correct, even with an increased lightmap resolution. Breaking the model up into smaller chunks (middle), they now have their own lightmaps (right), but use more of the UV space to store lighting information properly and with better quality at a lower lightmap resolution. 
+
+![image alt text](image_20.png)(w:977 h:933 convert:false)
+
+[REGION:tip]
+Try to break the model into smaller pieces where it makes sense, like at the edge of a seam or inset. This can help hide any baking issues or resolution differences between the pieces.
+[/REGION]
+
+![image alt text](image_21.png)(w:505 h:468 convert:false)![image alt text](image_22.png)(w:502 h:466 convert:false)
+
+*Left, a single mesh with a lightmap resolution of 2048; Right, is the broken up model chunks with lightmap resolutions of 512.*
+
+You'll notice in the image above, that the single mesh (left) is still dark but has better quality shadows, even after increasing the lightmap resolution. However, you don't get as much bounce light retained since the UV charts don't use enough of the UV space to store that information. On the right, the quality is similar, but because the UV charts use more of the lightmap UV space, it captures GI information much more accurately. 
+
+For additional information about lighting and build issues, you can visit the community maintained [Lighting Troubleshooting Guide](https://wiki.unrealengine.com/LightingTroubleshootingGuide) that has many good examples and solutions to common light build issues. 
+
+Below are some links to get you started with Lighting in UE4:
+
+* [Lighting Basics](Engine/Rendering/LightingAndShadows/Basics) - An overview of the basics of placing lights into levels and setting them up. 
+* [Lightmass Global Illumination](Engine/Rendering/LightingAndShadows/Lightmass) - A rundown of the all the features and settings of Lightmass uses for generating precomputed lighting in UE4.
+* [Static Lighting](Engine/Rendering/LightingAndShadows/LightMobility/StaticLights) - An overview of Static Lights.
+* [Lighting And Shadows Topics](Engine/Rendering/LightingAndShadows) - A listing of all available lighting features documentation broken up into key areas of focus.
+[COMMENT:none]
+* [Shadow Casting](Engine/Rendering/LightingAndShadows/Shadows) - A high-level overview of shadowing in UE4.
+[/COMMENT]
+
+##Materials and Textures
+
+In UE4, textures are applied to Materials that can be edited in the Material Editor. The Material Editor, like some other applications, uses a node-based graph interface that enables you to create shaders that can be applied to your geometry, such as Static and Skeletal Meshes, or with systems such as Cascade to create interesting materials. 
+
+![image alt text](image_23.jpg)(w:709 h:260 convert:false)
+
+Unreal Engine 4 uses a physically-based shading model, which is similar to the [Disney "Principled" BRDF Physically-Based Rendering model](https://disney-animation.s3.amazonaws.com/library/s2012_pbs_disney_brdf_notes_v2.pdf).
+It means that rather than defining a Material using arbitrary properties (such as diffuse color and specular power), you instead use properties more easily relatable to the real world. These include Base Color, Metallic, Specular, and Roughness. 
+
+###Quick Tips for Material Assignments to Static Meshes
+
+There are a couple of things you can do to quickly locate and assign your Materials to your Static Mesh from the Content Browser. 
+
+* **Assign a Material from your Content Browser** - Select your Material in the Content Browser and then in the Details panel, click the **arrow** next to the Material slot to assign the selected Material. 
+
+	![image alt text](image_24.png)(w:445 h:256 convert:false)
+
+* **Locate a Material from a Mesh in your Level** - Click the **Search** (magnifying glass) next to the Material Slot to instantly navigate to that Material's location in the Content Browser.
+	
+	![image alt text](image_25.png)(w:446 h:257 convert:false)
+
+* **Reset to Default** - Click the **Reset to Default** (yellow) arrow next to the Material Slot to set the Material back to the one assigned to the mesh as the default one. 
+
+
+Below are some links to get you started with using Materials and Textures in UE4:
+
+* [Texture Editor](Engine/Content/Types/Textures/Properties/Interface) - A reference guide for the Texture Editor.
+
+* [Material Editor](Engine/Rendering/Materials/Editor) - A reference guide to the Material Editor used in Unreal Engine 4.
+
+* [Essential Material Concepts](Engine/Rendering/Materials/IntroductionToMaterials) - An introductory overview of Unreal Engine 4 Materials and how they work.
+
+* [Using the Main Material Node](Engine/Rendering/Materials/HowTo/Main_Material_Node) - An overview of the Main Material Node used to create your Shaders in Unreal Engine 4.
+
+* [Material How-To's](Engine/Rendering/Materials/HowTo) - A list of topics related to Materials using different aspects of the Material Editor.
+
+* [Physically Based Materials](Engine/Rendering/Materials/PhysicallyBased) - An overview of the primary Material inputs and how best to use them using a physically based shading model.
+
+##Blueprints and Scripting
+
+The **Blueprint Visual Scripting** system in Unreal Engine is a complete gameplay scripting system based on the concept of using a node-based interface to create gameplay elements from within the Unreal Editor. As with many common scripting languages, it is used to define object-oriented classes or objects in the Engine. As you use UE4, you'll often find that objects defined as using Blueprint are colloquially referred to as just "Blueprints."
+
+The System is extremely flexible and powerful as it provides the ability for designers to use virtually the full range of concepts and tools only available to programmers. 
+
+![image alt text](image_26.png)(w:786 h:441 convert:false)
+
+Below are some links to get your started with using Blueprints:
+
+*  [Introduction to Blueprints](Engine/Blueprints/GettingStarted) - If you are just getting started with Blueprints, this provides a high-level overview of what they are and what they can do.
+
+* [Blueprint Quick Start](Engine/Blueprints/QuickStart) - A guide to help you get up and running by creating your first Blueprint. 
+
+* [Blueprint How-To's ](Engine/Blueprints/BP_HowTo) and [Video Tutorials](https://docs.unrealengine.com/latest/INT/Videos/Blueprints) - A collection of tutorials written and in video form that can help you learn more about visual scripting.
+
+##Sequencer
+
+The **Sequencer** Editor gives users the ability to create in-game cinematics with a specialized multi-track editor. It is the primary means of authoring cinematics using tracks which are used to define the makeup of your cinematic scenes. For example, Sequencer could be used to have multiple cameras slowly panning across an architectural visualization scene, or you could use [Composure](Engine/Composure) with Sequencer for compositing. 
+
+Similar to other video editing software, you can create camera tracks that can be dragged and rearranged as needed to quickly edit your scene together.
+
+![image alt text](Col05.jpg)(w:781 h:433 convert:false)
+[REGION:caption]
+_"Church of Light" scene created by Fabrice Bourrelly._
+[/REGION]
+
+Below are some links to get you started with using Sequencer: 
+
+* [Sequencer](Engine/Sequencer/Overview) -  An overview of the Sequencer Editor
+
+* [Sequencer Quick Start](Engine/Sequencer/QuickStart) - A guide to quickly get up and running with the Sequencer Editor.
+
+* [Composure Editor](Engine/Composure) - A compositing editor used for compositing digital and real-world footage together in real-time in UE4.
+
+##Building and Publishing your Project
+
+When you're ready to start testing or even ready for creating a final release of your project to hand off to anyone to open on their machines, you'll want to use the Unreal Build System to Cook and Package up your project into a standalone executable. This process enables you to easily share a project with other people in a way that is manageable and doesn't require them to have  Unreal Engine 4 already installed on their machine. It also protects your content by storing all your assets in a Pak file uses by Unreal to reference your assets when the executable is launched. 
+
+The process to Cook and Package your project performs the following actions: 
+
+1. When you select a target platform from the** File > Package Project** menu, the Cook and Package process will be kicked off. 
+
+2. The maps in your project will grab any referenced assets from the Content Browser and make copies of them in a staging directory. 
+
+3. Once the Cook process finishes, the Unreal Build Tool will start packaging those files and necessary Engine files into an executable that can be run without having to have Unreal Engine 4 installed. 
+
+4. Once finished, the project executable that is created can be run to load your a packaged version of your project. Assets are stored in a single Pak file, and any source code or Blueprint scripting has been take out and turned into machine code that is no longer accessible. 
+
+This is a very high-level view of what happens when you package a project. For a more detailed walkthrough and explanation of the process, see the [Packaging Projects](Engine/Basics/Projects/Packaging) page.
+
+[REGION:note]
+Cooking and Packaging content for some platforms, like Android or iOS, will require certain prerequisite SDKs or hardware to be present to package the project appropriately.
+[/REGION]
+
+
+
+
+##Where to go from here?
+
+Thanks for reading this guide! This was created for the Enterprise community with the help of Unreal Developers, and we very much appreciate any feedback and corrections you can offer. We will try to keep improving this document as we learn more about what is most helpful when transitioning to UE4.
+
+We have many additional Unreal Engine 4 learning resources available that you can use. These are linked below: 
+
+* [Documentation](%ROOT%)
+    *  [Getting Started ](GettingStarted)
+* [Tutorial and Training Videos](https://docs.unrealengine.com/latest/INT/Videos)
+    *  [Learning Management System](https://www.unrealengine.com/video-tutorials) 
+* API Reference - [\[C++\]](https://docs.unrealengine.com/latest/INT/API) [\[Blueprint\]](https://docs.unrealengine.com/latest/INT/BlueprintAPI)
+
diff --git a/Engine/Documentation/Source/Gameplay/Framework/GameplayFramework.INT.udn b/Engine/Documentation/Source/Gameplay/Framework/GameplayFramework.INT.udn
--- a/Engine/Documentation/Source/Gameplay/Framework/GameplayFramework.INT.udn
+++ b/Engine/Documentation/Source/Gameplay/Framework/GameplayFramework.INT.udn
@@ -30,33 +30,5 @@ to pages with overviews of each class and information about how they can be set
 
 [EXCERPT:LandingBoxes]
 [DIR(output:"topic" parent:"Gameplay/Framework" org:"hierarchy" end:"1")]
-[COMMENT:none]
-
-%Gameplay/Framework/GameFlow:topic%
-%Gameplay/Framework/GameMode:topic%
-%Gameplay/Framework/Pawn:topic%
-%Gameplay/Framework/Controller:topic%
-%Gameplay/Framework/Camera:topic%
-%Gameplay/Framework/UIAndHUD:topic%
-
-[/COMMENT]
 [/EXCERPT:LandingBoxes]
 
-
-[COMMENT:none]
-
-
-* [](Programming/UnrealArchitecture/Actors/Components) - %Programming/UnrealArchitecture/Actors/Components:Description%
-* [](Programming/UnrealArchitecture) - %Programming/UnrealArchitecture:Description%
-
-[EXCERPT:Systems]
-
-## [](Gameplay/Framework)
-
-[INCLUDE:Gameplay/Framework#Intro]
-[INCLUDE:Gameplay/Framework#LandingBoxes]
-
-[/EXCERPT:Systems]
-
-[/COMMENT]
-
diff --git a/Engine/Documentation/Source/Gameplay/HowTo/SpawnAndDestroyActors/GHT_1_Overview.INT.udn b/Engine/Documentation/Source/Gameplay/HowTo/SpawnAndDestroyActors/GHT_1_Overview.INT.udn
--- a/Engine/Documentation/Source/Gameplay/HowTo/SpawnAndDestroyActors/GHT_1_Overview.INT.udn
+++ b/Engine/Documentation/Source/Gameplay/HowTo/SpawnAndDestroyActors/GHT_1_Overview.INT.udn
@@ -8,9 +8,7 @@ type:overview
 order:1
 SkillLevel:Intermediate
 
-Often times in your game you will find that you want to add items to your levels dynamically during gameplay, this is called Spawning an Actor. Items you want to spawn could be anything from other players or enemies to pick-ups such as Health, Weapons, or Rewards. On this page, you will find information regarding how to spawn an Actor through Blueprints.[COMMENT:none]
- as well as other information regarding **Actor Spawning** and links to more documentation.
-[/COMMENT] 
+Often times in your game you will find that you want to add items to your levels dynamically during gameplay, this is called Spawning an Actor. Items you want to spawn could be anything from other players or enemies to pick-ups such as Health, Weapons, or Rewards. On this page, you will find information regarding how to spawn an Actor through Blueprints.
 
 ## Implementation Guides
 
diff --git a/Engine/Documentation/Source/Gameplay/Input/Input.INT.udn b/Engine/Documentation/Source/Gameplay/Input/Input.INT.udn
--- a/Engine/Documentation/Source/Gameplay/Input/Input.INT.udn
+++ b/Engine/Documentation/Source/Gameplay/Input/Input.INT.udn
@@ -28,20 +28,9 @@ PlayerInput is a UObject within the PlayerController class that manages player i
 defined within PlayerInput.  The first, **FInputActionKeyMapping**, defines an ActionMapping. The other, **FInputAxisKeyMapping**, defines an AxisMapping. 
 The hardware input definitions used in both ActionMappings and AxisMappings are established in InputCoreTypes.
 
-[EXCERPT:ActionMapping]
-$ ActionMappings : Map a discrete button or key press to a "friendly name" that will later be bound to event-driven behavior. The end effect is that pressing (and/or releasing) a key, mouse button, or keypad button
-directly triggers some game behavior.   
-[/EXCERPT:ActionMapping]
-
-[EXCERPT:AxisMapping]
-$ AxisMappings : Map keyboard, controller, or mouse inputs to a "friendly name" that will later be bound to continuous game behavior, such as movement.  The inputs mapped in AxisMappings are continuously polled, even if they are just reporting that their input value 
-is currently zero.  This allows for smooth transitions in movement or other game behavior, rather than the discrete game events triggered by inputs in ActionMappings.  
+[INCLUDE:#actionmapping]
 
-	[REGION:tip]
-	Hardware axes, such as controller joysticks, provide degrees of input, rather than discrete 1 (pressed) or 0 (not pressed) input.  That is, they can be moved to a small degree or a large degree, and 
-	your character's movement can vary accordingly. While these input methods are ideal for providing scalable amounts of movement input, AxisMappings can also map common movement keys, like WASD or Up, Down, Left, Right, to continuously-polled game behavior. 
-	[/REGION:tip]
-[/EXCERPT:AxisMapping]
+[INCLUDE:#axismapping]
 
 ### Setting Input Mappings
 
@@ -66,6 +55,23 @@ In this window, you can:
 ![](AxisMappings.png)
 [/EXCERPT:PlayerInput]
 
+[COMMENT:none]
+[EXCERPT:ActionMapping]
+$ ActionMappings : Map a discrete button or key press to a "friendly name" that will later be bound to event-driven behavior. The end effect is that pressing (and/or releasing) a key, mouse button, or keypad button
+directly triggers some game behavior.   
+[/EXCERPT:ActionMapping]
+
+[EXCERPT:AxisMapping]
+$ AxisMappings : Map keyboard, controller, or mouse inputs to a "friendly name" that will later be bound to continuous game behavior, such as movement.  The inputs mapped in AxisMappings are continuously polled, even if they are just reporting that their input value 
+is currently zero.  This allows for smooth transitions in movement or other game behavior, rather than the discrete game events triggered by inputs in ActionMappings.  
+
+	[REGION:tip]
+	Hardware axes, such as controller joysticks, provide degrees of input, rather than discrete 1 (pressed) or 0 (not pressed) input.  That is, they can be moved to a small degree or a large degree, and 
+	your character's movement can vary accordingly. While these input methods are ideal for providing scalable amounts of movement input, AxisMappings can also map common movement keys, like WASD or Up, Down, Left, Right, to continuously-polled game behavior. 
+	[/REGION:tip]
+[/EXCERPT:AxisMapping]
+[/COMMENT]
+
 ## InputComponent
 **InputComponents** are most commonly present in Pawns and Controllers, although they can be set in other Actors and Level Scripts if desired.  The InputComponent links the AxisMappings and ActionMappings in your
 project to game actions, usually functions, set up either in C++ code or Blueprint graphs.  
diff --git a/Engine/Documentation/Source/GettingStarted/RecommendedSpecifications/Specifications.INT.udn b/Engine/Documentation/Source/GettingStarted/RecommendedSpecifications/Specifications.INT.udn
--- a/Engine/Documentation/Source/GettingStarted/RecommendedSpecifications/Specifications.INT.udn
+++ b/Engine/Documentation/Source/GettingStarted/RecommendedSpecifications/Specifications.INT.udn
@@ -14,9 +14,7 @@
 [TOC(start:2)]
 
 [EXCERPT:Main]
-[EXCERPT:Intro]
-Unreal Engine has some specific hardware and software requirements for running the game and editor or developing with the engine.
-[/EXCERPT:Intro]
+[INCLUDE:#Intro]
 
 ## Recommended Hardware
 
@@ -61,35 +59,6 @@
 	[/PARAM]
 [/OBJECT]
 
-[COMMENT:none]
 
-[EXCERPT:WindowsMinHardware]
-|||
-| -------- | --- |
-| **Operating System** | Windows 7/8 64-bit |
-| **Processor** | Quad-core Intel or AMD, 2.5 GHz or faster |
-| **Memory** | 8 GB RAM |
-| **Video Card/DirectX Version** | DirectX 11 compatible graphics card |
-[/EXCERPT:WindowsMinHardware]
-[EXCERPT:MacOSMinHardware]
-|||
-| -------- | --- |
-| **Operating System** | macOS 10.13 High Sierra  |
-| **Processor** | Quad-core Intel, 2.5 GHz or faster |
-| **Memory** | 8 GB RAM |
-| **Video Card** | Metal 1.2 Compatible Graphics Card |
-[/EXCERPT:MacOSMinHardware]
-[EXCERPT:LinuxMinHardware]
-|||
-| -------- | --- |
-| **Operating System** | Ubuntu 15.04  |
-| **Processor** | Quad-core Intel or AMD, 2.5 GHz or faster |
-| **Memory** | 16 GB RAM |
-| **Video Card** | 	NVIDIA GeForce 470 GTX or higher with latest NVIDIA binary drivers |
-[/EXCERPT:LinuxMinHardware]
-
-[/COMMENT]
-
-
 [PUBLISH:Licensee]
 [REGION:note]
@@ -150,43 +119,6 @@
 [/REGION]
 [/PUBLISH:Licensee]
 
-[COMMENT:none]
 
-[EXCERPT:WindowsMinSoftware]
-| Running the Engine ||
-| -------- | --- |
-| **Operating System** | Windows 7/8 64-bit |
-| **DirectX Runtime** | [DirectX End-User Runtimes (June 2010)](http://www.microsoft.com/en-us/download/details.aspx?id=8109) |
-
-[/EXCERPT:WindowsMinSoftware]
-
-[EXCERPT:MacOSMinSoftware]
-| Running the Engine ||
-| -------- | --- |
-| **Operating System** | macOS 10.12.6 Sierra  |
-[/EXCERPT:MacOSMinSoftware]
-
-[EXCERPT:LinuxMinSoftware]
-| Running the Engine ||
-| -------- | --- |
-| **Operating System** | Any reasonable new Linux distro (CentOS 6.x being the oldest tested)  |
-| **Linux Kernel Version** | kernel 2.6.32 or newer  |
-| **Additional Dependencies** | glibc 2.12.2 or newer  |
-[/EXCERPT:LinuxMinSoftware]
-
-[EXCERPT:WER]
-[REGION:note]
-UE4's crash reporting feature requires you to have **Windows Error Reporting** enabled, which is the default. If you have disabled this feature, please enable it.
-[/REGION]
-[/EXCERPT:WER]
-
-[EXCERPT:VCRuntime]
-[Visual Studio C++ 2015 Redist](https://www.microsoft.com/en-us/download/details.aspx?id=48145)
-[/EXCERPT:VCRuntime]
-
-
-[/COMMENT]
-
-
 The requirements for programmers developing with the engine are listed below.
 
@@ -231,44 +163,7 @@
 	[/PARAM]
 [/OBJECT]
 
-[COMMENT:none]
 
-[EXCERPT:WindowsMinDevelopment]
-| Developing with the Engine ||
-| -------- | --- |
-| **All 'Running the Engine' requirements (automatically installed)** ||
-| **Visual Studio Version** | [INCLUDE:#VS] |
-| [REGION:tablesection]iOS Development[/REGION] ||
-| **iTunes Version** | [iTunes 11 or higher](http://www.apple.com/itunes/) |
-[/EXCERPT:WindowsMinDevelopment]
-[EXCERPT:MacOSMinDevelopment]
-| Developing with the Engine ||
-| -------- | --- |
-| **Xcode Version** | 9.0 [REGION:note]If you need to verify that you're running the correct version of Xcode, run the following command from the terminal window: `sudo xcode-select -s /Applications/Xcode-beta.app/Contents/Developer`[/REGION] |
-[/EXCERPT:MacOSMinDevelopment]
-[EXCERPT:LinuxMinDevelopment]
-| Developing with the Engine ||
-| -------- | --- |
-| **Operating System** | Ubuntu 15.04  |
-| **Compiler** | clang 3.5.x or 3.6.x  |
-| [REGION:tablesection]Optional[/REGION] ||
-| **IDE** | QtCreator 3.1 or KDevelop 4.6|
-| **Build System** | CMake 2.8 |
-[/EXCERPT:LinuxMinDevelopment]
 
-[/COMMENT]
-
-[COMMENT:none]
-
-[EXCERPT:VS]
-Visual Studio 2015 Professional or Visual Studio 2015 Community:
-
-* [Pro](https://www.visualstudio.com/downloads/download-visual-studio-vs)
-* [Community](https://www.visualstudio.com/en-us/products/visual-studio-community-vs.aspx)
-
-[/EXCERPT:VS]
-
-[/COMMENT]
-
 [PUBLISH:Licensee]
 ## Graphics Card Drivers
@@ -310,3 +205,93 @@
 [/REGION]
 [/PUBLISH:Licensee]
 [/EXCERPT:Main]
+
+[COMMENT:none]
+[EXCERPT:Intro]
+Unreal Engine has some specific hardware and software requirements for running the game and editor or developing with the engine.
+[/EXCERPT:Intro]
+[EXCERPT:VS]
+Visual Studio 2015 Professional or Visual Studio 2015 Community:
+
+* [Pro](https://www.visualstudio.com/downloads/download-visual-studio-vs)
+* [Community](https://www.visualstudio.com/en-us/products/visual-studio-community-vs.aspx)
+
+[/EXCERPT:VS]
+[EXCERPT:WindowsMinDevelopment]
+| Developing with the Engine ||
+| -------- | --- |
+| **All 'Running the Engine' requirements (automatically installed)** ||
+| **Visual Studio Version** | [INCLUDE:#VS] |
+| [REGION:tablesection]iOS Development[/REGION] ||
+| **iTunes Version** | [iTunes 11 or higher](http://www.apple.com/itunes/) |
+[/EXCERPT:WindowsMinDevelopment]
+[EXCERPT:MacOSMinDevelopment]
+| Developing with the Engine ||
+| -------- | --- |
+| **Xcode Version** | 9.0 [REGION:note]If you need to verify that you're running the correct version of Xcode, run the following command from the terminal window: `sudo xcode-select -s /Applications/Xcode-beta.app/Contents/Developer`[/REGION] |
+[/EXCERPT:MacOSMinDevelopment]
+[EXCERPT:LinuxMinDevelopment]
+| Developing with the Engine ||
+| -------- | --- |
+| **Operating System** | Ubuntu 15.04  |
+| **Compiler** | clang 3.5.x or 3.6.x  |
+| [REGION:tablesection]Optional[/REGION] ||
+| **IDE** | QtCreator 3.1 or KDevelop 4.6|
+| **Build System** | CMake 2.8 |
+[/EXCERPT:LinuxMinDevelopment]
+[EXCERPT:WindowsMinSoftware]
+| Running the Engine ||
+| -------- | --- |
+| **Operating System** | Windows 7/8 64-bit |
+| **DirectX Runtime** | [DirectX End-User Runtimes (June 2010)](http://www.microsoft.com/en-us/download/details.aspx?id=8109) |
+
+[/EXCERPT:WindowsMinSoftware]
+
+[EXCERPT:MacOSMinSoftware]
+| Running the Engine ||
+| -------- | --- |
+| **Operating System** | macOS 10.12.6 Sierra  |
+[/EXCERPT:MacOSMinSoftware]
+
+[EXCERPT:LinuxMinSoftware]
+| Running the Engine ||
+| -------- | --- |
+| **Operating System** | Any reasonable new Linux distro (CentOS 6.x being the oldest tested)  |
+| **Linux Kernel Version** | kernel 2.6.32 or newer  |
+| **Additional Dependencies** | glibc 2.12.2 or newer  |
+[/EXCERPT:LinuxMinSoftware]
+
+[EXCERPT:WER]
+[REGION:note]
+UE4's crash reporting feature requires you to have **Windows Error Reporting** enabled, which is the default. If you have disabled this feature, please enable it.
+[/REGION]
+[/EXCERPT:WER]
+
+[EXCERPT:VCRuntime]
+[Visual Studio C++ 2015 Redist](https://www.microsoft.com/en-us/download/details.aspx?id=48145)
+[/EXCERPT:VCRuntime]
+[EXCERPT:WindowsMinHardware]
+|||
+| -------- | --- |
+| **Operating System** | Windows 7/8 64-bit |
+| **Processor** | Quad-core Intel or AMD, 2.5 GHz or faster |
+| **Memory** | 8 GB RAM |
+| **Video Card/DirectX Version** | DirectX 11 compatible graphics card |
+[/EXCERPT:WindowsMinHardware]
+[EXCERPT:MacOSMinHardware]
+|||
+| -------- | --- |
+| **Operating System** | macOS 10.13 High Sierra  |
+| **Processor** | Quad-core Intel, 2.5 GHz or faster |
+| **Memory** | 8 GB RAM |
+| **Video Card** | Metal 1.2 Compatible Graphics Card |
+[/EXCERPT:MacOSMinHardware]
+[EXCERPT:LinuxMinHardware]
+|||
+| -------- | --- |
+| **Operating System** | Ubuntu 15.04  |
+| **Processor** | Quad-core Intel or AMD, 2.5 GHz or faster |
+| **Memory** | 16 GB RAM |
+| **Video Card** | 	NVIDIA GeForce 470 GTX or higher with latest NVIDIA binary drivers |
+[/EXCERPT:LinuxMinHardware]
+[/COMMENT]
diff --git a/Engine/Documentation/Source/GettingStarted/RunningUnrealEngine/RunningUnrealEngine.INT.udn b/Engine/Documentation/Source/GettingStarted/RunningUnrealEngine/RunningUnrealEngine.INT.udn
--- a/Engine/Documentation/Source/GettingStarted/RunningUnrealEngine/RunningUnrealEngine.INT.udn
+++ b/Engine/Documentation/Source/GettingStarted/RunningUnrealEngine/RunningUnrealEngine.INT.udn
@@ -42,159 +42,8 @@
 
 
 ## Running the Editor
-[EXCERPT:RunEditor]
-[PUBLISH:Licensee]
 
-This process will launch a copy of Unreal Editor.
+[INCLUDE:#RunEditor]
 
-[/PUBLISH:Licensee]
-
-[PUBLISH:Rocket]
-
-The process of running the editor simply requires passing the name of the project to run as an argument to the executable.
-
-[REGION:note]
-You can add the `-debug` switch to force the executable to load the debug version of the modules for your project, which 
-contain all of the debugging symbols. This is necessary even when debugging through Visual Studio with the configuration 
-set to **Debug** because the main executable is always compiled using the **Development** configuration. Of course, you 
-must first compile your modules using the Debug configuration so they exist for the executable to load.
-[/REGION]
-
-[/PUBLISH:Rocket]
-
-### Running the Editor From the Command Line
-
-[PUBLISH:Licensee]
-1.	From a command prompt, navigate to the `\Engine\Binaries\Win64` directory of your `UE4` directory.
-
-1.	Run the `UE4Editor.exe` passing it the name of the game to run:
-
-|Project Location|Command Prompt|
-|----------------|--------------|
-|UE4 Root | `UE4Editor.exe MyGame`|
-|Alternate Location| `UE4Editor.exe "[ProjectPath]\[ProjectName].uproject"`|
-
-
-
-[/PUBLISH:Licensee]
-
-[PUBLISH:Rocket]
-[OBJECT:ToggleButtonContent]
-	[PARAMLITERAL:category]
-	OS
-	[/PARAMLITERAL]
-	[PARAMLITERAL:id]
-	windows
-	[/PARAMLITERAL]
-	[PARAMLITERAL:active]
-	active_button_content
-	[/PARAMLITERAL]
-	[PARAM:content]
-	1.	From a command prompt, navigate to your `[LauncherInstall]\[VersionNumber]\Engine\Binaries\Win64` directory.
-
-	1.	Run the `UE4Editor.exe` passing it the path to your project:
-
-			UE4Editor.exe "[ProjectPath]\[ProjectName].uproject"
-	[/PARAM]
-[/OBJECT]
-[OBJECT:ToggleButtonContent]
-	[PARAMLITERAL:category]
-	OS
-	[/PARAMLITERAL]
-	[PARAMLITERAL:id]
-	mac
-	[/PARAMLITERAL]
-	[PARAMLITERAL:active]
-	[/PARAMLITERAL]
-	[PARAM:content]
-	1.	From a command prompt, navigate to your `[LauncherInstall]/[VersionNumber]/Engine/Binaries/Mac` directory.
-
-	1.	Run the `UE4Editor.app` passing it the path to your project:
-
-			open UE4Editor.app --args "[ProjectPath]/[ProjectName].uproject"
-	[/PARAM]
-[/OBJECT]
-[/PUBLISH:Rocket]
-
-
-### Running the Editor From the Executable
-
-[PUBLISH:Licensee]
-
-
-1.	Navigate to the `\Engine\Binaries\Win64` directory of your `UE4` directory.
-
-1.	**Right-click** on the `UE4Editor.exe` executable and choose **Create shortcut**.
-
-1.	Rename the shortcut to reflect the game it will run and the fact that it will launch the editor, i.e. **MyGameEditor.exe**.
-
-1.	**Right-click** on the newly created shortcut and choose **Properties**.
-
-1.	Add the name of the game to run as an argument at the end of the **Target** property:
-
-|Project Location|Target|
-|----------------|--------------|
-|UE4 Root | `[UE4Directory]\Engine\Binaries\Win64\UE4Editor.exe MyGame`|
-|Alternate Location| `[UE4Directory]\Engine\Binaries\Win64\UE4Editor.exe "[ProjectPath]\[ProjectName].uproject"`|
-
-
-1. Press **OK** to save the changes.
-
-1.	**Double-click** the shortcut to launch the editor.
-
-[/PUBLISH:Licensee]
-
-[PUBLISH:Rocket]
-[OBJECT:ToggleButtonContent]
-	[PARAMLITERAL:category]
-	OS
-	[/PARAMLITERAL]
-	[PARAMLITERAL:id]
-	windows
-	[/PARAMLITERAL]
-	[PARAMLITERAL:active]
-	active_button_content
-	[/PARAMLITERAL]
-	[PARAM:content]
-	1.	Navigate to your `[LauncherInstall]\[VersionNumber]\Engine\Binaries\Win64` directory.
-
-	1.	**Right-click** on the `UE4Editor.exe` executable and choose **Create shortcut**.
-
-	1.	Rename the shortcut to reflect the game it will run and the fact that it will launch the editor, i.e. **MyProject - Editor.exe**.
-
-	1.	**Right-click** on the newly created shortcut and choose **Properties**.
-
-	1.	Add the name of the game to run as an argument at the end of the **Target** property:
-
-			[LauncherInstall]\[VersionNumber]\Engine\Binaries\Win64\UE4Editor.exe "[ProjectPath]\[ProjectName].uproject"
-
-	1. Press **OK** to save the changes.
-
-	1.	**Double-click** the shortcut to launch the editor.
-	[/PARAM]
-[/OBJECT]
-[OBJECT:ToggleButtonContent]
-	[PARAMLITERAL:category]
-	OS
-	[/PARAMLITERAL]
-	[PARAMLITERAL:id]
-	mac
-	[/PARAMLITERAL]
-	[PARAMLITERAL:active]
-	[/PARAMLITERAL]
-	[PARAM:content]
-	The editor must be run [from the command prompt](#RunningtheEditorFromtheCommandLine) to load a specific project directly or [with no arguments](#RunningtheEditorwithNoArguments_Stand-Alone_) to access the Project Browser.
-	[/PARAM]
-[/OBJECT]
-[/PUBLISH:Rocket]
-
-### Running the Editor with No Arguments (Stand-alone)
-
-If the editor is not set to open the most recent project at startup, running the editor executable without any arguments will launch the Project Browser. From here, you can 
-[create a new project](Engine/Basics/Projects/Browser#creatingnewprojects), [open your existing projects](Engine/Basics/Projects/Browser#openingprojects),
-or open [content examples and sample games](Resources).
-
-[/EXCERPT:RunEditor]
-
 ## Running an Uncooked Game
 
@@ -516,37 +365,160 @@
 
 	travel DM-Deck
 
+[/EXCERPT:Main]
 
 [COMMENT:none]
+[EXCERPT:RunEditor]
+[PUBLISH:Licensee]
+
+This process will launch a copy of Unreal Editor.
+
+[/PUBLISH:Licensee]
+
+[PUBLISH:Rocket]
+
+The process of running the editor simply requires passing the name of the project to run as an argument to the executable.
+
+[REGION:note]
+You can add the `-debug` switch to force the executable to load the debug version of the modules for your project, which 
+contain all of the debugging symbols. This is necessary even when debugging through Visual Studio with the configuration 
+set to **Debug** because the main executable is always compiled using the **Development** configuration. Of course, you 
+must first compile your modules using the Debug configuration so they exist for the executable to load.
+[/REGION]
+
+[/PUBLISH:Rocket]
+
+### Running the Editor From the Command Line
 
 [PUBLISH:Licensee]
-## Running the Game with Cooked Content
+1.	From a command prompt, navigate to the `\Engine\Binaries\Win64` directory of your `UE4` directory.
 
-The executable for your game requires cooked content to run.
+1.	Run the `UE4Editor.exe` passing it the name of the game to run:
 
+|Project Location|Command Prompt|
+|----------------|--------------|
+|UE4 Root | `UE4Editor.exe MyGame`|
+|Alternate Location| `UE4Editor.exe "[ProjectPath]\[ProjectName].uproject"`|
 
-### Run from Command Line
-
-1.	From a command prompt, navigate to the `\Binaries\Win32` directory of your game's directory.
-
-1.	Run the **[GameName].exe**
-
-		[GameName].exe
-
-
-### Run from Executable
-
-
-1.	Navigate to the `\Binaries\Win32` folder of your game's directory.
-
-1.	**Double-click** the **[GameName].exe** to run the game.
 
 
 [/PUBLISH:Licensee]
 
+[PUBLISH:Rocket]
+[OBJECT:ToggleButtonContent]
+	[PARAMLITERAL:category]
+	OS
+	[/PARAMLITERAL]
+	[PARAMLITERAL:id]
+	windows
+	[/PARAMLITERAL]
+	[PARAMLITERAL:active]
+	active_button_content
+	[/PARAMLITERAL]
+	[PARAM:content]
+	1.	From a command prompt, navigate to your `[LauncherInstall]\[VersionNumber]\Engine\Binaries\Win64` directory.
+
+	1.	Run the `UE4Editor.exe` passing it the path to your project:
+
+			UE4Editor.exe "[ProjectPath]\[ProjectName].uproject"
+	[/PARAM]
+[/OBJECT]
+[OBJECT:ToggleButtonContent]
+	[PARAMLITERAL:category]
+	OS
+	[/PARAMLITERAL]
+	[PARAMLITERAL:id]
+	mac
+	[/PARAMLITERAL]
+	[PARAMLITERAL:active]
+	[/PARAMLITERAL]
+	[PARAM:content]
+	1.	From a command prompt, navigate to your `[LauncherInstall]/[VersionNumber]/Engine/Binaries/Mac` directory.
+
+	1.	Run the `UE4Editor.app` passing it the path to your project:
+
+			open UE4Editor.app --args "[ProjectPath]/[ProjectName].uproject"
+	[/PARAM]
+[/OBJECT]
+[/PUBLISH:Rocket]
+
+
+### Running the Editor From the Executable
+
+[PUBLISH:Licensee]
+
+
+1.	Navigate to the `\Engine\Binaries\Win64` directory of your `UE4` directory.
+
+1.	**Right-click** on the `UE4Editor.exe` executable and choose **Create shortcut**.
+
+1.	Rename the shortcut to reflect the game it will run and the fact that it will launch the editor, i.e. **MyGameEditor.exe**.
+
+1.	**Right-click** on the newly created shortcut and choose **Properties**.
+
+1.	Add the name of the game to run as an argument at the end of the **Target** property:
+
+|Project Location|Target|
+|----------------|--------------|
+|UE4 Root | `[UE4Directory]\Engine\Binaries\Win64\UE4Editor.exe MyGame`|
+|Alternate Location| `[UE4Directory]\Engine\Binaries\Win64\UE4Editor.exe "[ProjectPath]\[ProjectName].uproject"`|
+
+
+1. Press **OK** to save the changes.
+
+1.	**Double-click** the shortcut to launch the editor.
+
+[/PUBLISH:Licensee]
+
+[PUBLISH:Rocket]
+[OBJECT:ToggleButtonContent]
+	[PARAMLITERAL:category]
+	OS
+	[/PARAMLITERAL]
+	[PARAMLITERAL:id]
+	windows
+	[/PARAMLITERAL]
+	[PARAMLITERAL:active]
+	active_button_content
+	[/PARAMLITERAL]
+	[PARAM:content]
+	1.	Navigate to your `[LauncherInstall]\[VersionNumber]\Engine\Binaries\Win64` directory.
+
+	1.	**Right-click** on the `UE4Editor.exe` executable and choose **Create shortcut**.
+
+	1.	Rename the shortcut to reflect the game it will run and the fact that it will launch the editor, i.e. **MyProject - Editor.exe**.
+
+	1.	**Right-click** on the newly created shortcut and choose **Properties**.
+
+	1.	Add the name of the game to run as an argument at the end of the **Target** property:
+
+			[LauncherInstall]\[VersionNumber]\Engine\Binaries\Win64\UE4Editor.exe "[ProjectPath]\[ProjectName].uproject"
+
+	1. Press **OK** to save the changes.
+
+	1.	**Double-click** the shortcut to launch the editor.
+	[/PARAM]
+[/OBJECT]
+[OBJECT:ToggleButtonContent]
+	[PARAMLITERAL:category]
+	OS
+	[/PARAMLITERAL]
+	[PARAMLITERAL:id]
+	mac
+	[/PARAMLITERAL]
+	[PARAMLITERAL:active]
+	[/PARAMLITERAL]
+	[PARAM:content]
+	The editor must be run [from the command prompt](#RunningtheEditorFromtheCommandLine) to load a specific project directly or [with no arguments](#RunningtheEditorwithNoArguments_Stand-Alone_) to access the Project Browser.
+	[/PARAM]
+[/OBJECT]
+[/PUBLISH:Rocket]
+
+### Running the Editor with No Arguments (Stand-alone)
+
+If the editor is not set to open the most recent project at startup, running the editor executable without any arguments will launch the Project Browser. From here, you can 
+[create a new project](Engine/Basics/Projects/Browser#creatingnewprojects), [open your existing projects](Engine/Basics/Projects/Browser#openingprojects),
+or open [content examples and sample games](Resources).
+
+[/EXCERPT:RunEditor]
 [/COMMENT]
-
-
-
-
-[/EXCERPT:Main]
