Copying //UE4/Dev-Documentation to //UE4/Samples-Main (Source: //UE4/Dev-Documentation @ 4091886)

diff --git a/Engine/Documentation/Source/Modding/RoboRecall/HowTo/ExtendFunctionality/BuildGrabEvents/BuildGrabEvents.INT.udn b/Engine/Documentation/Source/Modding/RoboRecall/HowTo/ExtendFunctionality/BuildGrabEvents/BuildGrabEvents.INT.udn
--- a/Engine/Documentation/Source/Modding/RoboRecall/HowTo/ExtendFunctionality/BuildGrabEvents/BuildGrabEvents.INT.udn
+++ b/Engine/Documentation/Source/Modding/RoboRecall/HowTo/ExtendFunctionality/BuildGrabEvents/BuildGrabEvents.INT.udn
@@ -76,241 +76,232 @@ In this step you will be overriding the **Grab Events** on our GrenadeGunPistol
 		![image alt text](image_9.png)(w:786 convert:false)
 		[/PARAM]
 		[PARAM:snippet]
-			[INCLUDE:#snippet]
-		[/PARAM]
-	[/OBJECT]
-
-	[REGION:caption]
-	Click the icon in the upper left corner of this image to copy the Blueprint Graph and paste it into your project.
-	[/REGION]
+				Begin Object Class=K2Node_Event Name="K2Node_Event_7"
 
-9. In the Blueprint Editor toolbar, click **Compile** to update the Blueprint and click **Save** to preserve your changes.
-
-## Results
+					EventReference=(MemberParent=BlueprintGeneratedClass'/Game/Prototype/Game/Gun_Base.Gun_Base_C',MemberName="End Primary Grab",MemberGuid=5AC0E94444892A8794605F86BF8B03E2)
 
-These sets of Events are called when our gun is grabbed and released. You'll use the **bCanExplode** variable in the next section of the Blueprint to prevent the gun from exploding in your hand when you punch a robot with it.
+					bOverrideFunction=True
 
-[NAV]
- 
-[COMMENT:none]
+					NodePosX=48
 
-[EXCERPT:snippet]
-	Begin Object Class=K2Node_Event Name="K2Node_Event_7"
+					NodePosY=1584
 
-		EventReference=(MemberParent=BlueprintGeneratedClass'/Game/Prototype/Game/Gun_Base.Gun_Base_C',MemberName="End Primary Grab",MemberGuid=5AC0E94444892A8794605F86BF8B03E2)
+					NodeGuid=4DB300D94131B538C1B2E1AFF11EB6B2
 
-		bOverrideFunction=True
+					CustomProperties Pin (PinId=59D3C9964C20448F1FFD22B92B32E3F2,PinName="OutputDelegate",Direction="EGPD_Output",PinType.PinCategory="delegate",PinType.PinSubCategory="",PinType.PinSubCategoryObject=None,PinType.PinSubCategoryMemberReference=(MemberParent=BlueprintGeneratedClass'/Game/Prototype/Game/Gun_Base.Gun_Base_C',MemberName="End Primary Grab",MemberGuid=5AC0E94444892A8794605F86BF8B03E2),PinType.PinValueType=(),PinType.bIsMap=False,PinType.bIsSet=False,PinType.bIsArray=False,PinType.bIsReference=False,PinType.bIsConst=False,PinType.bIsWeakPointer=False,PersistentGuid=00000000000000000000000000000000,bHidden=False,bNotConnectable=False,bDefaultValueIsReadOnly=False,bDefaultValueIsIgnored=False,bAdvancedView=False,)
 
-		NodePosX=48
+					CustomProperties Pin (PinId=DDE6DFF14B7B9ED8FABA8A8A2B48E090,PinName="then",Direction="EGPD_Output",PinType.PinCategory="exec",PinType.PinSubCategory="",PinType.PinSubCategoryObject=None,PinType.PinSubCategoryMemberReference=(),PinType.PinValueType=(),PinType.bIsMap=False,PinType.bIsSet=False,PinType.bIsArray=False,PinType.bIsReference=False,PinType.bIsConst=False,PinType.bIsWeakPointer=False,LinkedTo=(K2Node_CallParentFunction_3 440FA2384D499E510943C5952AED1D52,),PersistentGuid=00000000000000000000000000000000,bHidden=False,bNotConnectable=False,bDefaultValueIsReadOnly=False,bDefaultValueIsIgnored=False,bAdvancedView=False,)
 
-		NodePosY=1584
+				End Object
 
-		NodeGuid=4DB300D94131B538C1B2E1AFF11EB6B2
+				Begin Object Class=K2Node_Event Name="K2Node_Event_10"
 
-		CustomProperties Pin (PinId=59D3C9964C20448F1FFD22B92B32E3F2,PinName="OutputDelegate",Direction="EGPD_Output",PinType.PinCategory="delegate",PinType.PinSubCategory="",PinType.PinSubCategoryObject=None,PinType.PinSubCategoryMemberReference=(MemberParent=BlueprintGeneratedClass'/Game/Prototype/Game/Gun_Base.Gun_Base_C',MemberName="End Primary Grab",MemberGuid=5AC0E94444892A8794605F86BF8B03E2),PinType.PinValueType=(),PinType.bIsMap=False,PinType.bIsSet=False,PinType.bIsArray=False,PinType.bIsReference=False,PinType.bIsConst=False,PinType.bIsWeakPointer=False,PersistentGuid=00000000000000000000000000000000,bHidden=False,bNotConnectable=False,bDefaultValueIsReadOnly=False,bDefaultValueIsIgnored=False,bAdvancedView=False,)
+					EventReference=(MemberParent=BlueprintGeneratedClass'/Game/Prototype/Game/Gun_Base.Gun_Base_C',MemberName="Begin Primary Grab",MemberGuid=86E248C047254FEECC0603BB6C0FE484)
 
-		CustomProperties Pin (PinId=DDE6DFF14B7B9ED8FABA8A8A2B48E090,PinName="then",Direction="EGPD_Output",PinType.PinCategory="exec",PinType.PinSubCategory="",PinType.PinSubCategoryObject=None,PinType.PinSubCategoryMemberReference=(),PinType.PinValueType=(),PinType.bIsMap=False,PinType.bIsSet=False,PinType.bIsArray=False,PinType.bIsReference=False,PinType.bIsConst=False,PinType.bIsWeakPointer=False,LinkedTo=(K2Node_CallParentFunction_3 440FA2384D499E510943C5952AED1D52,),PersistentGuid=00000000000000000000000000000000,bHidden=False,bNotConnectable=False,bDefaultValueIsReadOnly=False,bDefaultValueIsIgnored=False,bAdvancedView=False,)
+					bOverrideFunction=True
 
-	End Object
+					NodePosX=48
 
-	Begin Object Class=K2Node_Event Name="K2Node_Event_10"
+					NodePosY=1200
 
-		EventReference=(MemberParent=BlueprintGeneratedClass'/Game/Prototype/Game/Gun_Base.Gun_Base_C',MemberName="Begin Primary Grab",MemberGuid=86E248C047254FEECC0603BB6C0FE484)
+					NodeGuid=2698FDBE4F93DAEDA8C18FAAB85204CC
 
-		bOverrideFunction=True
+					CustomProperties Pin (PinId=D5884B0048B127BF178702BB8251314F,PinName="OutputDelegate",Direction="EGPD_Output",PinType.PinCategory="delegate",PinType.PinSubCategory="",PinType.PinSubCategoryObject=None,PinType.PinSubCategoryMemberReference=(MemberParent=BlueprintGeneratedClass'/Game/Prototype/Game/Gun_Base.Gun_Base_C',MemberName="Begin Primary Grab",MemberGuid=86E248C047254FEECC0603BB6C0FE484),PinType.PinValueType=(),PinType.bIsMap=False,PinType.bIsSet=False,PinType.bIsArray=False,PinType.bIsReference=False,PinType.bIsConst=False,PinType.bIsWeakPointer=False,PersistentGuid=00000000000000000000000000000000,bHidden=False,bNotConnectable=False,bDefaultValueIsReadOnly=False,bDefaultValueIsIgnored=False,bAdvancedView=False,)
 
-		NodePosX=48
+					CustomProperties Pin (PinId=01317E634C1C07F9D5D729B8421618F2,PinName="then",Direction="EGPD_Output",PinType.PinCategory="exec",PinType.PinSubCategory="",PinType.PinSubCategoryObject=None,PinType.PinSubCategoryMemberReference=(),PinType.PinValueType=(),PinType.bIsMap=False,PinType.bIsSet=False,PinType.bIsArray=False,PinType.bIsReference=False,PinType.bIsConst=False,PinType.bIsWeakPointer=False,LinkedTo=(K2Node_CallParentFunction_0 FDB36E2741A2BCEF856502AAC746B07F,),PersistentGuid=00000000000000000000000000000000,bHidden=False,bNotConnectable=False,bDefaultValueIsReadOnly=False,bDefaultValueIsIgnored=False,bAdvancedView=False,)
 
-		NodePosY=1200
+				End Object
 
-		NodeGuid=2698FDBE4F93DAEDA8C18FAAB85204CC
+				Begin Object Class=K2Node_CallParentFunction Name="K2Node_CallParentFunction_0"
 
-		CustomProperties Pin (PinId=D5884B0048B127BF178702BB8251314F,PinName="OutputDelegate",Direction="EGPD_Output",PinType.PinCategory="delegate",PinType.PinSubCategory="",PinType.PinSubCategoryObject=None,PinType.PinSubCategoryMemberReference=(MemberParent=BlueprintGeneratedClass'/Game/Prototype/Game/Gun_Base.Gun_Base_C',MemberName="Begin Primary Grab",MemberGuid=86E248C047254FEECC0603BB6C0FE484),PinType.PinValueType=(),PinType.bIsMap=False,PinType.bIsSet=False,PinType.bIsArray=False,PinType.bIsReference=False,PinType.bIsConst=False,PinType.bIsWeakPointer=False,PersistentGuid=00000000000000000000000000000000,bHidden=False,bNotConnectable=False,bDefaultValueIsReadOnly=False,bDefaultValueIsIgnored=False,bAdvancedView=False,)
+					FunctionReference=(MemberParent=BlueprintGeneratedClass'/Game/Prototype/Game/Gun_Pistol.Gun_Pistol_C',MemberName="Begin Primary Grab",MemberGuid=86E248C047254FEECC0603BB6C0FE484)
 
-		CustomProperties Pin (PinId=01317E634C1C07F9D5D729B8421618F2,PinName="then",Direction="EGPD_Output",PinType.PinCategory="exec",PinType.PinSubCategory="",PinType.PinSubCategoryObject=None,PinType.PinSubCategoryMemberReference=(),PinType.PinValueType=(),PinType.bIsMap=False,PinType.bIsSet=False,PinType.bIsArray=False,PinType.bIsReference=False,PinType.bIsConst=False,PinType.bIsWeakPointer=False,LinkedTo=(K2Node_CallParentFunction_0 FDB36E2741A2BCEF856502AAC746B07F,),PersistentGuid=00000000000000000000000000000000,bHidden=False,bNotConnectable=False,bDefaultValueIsReadOnly=False,bDefaultValueIsIgnored=False,bAdvancedView=False,)
+					NodePosX=320
 
-	End Object
+					NodePosY=1200
 
-	Begin Object Class=K2Node_CallParentFunction Name="K2Node_CallParentFunction_0"
+					NodeGuid=F79E93794CF4A9FF0F133187E49F3FF4
 
-		FunctionReference=(MemberParent=BlueprintGeneratedClass'/Game/Prototype/Game/Gun_Pistol.Gun_Pistol_C',MemberName="Begin Primary Grab",MemberGuid=86E248C047254FEECC0603BB6C0FE484)
+					CustomProperties Pin (PinId=FDB36E2741A2BCEF856502AAC746B07F,PinName="execute",PinType.PinCategory="exec",PinType.PinSubCategory="",PinType.PinSubCategoryObject=None,PinType.PinSubCategoryMemberReference=(),PinType.PinValueType=(),PinType.bIsMap=False,PinType.bIsSet=False,PinType.bIsArray=False,PinType.bIsReference=False,PinType.bIsConst=False,PinType.bIsWeakPointer=False,LinkedTo=(K2Node_Event_10 01317E634C1C07F9D5D729B8421618F2,),PersistentGuid=00000000000000000000000000000000,bHidden=False,bNotConnectable=False,bDefaultValueIsReadOnly=False,bDefaultValueIsIgnored=False,bAdvancedView=False,)
 
-		NodePosX=320
+					CustomProperties Pin (PinId=1B74752B4533A65EB598E787671714A5,PinName="then",Direction="EGPD_Output",PinType.PinCategory="exec",PinType.PinSubCategory="",PinType.PinSubCategoryObject=None,PinType.PinSubCategoryMemberReference=(),PinType.PinValueType=(),PinType.bIsMap=False,PinType.bIsSet=False,PinType.bIsArray=False,PinType.bIsReference=False,PinType.bIsConst=False,PinType.bIsWeakPointer=False,LinkedTo=(K2Node_VariableSet_0 AE7B514B460DD96C37ECE7A72350ADB1,),PersistentGuid=00000000000000000000000000000000,bHidden=False,bNotConnectable=False,bDefaultValueIsReadOnly=False,bDefaultValueIsIgnored=False,bAdvancedView=False,)
 
-		NodePosY=1200
+					CustomProperties Pin (PinId=A448FB8B4EEA51822C66DBA48B3087AA,PinName="self",PinFriendlyName=NSLOCTEXT("K2Node", "Target", "Target"),PinType.PinCategory="object",PinType.PinSubCategory="",PinType.PinSubCategoryObject=BlueprintGeneratedClass'/Game/Prototype/Game/Gun_Base.Gun_Base_C',PinType.PinSubCategoryMemberReference=(),PinType.PinValueType=(),PinType.bIsMap=False,PinType.bIsSet=False,PinType.bIsArray=False,PinType.bIsReference=False,PinType.bIsConst=False,PinType.bIsWeakPointer=False,PersistentGuid=00000000000000000000000000000000,bHidden=True,bNotConnectable=False,bDefaultValueIsReadOnly=False,bDefaultValueIsIgnored=False,bAdvancedView=False,)
 
-		NodeGuid=F79E93794CF4A9FF0F133187E49F3FF4
+				End Object
 
-		CustomProperties Pin (PinId=FDB36E2741A2BCEF856502AAC746B07F,PinName="execute",PinType.PinCategory="exec",PinType.PinSubCategory="",PinType.PinSubCategoryObject=None,PinType.PinSubCategoryMemberReference=(),PinType.PinValueType=(),PinType.bIsMap=False,PinType.bIsSet=False,PinType.bIsArray=False,PinType.bIsReference=False,PinType.bIsConst=False,PinType.bIsWeakPointer=False,LinkedTo=(K2Node_Event_10 01317E634C1C07F9D5D729B8421618F2,),PersistentGuid=00000000000000000000000000000000,bHidden=False,bNotConnectable=False,bDefaultValueIsReadOnly=False,bDefaultValueIsIgnored=False,bAdvancedView=False,)
+				Begin Object Class=K2Node_CallParentFunction Name="K2Node_CallParentFunction_3"
 
-		CustomProperties Pin (PinId=1B74752B4533A65EB598E787671714A5,PinName="then",Direction="EGPD_Output",PinType.PinCategory="exec",PinType.PinSubCategory="",PinType.PinSubCategoryObject=None,PinType.PinSubCategoryMemberReference=(),PinType.PinValueType=(),PinType.bIsMap=False,PinType.bIsSet=False,PinType.bIsArray=False,PinType.bIsReference=False,PinType.bIsConst=False,PinType.bIsWeakPointer=False,LinkedTo=(K2Node_VariableSet_0 AE7B514B460DD96C37ECE7A72350ADB1,),PersistentGuid=00000000000000000000000000000000,bHidden=False,bNotConnectable=False,bDefaultValueIsReadOnly=False,bDefaultValueIsIgnored=False,bAdvancedView=False,)
+					FunctionReference=(MemberParent=BlueprintGeneratedClass'/Game/Prototype/Game/Gun_Pistol.Gun_Pistol_C',MemberName="End Primary Grab",MemberGuid=5AC0E94444892A8794605F86BF8B03E2)
 
-		CustomProperties Pin (PinId=A448FB8B4EEA51822C66DBA48B3087AA,PinName="self",PinFriendlyName=NSLOCTEXT("K2Node", "Target", "Target"),PinType.PinCategory="object",PinType.PinSubCategory="",PinType.PinSubCategoryObject=BlueprintGeneratedClass'/Game/Prototype/Game/Gun_Base.Gun_Base_C',PinType.PinSubCategoryMemberReference=(),PinType.PinValueType=(),PinType.bIsMap=False,PinType.bIsSet=False,PinType.bIsArray=False,PinType.bIsReference=False,PinType.bIsConst=False,PinType.bIsWeakPointer=False,PersistentGuid=00000000000000000000000000000000,bHidden=True,bNotConnectable=False,bDefaultValueIsReadOnly=False,bDefaultValueIsIgnored=False,bAdvancedView=False,)
+					NodePosX=320
 
-	End Object
+					NodePosY=1584
 
-	Begin Object Class=K2Node_CallParentFunction Name="K2Node_CallParentFunction_3"
+					NodeGuid=3F018DC04674D5E4348E1E8885C1D07A
 
-		FunctionReference=(MemberParent=BlueprintGeneratedClass'/Game/Prototype/Game/Gun_Pistol.Gun_Pistol_C',MemberName="End Primary Grab",MemberGuid=5AC0E94444892A8794605F86BF8B03E2)
+					CustomProperties Pin (PinId=440FA2384D499E510943C5952AED1D52,PinName="execute",PinType.PinCategory="exec",PinType.PinSubCategory="",PinType.PinSubCategoryObject=None,PinType.PinSubCategoryMemberReference=(),PinType.PinValueType=(),PinType.bIsMap=False,PinType.bIsSet=False,PinType.bIsArray=False,PinType.bIsReference=False,PinType.bIsConst=False,PinType.bIsWeakPointer=False,LinkedTo=(K2Node_Event_7 DDE6DFF14B7B9ED8FABA8A8A2B48E090,),PersistentGuid=00000000000000000000000000000000,bHidden=False,bNotConnectable=False,bDefaultValueIsReadOnly=False,bDefaultValueIsIgnored=False,bAdvancedView=False,)
 
-		NodePosX=320
+					CustomProperties Pin (PinId=B988ABBD45E77F4BBB6292BC717283DB,PinName="then",Direction="EGPD_Output",PinType.PinCategory="exec",PinType.PinSubCategory="",PinType.PinSubCategoryObject=None,PinType.PinSubCategoryMemberReference=(),PinType.PinValueType=(),PinType.bIsMap=False,PinType.bIsSet=False,PinType.bIsArray=False,PinType.bIsReference=False,PinType.bIsConst=False,PinType.bIsWeakPointer=False,LinkedTo=(K2Node_VariableSet_1553 AE7B514B460DD96C37ECE7A72350ADB1,),PersistentGuid=00000000000000000000000000000000,bHidden=False,bNotConnectable=False,bDefaultValueIsReadOnly=False,bDefaultValueIsIgnored=False,bAdvancedView=False,)
 
-		NodePosY=1584
+					CustomProperties Pin (PinId=FC7E71C84B85D76ED3975980B0C4B1E7,PinName="self",PinFriendlyName=NSLOCTEXT("K2Node", "Target", "Target"),PinType.PinCategory="object",PinType.PinSubCategory="",PinType.PinSubCategoryObject=BlueprintGeneratedClass'/Game/Prototype/Game/Gun_Base.Gun_Base_C',PinType.PinSubCategoryMemberReference=(),PinType.PinValueType=(),PinType.bIsMap=False,PinType.bIsSet=False,PinType.bIsArray=False,PinType.bIsReference=False,PinType.bIsConst=False,PinType.bIsWeakPointer=False,PersistentGuid=00000000000000000000000000000000,bHidden=True,bNotConnectable=False,bDefaultValueIsReadOnly=False,bDefaultValueIsIgnored=False,bAdvancedView=False,)
 
-		NodeGuid=3F018DC04674D5E4348E1E8885C1D07A
+				End Object
 
-		CustomProperties Pin (PinId=440FA2384D499E510943C5952AED1D52,PinName="execute",PinType.PinCategory="exec",PinType.PinSubCategory="",PinType.PinSubCategoryObject=None,PinType.PinSubCategoryMemberReference=(),PinType.PinValueType=(),PinType.bIsMap=False,PinType.bIsSet=False,PinType.bIsArray=False,PinType.bIsReference=False,PinType.bIsConst=False,PinType.bIsWeakPointer=False,LinkedTo=(K2Node_Event_7 DDE6DFF14B7B9ED8FABA8A8A2B48E090,),PersistentGuid=00000000000000000000000000000000,bHidden=False,bNotConnectable=False,bDefaultValueIsReadOnly=False,bDefaultValueIsIgnored=False,bAdvancedView=False,)
+				Begin Object Class=K2Node_Event Name="K2Node_Event_14"
 
-		CustomProperties Pin (PinId=B988ABBD45E77F4BBB6292BC717283DB,PinName="then",Direction="EGPD_Output",PinType.PinCategory="exec",PinType.PinSubCategory="",PinType.PinSubCategoryObject=None,PinType.PinSubCategoryMemberReference=(),PinType.PinValueType=(),PinType.bIsMap=False,PinType.bIsSet=False,PinType.bIsArray=False,PinType.bIsReference=False,PinType.bIsConst=False,PinType.bIsWeakPointer=False,LinkedTo=(K2Node_VariableSet_1553 AE7B514B460DD96C37ECE7A72350ADB1,),PersistentGuid=00000000000000000000000000000000,bHidden=False,bNotConnectable=False,bDefaultValueIsReadOnly=False,bDefaultValueIsIgnored=False,bAdvancedView=False,)
+					EventReference=(MemberParent=BlueprintGeneratedClass'/Game/Prototype/Game/Gun_Base.Gun_Base_C',MemberName="Begin Secondary Grab",MemberGuid=788DAEE54912BA1513C0EA9EBD35484D)
 
-		CustomProperties Pin (PinId=FC7E71C84B85D76ED3975980B0C4B1E7,PinName="self",PinFriendlyName=NSLOCTEXT("K2Node", "Target", "Target"),PinType.PinCategory="object",PinType.PinSubCategory="",PinType.PinSubCategoryObject=BlueprintGeneratedClass'/Game/Prototype/Game/Gun_Base.Gun_Base_C',PinType.PinSubCategoryMemberReference=(),PinType.PinValueType=(),PinType.bIsMap=False,PinType.bIsSet=False,PinType.bIsArray=False,PinType.bIsReference=False,PinType.bIsConst=False,PinType.bIsWeakPointer=False,PersistentGuid=00000000000000000000000000000000,bHidden=True,bNotConnectable=False,bDefaultValueIsReadOnly=False,bDefaultValueIsIgnored=False,bAdvancedView=False,)
+					bOverrideFunction=True
 
-	End Object
+					NodePosX=48
 
-	Begin Object Class=K2Node_Event Name="K2Node_Event_14"
+					NodePosY=1296
 
-		EventReference=(MemberParent=BlueprintGeneratedClass'/Game/Prototype/Game/Gun_Base.Gun_Base_C',MemberName="Begin Secondary Grab",MemberGuid=788DAEE54912BA1513C0EA9EBD35484D)
+					NodeGuid=68AB4E004B109F18E9831889357D0DB8
 
-		bOverrideFunction=True
+					CustomProperties Pin (PinId=9AB423404CAEE074165D7CABBA120EE1,PinName="OutputDelegate",Direction="EGPD_Output",PinType.PinCategory="delegate",PinType.PinSubCategory="",PinType.PinSubCategoryObject=None,PinType.PinSubCategoryMemberReference=(MemberParent=BlueprintGeneratedClass'/Game/Prototype/Game/Gun_Base.Gun_Base_C',MemberName="Begin Secondary Grab",MemberGuid=788DAEE54912BA1513C0EA9EBD35484D),PinType.PinValueType=(),PinType.bIsMap=False,PinType.bIsSet=False,PinType.bIsArray=False,PinType.bIsReference=False,PinType.bIsConst=False,PinType.bIsWeakPointer=False,PersistentGuid=00000000000000000000000000000000,bHidden=False,bNotConnectable=False,bDefaultValueIsReadOnly=False,bDefaultValueIsIgnored=False,bAdvancedView=False,)
 
-		NodePosX=48
+					CustomProperties Pin (PinId=8A7BF15740D6B74BE105CCBDD79423E5,PinName="then",Direction="EGPD_Output",PinType.PinCategory="exec",PinType.PinSubCategory="",PinType.PinSubCategoryObject=None,PinType.PinSubCategoryMemberReference=(),PinType.PinValueType=(),PinType.bIsMap=False,PinType.bIsSet=False,PinType.bIsArray=False,PinType.bIsReference=False,PinType.bIsConst=False,PinType.bIsWeakPointer=False,LinkedTo=(K2Node_CallParentFunction_8 280509624149FB7B8299579E0218E824,),PersistentGuid=00000000000000000000000000000000,bHidden=False,bNotConnectable=False,bDefaultValueIsReadOnly=False,bDefaultValueIsIgnored=False,bAdvancedView=False,)
 
-		NodePosY=1296
+				End Object
 
-		NodeGuid=68AB4E004B109F18E9831889357D0DB8
+				Begin Object Class=K2Node_Event Name="K2Node_Event_19"
 
-		CustomProperties Pin (PinId=9AB423404CAEE074165D7CABBA120EE1,PinName="OutputDelegate",Direction="EGPD_Output",PinType.PinCategory="delegate",PinType.PinSubCategory="",PinType.PinSubCategoryObject=None,PinType.PinSubCategoryMemberReference=(MemberParent=BlueprintGeneratedClass'/Game/Prototype/Game/Gun_Base.Gun_Base_C',MemberName="Begin Secondary Grab",MemberGuid=788DAEE54912BA1513C0EA9EBD35484D),PinType.PinValueType=(),PinType.bIsMap=False,PinType.bIsSet=False,PinType.bIsArray=False,PinType.bIsReference=False,PinType.bIsConst=False,PinType.bIsWeakPointer=False,PersistentGuid=00000000000000000000000000000000,bHidden=False,bNotConnectable=False,bDefaultValueIsReadOnly=False,bDefaultValueIsIgnored=False,bAdvancedView=False,)
+					EventReference=(MemberParent=BlueprintGeneratedClass'/Game/Prototype/Game/Gun_Base.Gun_Base_C',MemberName="End Secondary Grab",MemberGuid=CE30C1A0429D6BC4F021EC894F025A9A)
 
-		CustomProperties Pin (PinId=8A7BF15740D6B74BE105CCBDD79423E5,PinName="then",Direction="EGPD_Output",PinType.PinCategory="exec",PinType.PinSubCategory="",PinType.PinSubCategoryObject=None,PinType.PinSubCategoryMemberReference=(),PinType.PinValueType=(),PinType.bIsMap=False,PinType.bIsSet=False,PinType.bIsArray=False,PinType.bIsReference=False,PinType.bIsConst=False,PinType.bIsWeakPointer=False,LinkedTo=(K2Node_CallParentFunction_8 280509624149FB7B8299579E0218E824,),PersistentGuid=00000000000000000000000000000000,bHidden=False,bNotConnectable=False,bDefaultValueIsReadOnly=False,bDefaultValueIsIgnored=False,bAdvancedView=False,)
+					bOverrideFunction=True
 
-	End Object
+					NodePosX=48
 
-	Begin Object Class=K2Node_Event Name="K2Node_Event_19"
+					NodePosY=1504
 
-		EventReference=(MemberParent=BlueprintGeneratedClass'/Game/Prototype/Game/Gun_Base.Gun_Base_C',MemberName="End Secondary Grab",MemberGuid=CE30C1A0429D6BC4F021EC894F025A9A)
+					NodeGuid=CF2791714730BBCE92043AA025D0C200
 
-		bOverrideFunction=True
+					CustomProperties Pin (PinId=04FA7A954BA262A18863708BE6933EF8,PinName="OutputDelegate",Direction="EGPD_Output",PinType.PinCategory="delegate",PinType.PinSubCategory="",PinType.PinSubCategoryObject=None,PinType.PinSubCategoryMemberReference=(MemberParent=BlueprintGeneratedClass'/Game/Prototype/Game/Gun_Base.Gun_Base_C',MemberName="End Secondary Grab",MemberGuid=CE30C1A0429D6BC4F021EC894F025A9A),PinType.PinValueType=(),PinType.bIsMap=False,PinType.bIsSet=False,PinType.bIsArray=False,PinType.bIsReference=False,PinType.bIsConst=False,PinType.bIsWeakPointer=False,PersistentGuid=00000000000000000000000000000000,bHidden=False,bNotConnectable=False,bDefaultValueIsReadOnly=False,bDefaultValueIsIgnored=False,bAdvancedView=False,)
 
-		NodePosX=48
+					CustomProperties Pin (PinId=87AB296D4EE26D0186F710BB3FA4C7B7,PinName="then",Direction="EGPD_Output",PinType.PinCategory="exec",PinType.PinSubCategory="",PinType.PinSubCategoryObject=None,PinType.PinSubCategoryMemberReference=(),PinType.PinValueType=(),PinType.bIsMap=False,PinType.bIsSet=False,PinType.bIsArray=False,PinType.bIsReference=False,PinType.bIsConst=False,PinType.bIsWeakPointer=False,LinkedTo=(K2Node_CallParentFunction_9 C2AB49794E3B99FF2DDC1983D4C3D443,),PersistentGuid=00000000000000000000000000000000,bHidden=False,bNotConnectable=False,bDefaultValueIsReadOnly=False,bDefaultValueIsIgnored=False,bAdvancedView=False,)
 
-		NodePosY=1504
+				End Object
 
-		NodeGuid=CF2791714730BBCE92043AA025D0C200
+				Begin Object Class=K2Node_CallParentFunction Name="K2Node_CallParentFunction_8"
 
-		CustomProperties Pin (PinId=04FA7A954BA262A18863708BE6933EF8,PinName="OutputDelegate",Direction="EGPD_Output",PinType.PinCategory="delegate",PinType.PinSubCategory="",PinType.PinSubCategoryObject=None,PinType.PinSubCategoryMemberReference=(MemberParent=BlueprintGeneratedClass'/Game/Prototype/Game/Gun_Base.Gun_Base_C',MemberName="End Secondary Grab",MemberGuid=CE30C1A0429D6BC4F021EC894F025A9A),PinType.PinValueType=(),PinType.bIsMap=False,PinType.bIsSet=False,PinType.bIsArray=False,PinType.bIsReference=False,PinType.bIsConst=False,PinType.bIsWeakPointer=False,PersistentGuid=00000000000000000000000000000000,bHidden=False,bNotConnectable=False,bDefaultValueIsReadOnly=False,bDefaultValueIsIgnored=False,bAdvancedView=False,)
+					FunctionReference=(MemberParent=BlueprintGeneratedClass'/Game/Prototype/Game/Gun_Pistol.Gun_Pistol_C',MemberName="Begin Secondary Grab",MemberGuid=788DAEE54912BA1513C0EA9EBD35484D)
 
-		CustomProperties Pin (PinId=87AB296D4EE26D0186F710BB3FA4C7B7,PinName="then",Direction="EGPD_Output",PinType.PinCategory="exec",PinType.PinSubCategory="",PinType.PinSubCategoryObject=None,PinType.PinSubCategoryMemberReference=(),PinType.PinValueType=(),PinType.bIsMap=False,PinType.bIsSet=False,PinType.bIsArray=False,PinType.bIsReference=False,PinType.bIsConst=False,PinType.bIsWeakPointer=False,LinkedTo=(K2Node_CallParentFunction_9 C2AB49794E3B99FF2DDC1983D4C3D443,),PersistentGuid=00000000000000000000000000000000,bHidden=False,bNotConnectable=False,bDefaultValueIsReadOnly=False,bDefaultValueIsIgnored=False,bAdvancedView=False,)
+					NodePosX=320
 
-	End Object
+					NodePosY=1296
 
-	Begin Object Class=K2Node_CallParentFunction Name="K2Node_CallParentFunction_8"
+					NodeGuid=5E8F1A584BE808DB38A4338935BEB20E
 
-		FunctionReference=(MemberParent=BlueprintGeneratedClass'/Game/Prototype/Game/Gun_Pistol.Gun_Pistol_C',MemberName="Begin Secondary Grab",MemberGuid=788DAEE54912BA1513C0EA9EBD35484D)
+					CustomProperties Pin (PinId=280509624149FB7B8299579E0218E824,PinName="execute",PinType.PinCategory="exec",PinType.PinSubCategory="",PinType.PinSubCategoryObject=None,PinType.PinSubCategoryMemberReference=(),PinType.PinValueType=(),PinType.bIsMap=False,PinType.bIsSet=False,PinType.bIsArray=False,PinType.bIsReference=False,PinType.bIsConst=False,PinType.bIsWeakPointer=False,LinkedTo=(K2Node_Event_14 8A7BF15740D6B74BE105CCBDD79423E5,),PersistentGuid=00000000000000000000000000000000,bHidden=False,bNotConnectable=False,bDefaultValueIsReadOnly=False,bDefaultValueIsIgnored=False,bAdvancedView=False,)
 
-		NodePosX=320
+					CustomProperties Pin (PinId=84FF3C864D5859B0905DAD8DC8595A9E,PinName="then",Direction="EGPD_Output",PinType.PinCategory="exec",PinType.PinSubCategory="",PinType.PinSubCategoryObject=None,PinType.PinSubCategoryMemberReference=(),PinType.PinValueType=(),PinType.bIsMap=False,PinType.bIsSet=False,PinType.bIsArray=False,PinType.bIsReference=False,PinType.bIsConst=False,PinType.bIsWeakPointer=False,LinkedTo=(K2Node_VariableSet_0 AE7B514B460DD96C37ECE7A72350ADB1,),PersistentGuid=00000000000000000000000000000000,bHidden=False,bNotConnectable=False,bDefaultValueIsReadOnly=False,bDefaultValueIsIgnored=False,bAdvancedView=False,)
 
-		NodePosY=1296
+					CustomProperties Pin (PinId=BE812250483E2AEE00E09B9EBE6F7E60,PinName="self",PinFriendlyName=NSLOCTEXT("K2Node", "Target", "Target"),PinType.PinCategory="object",PinType.PinSubCategory="",PinType.PinSubCategoryObject=BlueprintGeneratedClass'/Game/Prototype/Game/Gun_Base.Gun_Base_C',PinType.PinSubCategoryMemberReference=(),PinType.PinValueType=(),PinType.bIsMap=False,PinType.bIsSet=False,PinType.bIsArray=False,PinType.bIsReference=False,PinType.bIsConst=False,PinType.bIsWeakPointer=False,PersistentGuid=00000000000000000000000000000000,bHidden=True,bNotConnectable=False,bDefaultValueIsReadOnly=False,bDefaultValueIsIgnored=False,bAdvancedView=False,)
 
-		NodeGuid=5E8F1A584BE808DB38A4338935BEB20E
+				End Object
 
-		CustomProperties Pin (PinId=280509624149FB7B8299579E0218E824,PinName="execute",PinType.PinCategory="exec",PinType.PinSubCategory="",PinType.PinSubCategoryObject=None,PinType.PinSubCategoryMemberReference=(),PinType.PinValueType=(),PinType.bIsMap=False,PinType.bIsSet=False,PinType.bIsArray=False,PinType.bIsReference=False,PinType.bIsConst=False,PinType.bIsWeakPointer=False,LinkedTo=(K2Node_Event_14 8A7BF15740D6B74BE105CCBDD79423E5,),PersistentGuid=00000000000000000000000000000000,bHidden=False,bNotConnectable=False,bDefaultValueIsReadOnly=False,bDefaultValueIsIgnored=False,bAdvancedView=False,)
+				Begin Object Class=K2Node_CallParentFunction Name="K2Node_CallParentFunction_9"
 
-		CustomProperties Pin (PinId=84FF3C864D5859B0905DAD8DC8595A9E,PinName="then",Direction="EGPD_Output",PinType.PinCategory="exec",PinType.PinSubCategory="",PinType.PinSubCategoryObject=None,PinType.PinSubCategoryMemberReference=(),PinType.PinValueType=(),PinType.bIsMap=False,PinType.bIsSet=False,PinType.bIsArray=False,PinType.bIsReference=False,PinType.bIsConst=False,PinType.bIsWeakPointer=False,LinkedTo=(K2Node_VariableSet_0 AE7B514B460DD96C37ECE7A72350ADB1,),PersistentGuid=00000000000000000000000000000000,bHidden=False,bNotConnectable=False,bDefaultValueIsReadOnly=False,bDefaultValueIsIgnored=False,bAdvancedView=False,)
+					FunctionReference=(MemberParent=BlueprintGeneratedClass'/Game/Prototype/Game/Gun_Pistol.Gun_Pistol_C',MemberName="End Secondary Grab",MemberGuid=CE30C1A0429D6BC4F021EC894F025A9A)
 
-		CustomProperties Pin (PinId=BE812250483E2AEE00E09B9EBE6F7E60,PinName="self",PinFriendlyName=NSLOCTEXT("K2Node", "Target", "Target"),PinType.PinCategory="object",PinType.PinSubCategory="",PinType.PinSubCategoryObject=BlueprintGeneratedClass'/Game/Prototype/Game/Gun_Base.Gun_Base_C',PinType.PinSubCategoryMemberReference=(),PinType.PinValueType=(),PinType.bIsMap=False,PinType.bIsSet=False,PinType.bIsArray=False,PinType.bIsReference=False,PinType.bIsConst=False,PinType.bIsWeakPointer=False,PersistentGuid=00000000000000000000000000000000,bHidden=True,bNotConnectable=False,bDefaultValueIsReadOnly=False,bDefaultValueIsIgnored=False,bAdvancedView=False,)
+					NodePosX=320
 
-	End Object
+					NodePosY=1504
 
-	Begin Object Class=K2Node_CallParentFunction Name="K2Node_CallParentFunction_9"
+					NodeGuid=0E48003242FC85D2D68655B218A52797
 
-		FunctionReference=(MemberParent=BlueprintGeneratedClass'/Game/Prototype/Game/Gun_Pistol.Gun_Pistol_C',MemberName="End Secondary Grab",MemberGuid=CE30C1A0429D6BC4F021EC894F025A9A)
+					CustomProperties Pin (PinId=C2AB49794E3B99FF2DDC1983D4C3D443,PinName="execute",PinType.PinCategory="exec",PinType.PinSubCategory="",PinType.PinSubCategoryObject=None,PinType.PinSubCategoryMemberReference=(),PinType.PinValueType=(),PinType.bIsMap=False,PinType.bIsSet=False,PinType.bIsArray=False,PinType.bIsReference=False,PinType.bIsConst=False,PinType.bIsWeakPointer=False,LinkedTo=(K2Node_Event_19 87AB296D4EE26D0186F710BB3FA4C7B7,),PersistentGuid=00000000000000000000000000000000,bHidden=False,bNotConnectable=False,bDefaultValueIsReadOnly=False,bDefaultValueIsIgnored=False,bAdvancedView=False,)
 
-		NodePosX=320
+					CustomProperties Pin (PinId=04A663DF4151E8B743CBE3A7C7BCF848,PinName="then",Direction="EGPD_Output",PinType.PinCategory="exec",PinType.PinSubCategory="",PinType.PinSubCategoryObject=None,PinType.PinSubCategoryMemberReference=(),PinType.PinValueType=(),PinType.bIsMap=False,PinType.bIsSet=False,PinType.bIsArray=False,PinType.bIsReference=False,PinType.bIsConst=False,PinType.bIsWeakPointer=False,LinkedTo=(K2Node_VariableSet_1553 AE7B514B460DD96C37ECE7A72350ADB1,),PersistentGuid=00000000000000000000000000000000,bHidden=False,bNotConnectable=False,bDefaultValueIsReadOnly=False,bDefaultValueIsIgnored=False,bAdvancedView=False,)
 
-		NodePosY=1504
+					CustomProperties Pin (PinId=1E9924B4423BFF3C997BB4974473931D,PinName="self",PinFriendlyName=NSLOCTEXT("K2Node", "Target", "Target"),PinType.PinCategory="object",PinType.PinSubCategory="",PinType.PinSubCategoryObject=BlueprintGeneratedClass'/Game/Prototype/Game/Gun_Base.Gun_Base_C',PinType.PinSubCategoryMemberReference=(),PinType.PinValueType=(),PinType.bIsMap=False,PinType.bIsSet=False,PinType.bIsArray=False,PinType.bIsReference=False,PinType.bIsConst=False,PinType.bIsWeakPointer=False,PersistentGuid=00000000000000000000000000000000,bHidden=True,bNotConnectable=False,bDefaultValueIsReadOnly=False,bDefaultValueIsIgnored=False,bAdvancedView=False,)
 
-		NodeGuid=0E48003242FC85D2D68655B218A52797
+				End Object
 
-		CustomProperties Pin (PinId=C2AB49794E3B99FF2DDC1983D4C3D443,PinName="execute",PinType.PinCategory="exec",PinType.PinSubCategory="",PinType.PinSubCategoryObject=None,PinType.PinSubCategoryMemberReference=(),PinType.PinValueType=(),PinType.bIsMap=False,PinType.bIsSet=False,PinType.bIsArray=False,PinType.bIsReference=False,PinType.bIsConst=False,PinType.bIsWeakPointer=False,LinkedTo=(K2Node_Event_19 87AB296D4EE26D0186F710BB3FA4C7B7,),PersistentGuid=00000000000000000000000000000000,bHidden=False,bNotConnectable=False,bDefaultValueIsReadOnly=False,bDefaultValueIsIgnored=False,bAdvancedView=False,)
+				Begin Object Class=EdGraphNode_Comment Name="EdGraphNode_Comment_276"
 
-		CustomProperties Pin (PinId=04A663DF4151E8B743CBE3A7C7BCF848,PinName="then",Direction="EGPD_Output",PinType.PinCategory="exec",PinType.PinSubCategory="",PinType.PinSubCategoryObject=None,PinType.PinSubCategoryMemberReference=(),PinType.PinValueType=(),PinType.bIsMap=False,PinType.bIsSet=False,PinType.bIsArray=False,PinType.bIsReference=False,PinType.bIsConst=False,PinType.bIsWeakPointer=False,LinkedTo=(K2Node_VariableSet_1553 AE7B514B460DD96C37ECE7A72350ADB1,),PersistentGuid=00000000000000000000000000000000,bHidden=False,bNotConnectable=False,bDefaultValueIsReadOnly=False,bDefaultValueIsIgnored=False,bAdvancedView=False,)
+					NodePosY=1152
 
-		CustomProperties Pin (PinId=1E9924B4423BFF3C997BB4974473931D,PinName="self",PinFriendlyName=NSLOCTEXT("K2Node", "Target", "Target"),PinType.PinCategory="object",PinType.PinSubCategory="",PinType.PinSubCategoryObject=BlueprintGeneratedClass'/Game/Prototype/Game/Gun_Base.Gun_Base_C',PinType.PinSubCategoryMemberReference=(),PinType.PinValueType=(),PinType.bIsMap=False,PinType.bIsSet=False,PinType.bIsArray=False,PinType.bIsReference=False,PinType.bIsConst=False,PinType.bIsWeakPointer=False,PersistentGuid=00000000000000000000000000000000,bHidden=True,bNotConnectable=False,bDefaultValueIsReadOnly=False,bDefaultValueIsIgnored=False,bAdvancedView=False,)
+					NodeWidth=880
 
-	End Object
+					NodeHeight=256
 
-	Begin Object Class=EdGraphNode_Comment Name="EdGraphNode_Comment_276"
+					NodeComment="Begin Grab Events"
 
-		NodePosY=1152
+					NodeGuid=73C6DFFF47B2E78D7A929382C5BAA6C2
 
-		NodeWidth=880
+				End Object
 
-		NodeHeight=256
+				Begin Object Class=EdGraphNode_Comment Name="EdGraphNode_Comment_277"
 
-		NodeComment="Begin Grab Events"
+					NodePosY=1456
 
-		NodeGuid=73C6DFFF47B2E78D7A929382C5BAA6C2
+					NodeWidth=880
 
-	End Object
+					NodeHeight=240
 
-	Begin Object Class=EdGraphNode_Comment Name="EdGraphNode_Comment_277"
+					NodeComment="End Grab Events"
 
-		NodePosY=1456
+					NodeGuid=1E4AC8CB452379DDC9B4F6964B10C06B
 
-		NodeWidth=880
+				End Object
 
-		NodeHeight=240
+				Begin Object Class=K2Node_VariableSet Name="K2Node_VariableSet_0"
 
-		NodeComment="End Grab Events"
+					VariableReference=(MemberName="bCanExplode",MemberGuid=9BCA29C8404B219E30C266861239680A,bSelfContext=True)
 
-		NodeGuid=1E4AC8CB452379DDC9B4F6964B10C06B
+					NodePosX=624
 
-	End Object
+					NodePosY=1264
 
-	Begin Object Class=K2Node_VariableSet Name="K2Node_VariableSet_0"
+					NodeGuid=8178103244831FA0B54DDAA8525A4191
 
-		VariableReference=(MemberName="bCanExplode",MemberGuid=9BCA29C8404B219E30C266861239680A,bSelfContext=True)
+					CustomProperties Pin (PinId=AE7B514B460DD96C37ECE7A72350ADB1,PinName="execute",PinType.PinCategory="exec",PinType.PinSubCategory="",PinType.PinSubCategoryObject=None,PinType.PinSubCategoryMemberReference=(),PinType.PinValueType=(),PinType.bIsMap=False,PinType.bIsSet=False,PinType.bIsArray=False,PinType.bIsReference=False,PinType.bIsConst=False,PinType.bIsWeakPointer=False,LinkedTo=(K2Node_CallParentFunction_0 1B74752B4533A65EB598E787671714A5,K2Node_CallParentFunction_8 84FF3C864D5859B0905DAD8DC8595A9E,),PersistentGuid=00000000000000000000000000000000,bHidden=False,bNotConnectable=False,bDefaultValueIsReadOnly=False,bDefaultValueIsIgnored=False,bAdvancedView=False,)
 
-		NodePosX=624
+					CustomProperties Pin (PinId=E0F4A52E424426C37B5EF98C8FFAF121,PinName="then",Direction="EGPD_Output",PinType.PinCategory="exec",PinType.PinSubCategory="",PinType.PinSubCategoryObject=None,PinType.PinSubCategoryMemberReference=(),PinType.PinValueType=(),PinType.bIsMap=False,PinType.bIsSet=False,PinType.bIsArray=False,PinType.bIsReference=False,PinType.bIsConst=False,PinType.bIsWeakPointer=False,PersistentGuid=00000000000000000000000000000000,bHidden=False,bNotConnectable=False,bDefaultValueIsReadOnly=False,bDefaultValueIsIgnored=False,bAdvancedView=False,)
 
-		NodePosY=1264
+					CustomProperties Pin (PinId=DAB7D07A426D67E49C8D63A7911BFC1B,PinName="bCanExplode",PinType.PinCategory="bool",PinType.PinSubCategory="",PinType.PinSubCategoryObject=None,PinType.PinSubCategoryMemberReference=(),PinType.PinValueType=(),PinType.bIsMap=False,PinType.bIsSet=False,PinType.bIsArray=False,PinType.bIsReference=False,PinType.bIsConst=False,PinType.bIsWeakPointer=False,DefaultValue="false",AutogeneratedDefaultValue="false",PersistentGuid=00000000000000000000000000000000,bHidden=False,bNotConnectable=False,bDefaultValueIsReadOnly=False,bDefaultValueIsIgnored=False,bAdvancedView=False,)
 
-		NodeGuid=8178103244831FA0B54DDAA8525A4191
+					CustomProperties Pin (PinId=26DA759E42C15E427748F283E9F317CC,PinName="Output_Get",PinToolTip="Retrieves the value of the variable, can use instead of a separate Get node",Direction="EGPD_Output",PinType.PinCategory="bool",PinType.PinSubCategory="",PinType.PinSubCategoryObject=None,PinType.PinSubCategoryMemberReference=(),PinType.PinValueType=(),PinType.bIsMap=False,PinType.bIsSet=False,PinType.bIsArray=False,PinType.bIsReference=False,PinType.bIsConst=False,PinType.bIsWeakPointer=False,DefaultValue="false",AutogeneratedDefaultValue="false",PersistentGuid=00000000000000000000000000000000,bHidden=False,bNotConnectable=False,bDefaultValueIsReadOnly=False,bDefaultValueIsIgnored=False,bAdvancedView=False,)
 
-		CustomProperties Pin (PinId=AE7B514B460DD96C37ECE7A72350ADB1,PinName="execute",PinType.PinCategory="exec",PinType.PinSubCategory="",PinType.PinSubCategoryObject=None,PinType.PinSubCategoryMemberReference=(),PinType.PinValueType=(),PinType.bIsMap=False,PinType.bIsSet=False,PinType.bIsArray=False,PinType.bIsReference=False,PinType.bIsConst=False,PinType.bIsWeakPointer=False,LinkedTo=(K2Node_CallParentFunction_0 1B74752B4533A65EB598E787671714A5,K2Node_CallParentFunction_8 84FF3C864D5859B0905DAD8DC8595A9E,),PersistentGuid=00000000000000000000000000000000,bHidden=False,bNotConnectable=False,bDefaultValueIsReadOnly=False,bDefaultValueIsIgnored=False,bAdvancedView=False,)
+					CustomProperties Pin (PinId=93A90B8F461C1A8D79FF37B7F054C5F3,PinName="self",PinFriendlyName=NSLOCTEXT("K2Node", "Target", "Target"),PinType.PinCategory="object",PinType.PinSubCategory="",PinType.PinSubCategoryObject=BlueprintGeneratedClass'/GrenadeGun/Blueprints/Weapons/GrenadeGunPistol.GrenadeGunPistol_C',PinType.PinSubCategoryMemberReference=(),PinType.PinValueType=(),PinType.bIsMap=False,PinType.bIsSet=False,PinType.bIsArray=False,PinType.bIsReference=False,PinType.bIsConst=False,PinType.bIsWeakPointer=False,PersistentGuid=00000000000000000000000000000000,bHidden=True,bNotConnectable=False,bDefaultValueIsReadOnly=False,bDefaultValueIsIgnored=False,bAdvancedView=False,)
 
-		CustomProperties Pin (PinId=E0F4A52E424426C37B5EF98C8FFAF121,PinName="then",Direction="EGPD_Output",PinType.PinCategory="exec",PinType.PinSubCategory="",PinType.PinSubCategoryObject=None,PinType.PinSubCategoryMemberReference=(),PinType.PinValueType=(),PinType.bIsMap=False,PinType.bIsSet=False,PinType.bIsArray=False,PinType.bIsReference=False,PinType.bIsConst=False,PinType.bIsWeakPointer=False,PersistentGuid=00000000000000000000000000000000,bHidden=False,bNotConnectable=False,bDefaultValueIsReadOnly=False,bDefaultValueIsIgnored=False,bAdvancedView=False,)
+				End Object
 
-		CustomProperties Pin (PinId=DAB7D07A426D67E49C8D63A7911BFC1B,PinName="bCanExplode",PinType.PinCategory="bool",PinType.PinSubCategory="",PinType.PinSubCategoryObject=None,PinType.PinSubCategoryMemberReference=(),PinType.PinValueType=(),PinType.bIsMap=False,PinType.bIsSet=False,PinType.bIsArray=False,PinType.bIsReference=False,PinType.bIsConst=False,PinType.bIsWeakPointer=False,DefaultValue="false",AutogeneratedDefaultValue="false",PersistentGuid=00000000000000000000000000000000,bHidden=False,bNotConnectable=False,bDefaultValueIsReadOnly=False,bDefaultValueIsIgnored=False,bAdvancedView=False,)
+				Begin Object Class=K2Node_VariableSet Name="K2Node_VariableSet_1553"
 
-		CustomProperties Pin (PinId=26DA759E42C15E427748F283E9F317CC,PinName="Output_Get",PinToolTip="Retrieves the value of the variable, can use instead of a separate Get node",Direction="EGPD_Output",PinType.PinCategory="bool",PinType.PinSubCategory="",PinType.PinSubCategoryObject=None,PinType.PinSubCategoryMemberReference=(),PinType.PinValueType=(),PinType.bIsMap=False,PinType.bIsSet=False,PinType.bIsArray=False,PinType.bIsReference=False,PinType.bIsConst=False,PinType.bIsWeakPointer=False,DefaultValue="false",AutogeneratedDefaultValue="false",PersistentGuid=00000000000000000000000000000000,bHidden=False,bNotConnectable=False,bDefaultValueIsReadOnly=False,bDefaultValueIsIgnored=False,bAdvancedView=False,)
+					VariableReference=(MemberName="bCanExplode",MemberGuid=9BCA29C8404B219E30C266861239680A,bSelfContext=True)
 
-		CustomProperties Pin (PinId=93A90B8F461C1A8D79FF37B7F054C5F3,PinName="self",PinFriendlyName=NSLOCTEXT("K2Node", "Target", "Target"),PinType.PinCategory="object",PinType.PinSubCategory="",PinType.PinSubCategoryObject=BlueprintGeneratedClass'/GrenadeGun/Blueprints/Weapons/GrenadeGunPistol.GrenadeGunPistol_C',PinType.PinSubCategoryMemberReference=(),PinType.PinValueType=(),PinType.bIsMap=False,PinType.bIsSet=False,PinType.bIsArray=False,PinType.bIsReference=False,PinType.bIsConst=False,PinType.bIsWeakPointer=False,PersistentGuid=00000000000000000000000000000000,bHidden=True,bNotConnectable=False,bDefaultValueIsReadOnly=False,bDefaultValueIsIgnored=False,bAdvancedView=False,)
+					NodePosX=624
 
-	End Object
+					NodePosY=1552
 
-	Begin Object Class=K2Node_VariableSet Name="K2Node_VariableSet_1553"
+					NodeGuid=CC14C5034D568D878B12CE8B90B0E8D6
 
-		VariableReference=(MemberName="bCanExplode",MemberGuid=9BCA29C8404B219E30C266861239680A,bSelfContext=True)
+					CustomProperties Pin (PinId=AE7B514B460DD96C37ECE7A72350ADB1,PinName="execute",PinType.PinCategory="exec",PinType.PinSubCategory="",PinType.PinSubCategoryObject=None,PinType.PinSubCategoryMemberReference=(),PinType.PinValueType=(),PinType.bIsMap=False,PinType.bIsSet=False,PinType.bIsArray=False,PinType.bIsReference=False,PinType.bIsConst=False,PinType.bIsWeakPointer=False,LinkedTo=(K2Node_CallParentFunction_9 04A663DF4151E8B743CBE3A7C7BCF848,K2Node_CallParentFunction_3 B988ABBD45E77F4BBB6292BC717283DB,),PersistentGuid=00000000000000000000000000000000,bHidden=False,bNotConnectable=False,bDefaultValueIsReadOnly=False,bDefaultValueIsIgnored=False,bAdvancedView=False,)
 
-		NodePosX=624
+					CustomProperties Pin (PinId=E0F4A52E424426C37B5EF98C8FFAF121,PinName="then",Direction="EGPD_Output",PinType.PinCategory="exec",PinType.PinSubCategory="",PinType.PinSubCategoryObject=None,PinType.PinSubCategoryMemberReference=(),PinType.PinValueType=(),PinType.bIsMap=False,PinType.bIsSet=False,PinType.bIsArray=False,PinType.bIsReference=False,PinType.bIsConst=False,PinType.bIsWeakPointer=False,PersistentGuid=00000000000000000000000000000000,bHidden=False,bNotConnectable=False,bDefaultValueIsReadOnly=False,bDefaultValueIsIgnored=False,bAdvancedView=False,)
 
-		NodePosY=1552
+					CustomProperties Pin (PinId=DAB7D07A426D67E49C8D63A7911BFC1B,PinName="bCanExplode",PinType.PinCategory="bool",PinType.PinSubCategory="",PinType.PinSubCategoryObject=None,PinType.PinSubCategoryMemberReference=(),PinType.PinValueType=(),PinType.bIsMap=False,PinType.bIsSet=False,PinType.bIsArray=False,PinType.bIsReference=False,PinType.bIsConst=False,PinType.bIsWeakPointer=False,DefaultValue="true",AutogeneratedDefaultValue="false",PersistentGuid=00000000000000000000000000000000,bHidden=False,bNotConnectable=False,bDefaultValueIsReadOnly=False,bDefaultValueIsIgnored=False,bAdvancedView=False,)
 
-		NodeGuid=CC14C5034D568D878B12CE8B90B0E8D6
+					CustomProperties Pin (PinId=26DA759E42C15E427748F283E9F317CC,PinName="Output_Get",PinToolTip="Retrieves the value of the variable, can use instead of a separate Get node",Direction="EGPD_Output",PinType.PinCategory="bool",PinType.PinSubCategory="",PinType.PinSubCategoryObject=None,PinType.PinSubCategoryMemberReference=(),PinType.PinValueType=(),PinType.bIsMap=False,PinType.bIsSet=False,PinType.bIsArray=False,PinType.bIsReference=False,PinType.bIsConst=False,PinType.bIsWeakPointer=False,DefaultValue="false",AutogeneratedDefaultValue="false",PersistentGuid=00000000000000000000000000000000,bHidden=False,bNotConnectable=False,bDefaultValueIsReadOnly=False,bDefaultValueIsIgnored=False,bAdvancedView=False,)
 
-		CustomProperties Pin (PinId=AE7B514B460DD96C37ECE7A72350ADB1,PinName="execute",PinType.PinCategory="exec",PinType.PinSubCategory="",PinType.PinSubCategoryObject=None,PinType.PinSubCategoryMemberReference=(),PinType.PinValueType=(),PinType.bIsMap=False,PinType.bIsSet=False,PinType.bIsArray=False,PinType.bIsReference=False,PinType.bIsConst=False,PinType.bIsWeakPointer=False,LinkedTo=(K2Node_CallParentFunction_9 04A663DF4151E8B743CBE3A7C7BCF848,K2Node_CallParentFunction_3 B988ABBD45E77F4BBB6292BC717283DB,),PersistentGuid=00000000000000000000000000000000,bHidden=False,bNotConnectable=False,bDefaultValueIsReadOnly=False,bDefaultValueIsIgnored=False,bAdvancedView=False,)
+					CustomProperties Pin (PinId=93A90B8F461C1A8D79FF37B7F054C5F3,PinName="self",PinFriendlyName=NSLOCTEXT("K2Node", "Target", "Target"),PinType.PinCategory="object",PinType.PinSubCategory="",PinType.PinSubCategoryObject=BlueprintGeneratedClass'/GrenadeGun/Blueprints/Weapons/GrenadeGunPistol.GrenadeGunPistol_C',PinType.PinSubCategoryMemberReference=(),PinType.PinValueType=(),PinType.bIsMap=False,PinType.bIsSet=False,PinType.bIsArray=False,PinType.bIsReference=False,PinType.bIsConst=False,PinType.bIsWeakPointer=False,PersistentGuid=00000000000000000000000000000000,bHidden=True,bNotConnectable=False,bDefaultValueIsReadOnly=False,bDefaultValueIsIgnored=False,bAdvancedView=False,)
 
-		CustomProperties Pin (PinId=E0F4A52E424426C37B5EF98C8FFAF121,PinName="then",Direction="EGPD_Output",PinType.PinCategory="exec",PinType.PinSubCategory="",PinType.PinSubCategoryObject=None,PinType.PinSubCategoryMemberReference=(),PinType.PinValueType=(),PinType.bIsMap=False,PinType.bIsSet=False,PinType.bIsArray=False,PinType.bIsReference=False,PinType.bIsConst=False,PinType.bIsWeakPointer=False,PersistentGuid=00000000000000000000000000000000,bHidden=False,bNotConnectable=False,bDefaultValueIsReadOnly=False,bDefaultValueIsIgnored=False,bAdvancedView=False,)
-
-		CustomProperties Pin (PinId=DAB7D07A426D67E49C8D63A7911BFC1B,PinName="bCanExplode",PinType.PinCategory="bool",PinType.PinSubCategory="",PinType.PinSubCategoryObject=None,PinType.PinSubCategoryMemberReference=(),PinType.PinValueType=(),PinType.bIsMap=False,PinType.bIsSet=False,PinType.bIsArray=False,PinType.bIsReference=False,PinType.bIsConst=False,PinType.bIsWeakPointer=False,DefaultValue="true",AutogeneratedDefaultValue="false",PersistentGuid=00000000000000000000000000000000,bHidden=False,bNotConnectable=False,bDefaultValueIsReadOnly=False,bDefaultValueIsIgnored=False,bAdvancedView=False,)
+				End Object
+		[/PARAM]
+	[/OBJECT]
 
-		CustomProperties Pin (PinId=26DA759E42C15E427748F283E9F317CC,PinName="Output_Get",PinToolTip="Retrieves the value of the variable, can use instead of a separate Get node",Direction="EGPD_Output",PinType.PinCategory="bool",PinType.PinSubCategory="",PinType.PinSubCategoryObject=None,PinType.PinSubCategoryMemberReference=(),PinType.PinValueType=(),PinType.bIsMap=False,PinType.bIsSet=False,PinType.bIsArray=False,PinType.bIsReference=False,PinType.bIsConst=False,PinType.bIsWeakPointer=False,DefaultValue="false",AutogeneratedDefaultValue="false",PersistentGuid=00000000000000000000000000000000,bHidden=False,bNotConnectable=False,bDefaultValueIsReadOnly=False,bDefaultValueIsIgnored=False,bAdvancedView=False,)
+	[REGION:caption]
+	Click the icon in the upper left corner of this image to copy the Blueprint Graph and paste it into your project.
+	[/REGION]
 
-		CustomProperties Pin (PinId=93A90B8F461C1A8D79FF37B7F054C5F3,PinName="self",PinFriendlyName=NSLOCTEXT("K2Node", "Target", "Target"),PinType.PinCategory="object",PinType.PinSubCategory="",PinType.PinSubCategoryObject=BlueprintGeneratedClass'/GrenadeGun/Blueprints/Weapons/GrenadeGunPistol.GrenadeGunPistol_C',PinType.PinSubCategoryMemberReference=(),PinType.PinValueType=(),PinType.bIsMap=False,PinType.bIsSet=False,PinType.bIsArray=False,PinType.bIsReference=False,PinType.bIsConst=False,PinType.bIsWeakPointer=False,PersistentGuid=00000000000000000000000000000000,bHidden=True,bNotConnectable=False,bDefaultValueIsReadOnly=False,bDefaultValueIsIgnored=False,bAdvancedView=False,)
+9. In the Blueprint Editor toolbar, click **Compile** to update the Blueprint and click **Save** to preserve your changes.
 
-	End Object
-[/EXCERPT:snippet]
+## Results
 
-[/COMMENT]
+These sets of Events are called when our gun is grabbed and released. You'll use the **bCanExplode** variable in the next section of the Blueprint to prevent the gun from exploding in your hand when you punch a robot with it.
 
+[NAV]
diff --git a/Engine/Documentation/Source/Modding/RoboRecall/HowTo/SetupGameplay/BeginLevelScript/BeginLevelScript.INT.udn b/Engine/Documentation/Source/Modding/RoboRecall/HowTo/SetupGameplay/BeginLevelScript/BeginLevelScript.INT.udn
--- a/Engine/Documentation/Source/Modding/RoboRecall/HowTo/SetupGameplay/BeginLevelScript/BeginLevelScript.INT.udn
+++ b/Engine/Documentation/Source/Modding/RoboRecall/HowTo/SetupGameplay/BeginLevelScript/BeginLevelScript.INT.udn
@@ -58,7 +58,96 @@ When the level begins, you want the player to start at the location of the Playe
 		![](rotation.png)(convert:false)
 		[/PARAM]
 		[PARAM:snippet]
-			[INCLUDE:#snippet]
+				Begin Object Class=K2Node_Event Name="K2Node_Event_52"
+				   EventReference=(MemberParent=Class'/Script/Engine.Actor',MemberName="ReceiveBeginPlay")
+				   bOverrideFunction=True
+				   NodePosX=-928
+				   NodePosY=-4480
+				   NodeGuid=6C849E7241A6108A36D5C1BEAAEAE94F
+				   CustomProperties Pin (PinId=B35535C443B530261C41B8A84FF2CC42,PinName="OutputDelegate",Direction="EGPD_Output",PinType.PinCategory="delegate",PinType.PinSubCategory="",PinType.PinSubCategoryObject=None,PinType.PinSubCategoryMemberReference=(MemberParent=Class'/Script/Engine.Actor',MemberName="ReceiveBeginPlay"),PinType.PinValueType=(),PinType.bIsMap=False,PinType.bIsSet=False,PinType.bIsArray=False,PinType.bIsReference=False,PinType.bIsConst=False,PinType.bIsWeakPointer=False,PersistentGuid=00000000000000000000000000000000,bHidden=False,bNotConnectable=False,bDefaultValueIsReadOnly=False,bDefaultValueIsIgnored=False,bAdvancedView=False,)
+				   CustomProperties Pin (PinId=5A6DBA7144DD625770231A9C725C0B49,PinName="then",Direction="EGPD_Output",PinType.PinCategory="exec",PinType.PinSubCategory="",PinType.PinSubCategoryObject=None,PinType.PinSubCategoryMemberReference=(),PinType.PinValueType=(),PinType.bIsMap=False,PinType.bIsSet=False,PinType.bIsArray=False,PinType.bIsReference=False,PinType.bIsConst=False,PinType.bIsWeakPointer=False,LinkedTo=(K2Node_DynamicCast_5 3228E8FA4B588A45DCEA3A839D739E9A,),PersistentGuid=00000000000000000000000000000000,bHidden=False,bNotConnectable=False,bDefaultValueIsReadOnly=False,bDefaultValueIsIgnored=False,bAdvancedView=False,)
+				End Object
+				Begin Object Class=K2Node_CallFunction Name="K2Node_CallFunction_280"
+				   bIsPureFunc=True
+				   FunctionReference=(MemberParent=Class'/Script/Engine.GameplayStatics',MemberName="GetGameMode")
+				   NodePosX=-720
+				   NodePosY=-4416
+				   NodeGuid=2CB9ABB94CF9F2C3BAEBB0B38ED0F003
+				   CustomProperties Pin (PinId=1AAC46094DE5E3D0CB1F8488BEDBFE3B,PinName="self",PinFriendlyName=NSLOCTEXT("K2Node", "Target", "Target"),PinToolTip="Target\nGameplay Statics Reference",PinType.PinCategory="object",PinType.PinSubCategory="",PinType.PinSubCategoryObject=Class'/Script/Engine.GameplayStatics',PinType.PinSubCategoryMemberReference=(),PinType.PinValueType=(),PinType.bIsMap=False,PinType.bIsSet=False,PinType.bIsArray=False,PinType.bIsReference=False,PinType.bIsConst=False,PinType.bIsWeakPointer=False,DefaultObject="/Script/Engine.Default__GameplayStatics",PersistentGuid=00000000000000000000000000000000,bHidden=True,bNotConnectable=False,bDefaultValueIsReadOnly=False,bDefaultValueIsIgnored=False,bAdvancedView=False,)
+				   CustomProperties Pin (PinId=539E882646828F110DA8539329878814,PinName="WorldContextObject",PinToolTip="World Context Object\nObject Reference",PinType.PinCategory="object",PinType.PinSubCategory="",PinType.PinSubCategoryObject=Class'/Script/CoreUObject.Object',PinType.PinSubCategoryMemberReference=(),PinType.PinValueType=(),PinType.bIsMap=False,PinType.bIsSet=False,PinType.bIsArray=False,PinType.bIsReference=False,PinType.bIsConst=True,PinType.bIsWeakPointer=False,PersistentGuid=00000000000000000000000000000000,bHidden=True,bNotConnectable=False,bDefaultValueIsReadOnly=False,bDefaultValueIsIgnored=False,bAdvancedView=False,)
+				   CustomProperties Pin (PinId=F81E36D8454550AAAEF10DB57E9FC71F,PinName="ReturnValue",PinToolTip="Return Value\nGame Mode Base Reference",Direction="EGPD_Output",PinType.PinCategory="object",PinType.PinSubCategory="",PinType.PinSubCategoryObject=Class'/Script/Engine.GameModeBase',PinType.PinSubCategoryMemberReference=(),PinType.PinValueType=(),PinType.bIsMap=False,PinType.bIsSet=False,PinType.bIsArray=False,PinType.bIsReference=False,PinType.bIsConst=False,PinType.bIsWeakPointer=False,LinkedTo=(K2Node_DynamicCast_5 8B0B2634480FE02DE82004ADB105AB79,),PersistentGuid=00000000000000000000000000000000,bHidden=False,bNotConnectable=False,bDefaultValueIsReadOnly=False,bDefaultValueIsIgnored=False,bAdvancedView=False,)
+				End Object
+				Begin Object Class=K2Node_DynamicCast Name="K2Node_DynamicCast_5"
+				   TargetType=BlueprintGeneratedClass'/Game/Prototype/Game/BP_OdinGameMode.BP_OdinGameMode_C'
+				   NodePosX=-512
+				   NodePosY=-4480
+				   NodeGuid=2A2928BD4CBEFB243D86D796866C7C1B
+				   CustomProperties Pin (PinId=3228E8FA4B588A45DCEA3A839D739E9A,PinName="execute",PinType.PinCategory="exec",PinType.PinSubCategory="",PinType.PinSubCategoryObject=None,PinType.PinSubCategoryMemberReference=(),PinType.PinValueType=(),PinType.bIsMap=False,PinType.bIsSet=False,PinType.bIsArray=False,PinType.bIsReference=False,PinType.bIsConst=False,PinType.bIsWeakPointer=False,LinkedTo=(K2Node_Event_52 5A6DBA7144DD625770231A9C725C0B49,),PersistentGuid=00000000000000000000000000000000,bHidden=False,bNotConnectable=False,bDefaultValueIsReadOnly=False,bDefaultValueIsIgnored=False,bAdvancedView=False,)
+				   CustomProperties Pin (PinId=C46EEAC7458FE7A6949E78BA06EE5B54,PinName="then",Direction="EGPD_Output",PinType.PinCategory="exec",PinType.PinSubCategory="",PinType.PinSubCategoryObject=None,PinType.PinSubCategoryMemberReference=(),PinType.PinValueType=(),PinType.bIsMap=False,PinType.bIsSet=False,PinType.bIsArray=False,PinType.bIsReference=False,PinType.bIsConst=False,PinType.bIsWeakPointer=False,LinkedTo=(K2Node_CallFunction_389 B9D0319843BB38F57642CEB7B91707F9,),PersistentGuid=00000000000000000000000000000000,bHidden=False,bNotConnectable=False,bDefaultValueIsReadOnly=False,bDefaultValueIsIgnored=False,bAdvancedView=False,)
+				   CustomProperties Pin (PinId=9EAA2A464990955344C076A409669F6D,PinName="CastFailed",Direction="EGPD_Output",PinType.PinCategory="exec",PinType.PinSubCategory="",PinType.PinSubCategoryObject=None,PinType.PinSubCategoryMemberReference=(),PinType.PinValueType=(),PinType.bIsMap=False,PinType.bIsSet=False,PinType.bIsArray=False,PinType.bIsReference=False,PinType.bIsConst=False,PinType.bIsWeakPointer=False,PersistentGuid=00000000000000000000000000000000,bHidden=False,bNotConnectable=False,bDefaultValueIsReadOnly=False,bDefaultValueIsIgnored=False,bAdvancedView=False,)
+				   CustomProperties Pin (PinId=8B0B2634480FE02DE82004ADB105AB79,PinName="Object",PinType.PinCategory="object",PinType.PinSubCategory="",PinType.PinSubCategoryObject=Class'/Script/CoreUObject.Object',PinType.PinSubCategoryMemberReference=(),PinType.PinValueType=(),PinType.bIsMap=False,PinType.bIsSet=False,PinType.bIsArray=False,PinType.bIsReference=False,PinType.bIsConst=False,PinType.bIsWeakPointer=False,LinkedTo=(K2Node_CallFunction_280 F81E36D8454550AAAEF10DB57E9FC71F,),PersistentGuid=00000000000000000000000000000000,bHidden=False,bNotConnectable=False,bDefaultValueIsReadOnly=False,bDefaultValueIsIgnored=False,bAdvancedView=False,)
+				   CustomProperties Pin (PinId=6B739B644EEBFA053240C1B460B2E558,PinName="AsBP Odin Game Mode",Direction="EGPD_Output",PinType.PinCategory="object",PinType.PinSubCategory="",PinType.PinSubCategoryObject=BlueprintGeneratedClass'/Game/Prototype/Game/BP_OdinGameMode.BP_OdinGameMode_C',PinType.PinSubCategoryMemberReference=(),PinType.PinValueType=(),PinType.bIsMap=False,PinType.bIsSet=False,PinType.bIsArray=False,PinType.bIsReference=False,PinType.bIsConst=False,PinType.bIsWeakPointer=False,LinkedTo=(K2Node_CallFunction_389 3117651A4CB6ACA2EAEB6BA4272EC84B,),PersistentGuid=00000000000000000000000000000000,bHidden=False,bNotConnectable=False,bDefaultValueIsReadOnly=False,bDefaultValueIsIgnored=False,bAdvancedView=False,)
+				   CustomProperties Pin (PinId=BEB2F603410D2A627B7D089246E322C7,PinName="bSuccess",Direction="EGPD_Output",PinType.PinCategory="bool",PinType.PinSubCategory="",PinType.PinSubCategoryObject=None,PinType.PinSubCategoryMemberReference=(),PinType.PinValueType=(),PinType.bIsMap=False,PinType.bIsSet=False,PinType.bIsArray=False,PinType.bIsReference=False,PinType.bIsConst=False,PinType.bIsWeakPointer=False,PersistentGuid=00000000000000000000000000000000,bHidden=True,bNotConnectable=False,bDefaultValueIsReadOnly=False,bDefaultValueIsIgnored=False,bAdvancedView=False,)
+				End Object
+				Begin Object Class=K2Node_CallFunction Name="K2Node_CallFunction_389"
+				   FunctionReference=(MemberParent=BlueprintGeneratedClass'/Game/Prototype/Game/BP_OdinGameMode.BP_OdinGameMode_C',MemberName="End Loading Sequence",MemberGuid=A1D38DD0498B32ACC9A936A032D370AC)
+				   NodePosX=-144
+				   NodePosY=-4496
+				   NodeGuid=4498E305455CCB64B47710BE1ECE74B6
+				   CustomProperties Pin (PinId=B9D0319843BB38F57642CEB7B91707F9,PinName="execute",PinToolTip="\nExec",PinType.PinCategory="exec",PinType.PinSubCategory="",PinType.PinSubCategoryObject=None,PinType.PinSubCategoryMemberReference=(),PinType.PinValueType=(),PinType.bIsMap=False,PinType.bIsSet=False,PinType.bIsArray=False,PinType.bIsReference=False,PinType.bIsConst=False,PinType.bIsWeakPointer=False,LinkedTo=(K2Node_DynamicCast_5 C46EEAC7458FE7A6949E78BA06EE5B54,),PersistentGuid=00000000000000000000000000000000,bHidden=False,bNotConnectable=False,bDefaultValueIsReadOnly=False,bDefaultValueIsIgnored=False,bAdvancedView=False,)
+				   CustomProperties Pin (PinId=080DB6114DA99E281AD26280646A113F,PinName="then",PinToolTip="\nExec",Direction="EGPD_Output",PinType.PinCategory="exec",PinType.PinSubCategory="",PinType.PinSubCategoryObject=None,PinType.PinSubCategoryMemberReference=(),PinType.PinValueType=(),PinType.bIsMap=False,PinType.bIsSet=False,PinType.bIsArray=False,PinType.bIsReference=False,PinType.bIsConst=False,PinType.bIsWeakPointer=False,LinkedTo=(K2Node_CallFunction_147 0D4F8EEB4254987507E56590D0CCA616,),PersistentGuid=00000000000000000000000000000000,bHidden=False,bNotConnectable=False,bDefaultValueIsReadOnly=False,bDefaultValueIsIgnored=False,bAdvancedView=False,)
+				   CustomProperties Pin (PinId=3117651A4CB6ACA2EAEB6BA4272EC84B,PinName="self",PinFriendlyName=NSLOCTEXT("K2Node", "Target", "Target"),PinToolTip="Target\nBP Odin Game Mode Reference",PinType.PinCategory="object",PinType.PinSubCategory="",PinType.PinSubCategoryObject=BlueprintGeneratedClass'/Game/Prototype/Game/BP_OdinGameMode.BP_OdinGameMode_C',PinType.PinSubCategoryMemberReference=(),PinType.PinValueType=(),PinType.bIsMap=False,PinType.bIsSet=False,PinType.bIsArray=False,PinType.bIsReference=False,PinType.bIsConst=False,PinType.bIsWeakPointer=False,LinkedTo=(K2Node_DynamicCast_5 6B739B644EEBFA053240C1B460B2E558,),PersistentGuid=00000000000000000000000000000000,bHidden=False,bNotConnectable=False,bDefaultValueIsReadOnly=False,bDefaultValueIsIgnored=False,bAdvancedView=False,)
+				End Object
+				Begin Object Class=K2Node_CallFunction Name="K2Node_CallFunction_146"
+				   bIsPureFunc=True
+				   FunctionReference=(MemberParent=Class'/Script/Engine.GameplayStatics',MemberName="GetPlayerPawn")
+				   NodePosX=-32
+				   NodePosY=-4352
+				   NodeGuid=5AC649304096413924B030B0014FD935
+				   CustomProperties Pin (PinId=329EDA7E40EFEF35615E988D15DE170A,PinName="self",PinFriendlyName=NSLOCTEXT("K2Node", "Target", "Target"),PinToolTip="Target\nGameplay Statics Reference",PinType.PinCategory="object",PinType.PinSubCategory="",PinType.PinSubCategoryObject=Class'/Script/Engine.GameplayStatics',PinType.PinSubCategoryMemberReference=(),PinType.PinValueType=(),PinType.bIsMap=False,PinType.bIsSet=False,PinType.bIsArray=False,PinType.bIsReference=False,PinType.bIsConst=False,PinType.bIsWeakPointer=False,DefaultObject="/Script/Engine.Default__GameplayStatics",PersistentGuid=00000000000000000000000000000000,bHidden=True,bNotConnectable=False,bDefaultValueIsReadOnly=False,bDefaultValueIsIgnored=False,bAdvancedView=False,)
+				   CustomProperties Pin (PinId=89848130466D916E810B3F99D04CBF1F,PinName="WorldContextObject",PinToolTip="World Context Object\nObject Reference",PinType.PinCategory="object",PinType.PinSubCategory="",PinType.PinSubCategoryObject=Class'/Script/CoreUObject.Object',PinType.PinSubCategoryMemberReference=(),PinType.PinValueType=(),PinType.bIsMap=False,PinType.bIsSet=False,PinType.bIsArray=False,PinType.bIsReference=False,PinType.bIsConst=True,PinType.bIsWeakPointer=False,PersistentGuid=00000000000000000000000000000000,bHidden=True,bNotConnectable=False,bDefaultValueIsReadOnly=False,bDefaultValueIsIgnored=False,bAdvancedView=False,)
+				   CustomProperties Pin (PinId=02F269A847C5688444C9AEAEA78D5124,PinName="PlayerIndex",PinToolTip="Player Index\nInteger",PinType.PinCategory="int",PinType.PinSubCategory="",PinType.PinSubCategoryObject=None,PinType.PinSubCategoryMemberReference=(),PinType.PinValueType=(),PinType.bIsMap=False,PinType.bIsSet=False,PinType.bIsArray=False,PinType.bIsReference=False,PinType.bIsConst=False,PinType.bIsWeakPointer=False,DefaultValue="0",AutogeneratedDefaultValue="0",PersistentGuid=00000000000000000000000000000000,bHidden=False,bNotConnectable=False,bDefaultValueIsReadOnly=False,bDefaultValueIsIgnored=False,bAdvancedView=False,)
+				   CustomProperties Pin (PinId=A55AE1CA4B113E55AEA44FA5BD868FEB,PinName="ReturnValue",PinToolTip="Return Value\nPawn Reference",Direction="EGPD_Output",PinType.PinCategory="object",PinType.PinSubCategory="",PinType.PinSubCategoryObject=Class'/Script/Engine.Pawn',PinType.PinSubCategoryMemberReference=(),PinType.PinValueType=(),PinType.bIsMap=False,PinType.bIsSet=False,PinType.bIsArray=False,PinType.bIsReference=False,PinType.bIsConst=False,PinType.bIsWeakPointer=False,LinkedTo=(K2Node_CallFunction_147 05DD2A284290ED5C1DD79C95F86B089B,),PersistentGuid=00000000000000000000000000000000,bHidden=False,bNotConnectable=False,bDefaultValueIsReadOnly=False,bDefaultValueIsIgnored=False,bAdvancedView=False,)
+				End Object
+				Begin Object Class=K2Node_CallFunction Name="K2Node_CallFunction_147"
+				   FunctionReference=(MemberName="K2_TeleportTo",bSelfContext=True)
+				   NodePosX=400
+				   NodePosY=-4496
+				   NodeGuid=97C3515E434225B49FD8669033C40E7D
+				   CustomProperties Pin (PinId=0D4F8EEB4254987507E56590D0CCA616,PinName="execute",PinToolTip="\nExec",PinType.PinCategory="exec",PinType.PinSubCategory="",PinType.PinSubCategoryObject=None,PinType.PinSubCategoryMemberReference=(),PinType.PinValueType=(),PinType.bIsMap=False,PinType.bIsSet=False,PinType.bIsArray=False,PinType.bIsReference=False,PinType.bIsConst=False,PinType.bIsWeakPointer=False,LinkedTo=(K2Node_CallFunction_389 080DB6114DA99E281AD26280646A113F,),PersistentGuid=00000000000000000000000000000000,bHidden=False,bNotConnectable=False,bDefaultValueIsReadOnly=False,bDefaultValueIsIgnored=False,bAdvancedView=False,)
+				   CustomProperties Pin (PinId=DA74E61542E071192041BF8670B2B044,PinName="then",PinToolTip="\nExec",Direction="EGPD_Output",PinType.PinCategory="exec",PinType.PinSubCategory="",PinType.PinSubCategoryObject=None,PinType.PinSubCategoryMemberReference=(),PinType.PinValueType=(),PinType.bIsMap=False,PinType.bIsSet=False,PinType.bIsArray=False,PinType.bIsReference=False,PinType.bIsConst=False,PinType.bIsWeakPointer=False,PersistentGuid=00000000000000000000000000000000,bHidden=False,bNotConnectable=False,bDefaultValueIsReadOnly=False,bDefaultValueIsIgnored=False,bAdvancedView=False,)
+				   CustomProperties Pin (PinId=05DD2A284290ED5C1DD79C95F86B089B,PinName="self",PinFriendlyName=NSLOCTEXT("K2Node", "Target", "Target"),PinToolTip="Target\nActor Reference",PinType.PinCategory="object",PinType.PinSubCategory="",PinType.PinSubCategoryObject=Class'/Script/Engine.Actor',PinType.PinSubCategoryMemberReference=(),PinType.PinValueType=(),PinType.bIsMap=False,PinType.bIsSet=False,PinType.bIsArray=False,PinType.bIsReference=False,PinType.bIsConst=False,PinType.bIsWeakPointer=False,LinkedTo=(K2Node_CallFunction_146 A55AE1CA4B113E55AEA44FA5BD868FEB,),PersistentGuid=00000000000000000000000000000000,bHidden=False,bNotConnectable=False,bDefaultValueIsReadOnly=False,bDefaultValueIsIgnored=False,bAdvancedView=False,)
+				   CustomProperties Pin (PinId=7ADD16464BF350CD2D3D5DA206062506,PinName="DestLocation",PinToolTip="Dest Location\nVector \n\nThe target destination point",PinType.PinCategory="struct",PinType.PinSubCategory="",PinType.PinSubCategoryObject=ScriptStruct'/Script/CoreUObject.Vector',PinType.PinSubCategoryMemberReference=(),PinType.PinValueType=(),PinType.bIsMap=False,PinType.bIsSet=False,PinType.bIsArray=False,PinType.bIsReference=False,PinType.bIsConst=False,PinType.bIsWeakPointer=False,DefaultValue="0, 0, 0",AutogeneratedDefaultValue="0, 0, 0",LinkedTo=(K2Node_CallFunction_161 8EDCF5B4407E0D67A46F56B3A7EBDC6E,),PersistentGuid=00000000000000000000000000000000,bHidden=False,bNotConnectable=False,bDefaultValueIsReadOnly=False,bDefaultValueIsIgnored=False,bAdvancedView=False,)
+				   CustomProperties Pin (PinId=CADAD3724E0A28E00552BCA72427B545,PinName="DestRotation",PinToolTip="Dest Rotation\nRotator \n\nThe target rotation at the destination",PinType.PinCategory="struct",PinType.PinSubCategory="",PinType.PinSubCategoryObject=ScriptStruct'/Script/CoreUObject.Rotator',PinType.PinSubCategoryMemberReference=(),PinType.PinValueType=(),PinType.bIsMap=False,PinType.bIsSet=False,PinType.bIsArray=False,PinType.bIsReference=False,PinType.bIsConst=False,PinType.bIsWeakPointer=False,DefaultValue="0, 0, 0",AutogeneratedDefaultValue="0, 0, 0",LinkedTo=(K2Node_CallFunction_197 177F737D497F0B71A0D9548711E45134,),PersistentGuid=00000000000000000000000000000000,bHidden=False,bNotConnectable=False,bDefaultValueIsReadOnly=False,bDefaultValueIsIgnored=False,bAdvancedView=False,)
+				   CustomProperties Pin (PinId=150C2F4B41C32C40A5EE5F928C2473FE,PinName="ReturnValue",PinToolTip="Return Value\nBoolean\n\ntrue if the actor has been successfully moved, or false if it couldn\'t fit.",Direction="EGPD_Output",PinType.PinCategory="bool",PinType.PinSubCategory="",PinType.PinSubCategoryObject=None,PinType.PinSubCategoryMemberReference=(),PinType.PinValueType=(),PinType.bIsMap=False,PinType.bIsSet=False,PinType.bIsArray=False,PinType.bIsReference=False,PinType.bIsConst=False,PinType.bIsWeakPointer=False,DefaultValue="false",AutogeneratedDefaultValue="false",PersistentGuid=00000000000000000000000000000000,bHidden=False,bNotConnectable=False,bDefaultValueIsReadOnly=False,bDefaultValueIsIgnored=False,bAdvancedView=False,)
+				End Object
+				Begin Object Class=K2Node_Literal Name="K2Node_Literal_0"
+				   ObjectRef=PlayerStart'/BasicCity/Maps/BasicCityCityCenterDay.BasicCityCityCenterDay:PersistentLevel.PlayerStart2_3'
+				   NodePosX=-272
+				   NodePosY=-4192
+				   NodeGuid=58EDCAD94D643329430D2D879E37B3E5
+				   CustomProperties Pin (PinId=AFA007CD4AFA013315CE299F79589275,PinName="PlayerStart2",PinFriendlyName=NSLOCTEXT("", "C0C388B24DBCBE06EB1E1B8C7FFFE555", "PlayerStart2"),Direction="EGPD_Output",PinType.PinCategory="object",PinType.PinSubCategory="",PinType.PinSubCategoryObject=Class'/Script/Engine.PlayerStart',PinType.PinSubCategoryMemberReference=(),PinType.PinValueType=(),PinType.bIsMap=False,PinType.bIsSet=False,PinType.bIsArray=False,PinType.bIsReference=False,PinType.bIsConst=False,PinType.bIsWeakPointer=False,LinkedTo=(K2Node_CallFunction_161 76593FC04576E3C4E850799E8DD55EC5,K2Node_CallFunction_197 40BECA4844231CB62E8E729102B58F73,),PersistentGuid=00000000000000000000000000000000,bHidden=False,bNotConnectable=False,bDefaultValueIsReadOnly=False,bDefaultValueIsIgnored=False,bAdvancedView=False,)
+				End Object
+				Begin Object Class=K2Node_CallFunction Name="K2Node_CallFunction_161"
+				   bIsPureFunc=True
+				   bIsConstFunc=True
+				   FunctionReference=(MemberName="K2_GetActorLocation",bSelfContext=True)
+				   NodePosX=16
+				   NodePosY=-4272
+				   NodeGuid=72AB37F441B22C2E9CB0D3BFA6D9A48F
+				   CustomProperties Pin (PinId=76593FC04576E3C4E850799E8DD55EC5,PinName="self",PinFriendlyName=NSLOCTEXT("K2Node", "Target", "Target"),PinToolTip="Target\nActor Reference",PinType.PinCategory="object",PinType.PinSubCategory="",PinType.PinSubCategoryObject=Class'/Script/Engine.Actor',PinType.PinSubCategoryMemberReference=(),PinType.PinValueType=(),PinType.bIsMap=False,PinType.bIsSet=False,PinType.bIsArray=False,PinType.bIsReference=False,PinType.bIsConst=False,PinType.bIsWeakPointer=False,LinkedTo=(K2Node_Literal_0 AFA007CD4AFA013315CE299F79589275,),PersistentGuid=00000000000000000000000000000000,bHidden=False,bNotConnectable=False,bDefaultValueIsReadOnly=False,bDefaultValueIsIgnored=False,bAdvancedView=False,)
+				   CustomProperties Pin (PinId=8EDCF5B4407E0D67A46F56B3A7EBDC6E,PinName="ReturnValue",PinToolTip="Return Value\nVector ",Direction="EGPD_Output",PinType.PinCategory="struct",PinType.PinSubCategory="",PinType.PinSubCategoryObject=ScriptStruct'/Script/CoreUObject.Vector',PinType.PinSubCategoryMemberReference=(),PinType.PinValueType=(),PinType.bIsMap=False,PinType.bIsSet=False,PinType.bIsArray=False,PinType.bIsReference=False,PinType.bIsConst=False,PinType.bIsWeakPointer=False,AutogeneratedDefaultValue="0, 0, 0",LinkedTo=(K2Node_CallFunction_147 7ADD16464BF350CD2D3D5DA206062506,),PersistentGuid=00000000000000000000000000000000,bHidden=False,bNotConnectable=False,bDefaultValueIsReadOnly=False,bDefaultValueIsIgnored=False,bAdvancedView=False,)
+				End Object
+				Begin Object Class=K2Node_CallFunction Name="K2Node_CallFunction_197"
+				   bIsPureFunc=True
+				   bIsConstFunc=True
+				   FunctionReference=(MemberName="K2_GetActorRotation",bSelfContext=True)
+				   NodePosX=16
+				   NodePosY=-4176
+				   NodeGuid=AF5DAC314C10340F9EA01999BB42D085
+				   CustomProperties Pin (PinId=40BECA4844231CB62E8E729102B58F73,PinName="self",PinFriendlyName=NSLOCTEXT("K2Node", "Target", "Target"),PinToolTip="Target\nActor Reference",PinType.PinCategory="object",PinType.PinSubCategory="",PinType.PinSubCategoryObject=Class'/Script/Engine.Actor',PinType.PinSubCategoryMemberReference=(),PinType.PinValueType=(),PinType.bIsMap=False,PinType.bIsSet=False,PinType.bIsArray=False,PinType.bIsReference=False,PinType.bIsConst=False,PinType.bIsWeakPointer=False,LinkedTo=(K2Node_Literal_0 AFA007CD4AFA013315CE299F79589275,),PersistentGuid=00000000000000000000000000000000,bHidden=False,bNotConnectable=False,bDefaultValueIsReadOnly=False,bDefaultValueIsIgnored=False,bAdvancedView=False,)
+				   CustomProperties Pin (PinId=177F737D497F0B71A0D9548711E45134,PinName="ReturnValue",PinToolTip="Return Value\nRotator ",Direction="EGPD_Output",PinType.PinCategory="struct",PinType.PinSubCategory="",PinType.PinSubCategoryObject=ScriptStruct'/Script/CoreUObject.Rotator',PinType.PinSubCategoryMemberReference=(),PinType.PinValueType=(),PinType.bIsMap=False,PinType.bIsSet=False,PinType.bIsArray=False,PinType.bIsReference=False,PinType.bIsConst=False,PinType.bIsWeakPointer=False,AutogeneratedDefaultValue="0, 0, 0",LinkedTo=(K2Node_CallFunction_147 CADAD3724E0A28E00552BCA72427B545,),PersistentGuid=00000000000000000000000000000000,bHidden=False,bNotConnectable=False,bDefaultValueIsReadOnly=False,bDefaultValueIsIgnored=False,bAdvancedView=False,)
+				End Object
 		[/PARAM]
 	[/OBJECT]
 
@@ -73,101 +162,3 @@ When you click **Play** in the Level Editor toolbar , you will now start off in
 ![](test.png)
 
 [NAV]
-
-
-[COMMENT:none]
-
-[EXCERPT:snippet]
-	Begin Object Class=K2Node_Event Name="K2Node_Event_52"
-	   EventReference=(MemberParent=Class'/Script/Engine.Actor',MemberName="ReceiveBeginPlay")
-	   bOverrideFunction=True
-	   NodePosX=-928
-	   NodePosY=-4480
-	   NodeGuid=6C849E7241A6108A36D5C1BEAAEAE94F
-	   CustomProperties Pin (PinId=B35535C443B530261C41B8A84FF2CC42,PinName="OutputDelegate",Direction="EGPD_Output",PinType.PinCategory="delegate",PinType.PinSubCategory="",PinType.PinSubCategoryObject=None,PinType.PinSubCategoryMemberReference=(MemberParent=Class'/Script/Engine.Actor',MemberName="ReceiveBeginPlay"),PinType.PinValueType=(),PinType.bIsMap=False,PinType.bIsSet=False,PinType.bIsArray=False,PinType.bIsReference=False,PinType.bIsConst=False,PinType.bIsWeakPointer=False,PersistentGuid=00000000000000000000000000000000,bHidden=False,bNotConnectable=False,bDefaultValueIsReadOnly=False,bDefaultValueIsIgnored=False,bAdvancedView=False,)
-	   CustomProperties Pin (PinId=5A6DBA7144DD625770231A9C725C0B49,PinName="then",Direction="EGPD_Output",PinType.PinCategory="exec",PinType.PinSubCategory="",PinType.PinSubCategoryObject=None,PinType.PinSubCategoryMemberReference=(),PinType.PinValueType=(),PinType.bIsMap=False,PinType.bIsSet=False,PinType.bIsArray=False,PinType.bIsReference=False,PinType.bIsConst=False,PinType.bIsWeakPointer=False,LinkedTo=(K2Node_DynamicCast_5 3228E8FA4B588A45DCEA3A839D739E9A,),PersistentGuid=00000000000000000000000000000000,bHidden=False,bNotConnectable=False,bDefaultValueIsReadOnly=False,bDefaultValueIsIgnored=False,bAdvancedView=False,)
-	End Object
-	Begin Object Class=K2Node_CallFunction Name="K2Node_CallFunction_280"
-	   bIsPureFunc=True
-	   FunctionReference=(MemberParent=Class'/Script/Engine.GameplayStatics',MemberName="GetGameMode")
-	   NodePosX=-720
-	   NodePosY=-4416
-	   NodeGuid=2CB9ABB94CF9F2C3BAEBB0B38ED0F003
-	   CustomProperties Pin (PinId=1AAC46094DE5E3D0CB1F8488BEDBFE3B,PinName="self",PinFriendlyName=NSLOCTEXT("K2Node", "Target", "Target"),PinToolTip="Target\nGameplay Statics Reference",PinType.PinCategory="object",PinType.PinSubCategory="",PinType.PinSubCategoryObject=Class'/Script/Engine.GameplayStatics',PinType.PinSubCategoryMemberReference=(),PinType.PinValueType=(),PinType.bIsMap=False,PinType.bIsSet=False,PinType.bIsArray=False,PinType.bIsReference=False,PinType.bIsConst=False,PinType.bIsWeakPointer=False,DefaultObject="/Script/Engine.Default__GameplayStatics",PersistentGuid=00000000000000000000000000000000,bHidden=True,bNotConnectable=False,bDefaultValueIsReadOnly=False,bDefaultValueIsIgnored=False,bAdvancedView=False,)
-	   CustomProperties Pin (PinId=539E882646828F110DA8539329878814,PinName="WorldContextObject",PinToolTip="World Context Object\nObject Reference",PinType.PinCategory="object",PinType.PinSubCategory="",PinType.PinSubCategoryObject=Class'/Script/CoreUObject.Object',PinType.PinSubCategoryMemberReference=(),PinType.PinValueType=(),PinType.bIsMap=False,PinType.bIsSet=False,PinType.bIsArray=False,PinType.bIsReference=False,PinType.bIsConst=True,PinType.bIsWeakPointer=False,PersistentGuid=00000000000000000000000000000000,bHidden=True,bNotConnectable=False,bDefaultValueIsReadOnly=False,bDefaultValueIsIgnored=False,bAdvancedView=False,)
-	   CustomProperties Pin (PinId=F81E36D8454550AAAEF10DB57E9FC71F,PinName="ReturnValue",PinToolTip="Return Value\nGame Mode Base Reference",Direction="EGPD_Output",PinType.PinCategory="object",PinType.PinSubCategory="",PinType.PinSubCategoryObject=Class'/Script/Engine.GameModeBase',PinType.PinSubCategoryMemberReference=(),PinType.PinValueType=(),PinType.bIsMap=False,PinType.bIsSet=False,PinType.bIsArray=False,PinType.bIsReference=False,PinType.bIsConst=False,PinType.bIsWeakPointer=False,LinkedTo=(K2Node_DynamicCast_5 8B0B2634480FE02DE82004ADB105AB79,),PersistentGuid=00000000000000000000000000000000,bHidden=False,bNotConnectable=False,bDefaultValueIsReadOnly=False,bDefaultValueIsIgnored=False,bAdvancedView=False,)
-	End Object
-	Begin Object Class=K2Node_DynamicCast Name="K2Node_DynamicCast_5"
-	   TargetType=BlueprintGeneratedClass'/Game/Prototype/Game/BP_OdinGameMode.BP_OdinGameMode_C'
-	   NodePosX=-512
-	   NodePosY=-4480
-	   NodeGuid=2A2928BD4CBEFB243D86D796866C7C1B
-	   CustomProperties Pin (PinId=3228E8FA4B588A45DCEA3A839D739E9A,PinName="execute",PinType.PinCategory="exec",PinType.PinSubCategory="",PinType.PinSubCategoryObject=None,PinType.PinSubCategoryMemberReference=(),PinType.PinValueType=(),PinType.bIsMap=False,PinType.bIsSet=False,PinType.bIsArray=False,PinType.bIsReference=False,PinType.bIsConst=False,PinType.bIsWeakPointer=False,LinkedTo=(K2Node_Event_52 5A6DBA7144DD625770231A9C725C0B49,),PersistentGuid=00000000000000000000000000000000,bHidden=False,bNotConnectable=False,bDefaultValueIsReadOnly=False,bDefaultValueIsIgnored=False,bAdvancedView=False,)
-	   CustomProperties Pin (PinId=C46EEAC7458FE7A6949E78BA06EE5B54,PinName="then",Direction="EGPD_Output",PinType.PinCategory="exec",PinType.PinSubCategory="",PinType.PinSubCategoryObject=None,PinType.PinSubCategoryMemberReference=(),PinType.PinValueType=(),PinType.bIsMap=False,PinType.bIsSet=False,PinType.bIsArray=False,PinType.bIsReference=False,PinType.bIsConst=False,PinType.bIsWeakPointer=False,LinkedTo=(K2Node_CallFunction_389 B9D0319843BB38F57642CEB7B91707F9,),PersistentGuid=00000000000000000000000000000000,bHidden=False,bNotConnectable=False,bDefaultValueIsReadOnly=False,bDefaultValueIsIgnored=False,bAdvancedView=False,)
-	   CustomProperties Pin (PinId=9EAA2A464990955344C076A409669F6D,PinName="CastFailed",Direction="EGPD_Output",PinType.PinCategory="exec",PinType.PinSubCategory="",PinType.PinSubCategoryObject=None,PinType.PinSubCategoryMemberReference=(),PinType.PinValueType=(),PinType.bIsMap=False,PinType.bIsSet=False,PinType.bIsArray=False,PinType.bIsReference=False,PinType.bIsConst=False,PinType.bIsWeakPointer=False,PersistentGuid=00000000000000000000000000000000,bHidden=False,bNotConnectable=False,bDefaultValueIsReadOnly=False,bDefaultValueIsIgnored=False,bAdvancedView=False,)
-	   CustomProperties Pin (PinId=8B0B2634480FE02DE82004ADB105AB79,PinName="Object",PinType.PinCategory="object",PinType.PinSubCategory="",PinType.PinSubCategoryObject=Class'/Script/CoreUObject.Object',PinType.PinSubCategoryMemberReference=(),PinType.PinValueType=(),PinType.bIsMap=False,PinType.bIsSet=False,PinType.bIsArray=False,PinType.bIsReference=False,PinType.bIsConst=False,PinType.bIsWeakPointer=False,LinkedTo=(K2Node_CallFunction_280 F81E36D8454550AAAEF10DB57E9FC71F,),PersistentGuid=00000000000000000000000000000000,bHidden=False,bNotConnectable=False,bDefaultValueIsReadOnly=False,bDefaultValueIsIgnored=False,bAdvancedView=False,)
-	   CustomProperties Pin (PinId=6B739B644EEBFA053240C1B460B2E558,PinName="AsBP Odin Game Mode",Direction="EGPD_Output",PinType.PinCategory="object",PinType.PinSubCategory="",PinType.PinSubCategoryObject=BlueprintGeneratedClass'/Game/Prototype/Game/BP_OdinGameMode.BP_OdinGameMode_C',PinType.PinSubCategoryMemberReference=(),PinType.PinValueType=(),PinType.bIsMap=False,PinType.bIsSet=False,PinType.bIsArray=False,PinType.bIsReference=False,PinType.bIsConst=False,PinType.bIsWeakPointer=False,LinkedTo=(K2Node_CallFunction_389 3117651A4CB6ACA2EAEB6BA4272EC84B,),PersistentGuid=00000000000000000000000000000000,bHidden=False,bNotConnectable=False,bDefaultValueIsReadOnly=False,bDefaultValueIsIgnored=False,bAdvancedView=False,)
-	   CustomProperties Pin (PinId=BEB2F603410D2A627B7D089246E322C7,PinName="bSuccess",Direction="EGPD_Output",PinType.PinCategory="bool",PinType.PinSubCategory="",PinType.PinSubCategoryObject=None,PinType.PinSubCategoryMemberReference=(),PinType.PinValueType=(),PinType.bIsMap=False,PinType.bIsSet=False,PinType.bIsArray=False,PinType.bIsReference=False,PinType.bIsConst=False,PinType.bIsWeakPointer=False,PersistentGuid=00000000000000000000000000000000,bHidden=True,bNotConnectable=False,bDefaultValueIsReadOnly=False,bDefaultValueIsIgnored=False,bAdvancedView=False,)
-	End Object
-	Begin Object Class=K2Node_CallFunction Name="K2Node_CallFunction_389"
-	   FunctionReference=(MemberParent=BlueprintGeneratedClass'/Game/Prototype/Game/BP_OdinGameMode.BP_OdinGameMode_C',MemberName="End Loading Sequence",MemberGuid=A1D38DD0498B32ACC9A936A032D370AC)
-	   NodePosX=-144
-	   NodePosY=-4496
-	   NodeGuid=4498E305455CCB64B47710BE1ECE74B6
-	   CustomProperties Pin (PinId=B9D0319843BB38F57642CEB7B91707F9,PinName="execute",PinToolTip="\nExec",PinType.PinCategory="exec",PinType.PinSubCategory="",PinType.PinSubCategoryObject=None,PinType.PinSubCategoryMemberReference=(),PinType.PinValueType=(),PinType.bIsMap=False,PinType.bIsSet=False,PinType.bIsArray=False,PinType.bIsReference=False,PinType.bIsConst=False,PinType.bIsWeakPointer=False,LinkedTo=(K2Node_DynamicCast_5 C46EEAC7458FE7A6949E78BA06EE5B54,),PersistentGuid=00000000000000000000000000000000,bHidden=False,bNotConnectable=False,bDefaultValueIsReadOnly=False,bDefaultValueIsIgnored=False,bAdvancedView=False,)
-	   CustomProperties Pin (PinId=080DB6114DA99E281AD26280646A113F,PinName="then",PinToolTip="\nExec",Direction="EGPD_Output",PinType.PinCategory="exec",PinType.PinSubCategory="",PinType.PinSubCategoryObject=None,PinType.PinSubCategoryMemberReference=(),PinType.PinValueType=(),PinType.bIsMap=False,PinType.bIsSet=False,PinType.bIsArray=False,PinType.bIsReference=False,PinType.bIsConst=False,PinType.bIsWeakPointer=False,LinkedTo=(K2Node_CallFunction_147 0D4F8EEB4254987507E56590D0CCA616,),PersistentGuid=00000000000000000000000000000000,bHidden=False,bNotConnectable=False,bDefaultValueIsReadOnly=False,bDefaultValueIsIgnored=False,bAdvancedView=False,)
-	   CustomProperties Pin (PinId=3117651A4CB6ACA2EAEB6BA4272EC84B,PinName="self",PinFriendlyName=NSLOCTEXT("K2Node", "Target", "Target"),PinToolTip="Target\nBP Odin Game Mode Reference",PinType.PinCategory="object",PinType.PinSubCategory="",PinType.PinSubCategoryObject=BlueprintGeneratedClass'/Game/Prototype/Game/BP_OdinGameMode.BP_OdinGameMode_C',PinType.PinSubCategoryMemberReference=(),PinType.PinValueType=(),PinType.bIsMap=False,PinType.bIsSet=False,PinType.bIsArray=False,PinType.bIsReference=False,PinType.bIsConst=False,PinType.bIsWeakPointer=False,LinkedTo=(K2Node_DynamicCast_5 6B739B644EEBFA053240C1B460B2E558,),PersistentGuid=00000000000000000000000000000000,bHidden=False,bNotConnectable=False,bDefaultValueIsReadOnly=False,bDefaultValueIsIgnored=False,bAdvancedView=False,)
-	End Object
-	Begin Object Class=K2Node_CallFunction Name="K2Node_CallFunction_146"
-	   bIsPureFunc=True
-	   FunctionReference=(MemberParent=Class'/Script/Engine.GameplayStatics',MemberName="GetPlayerPawn")
-	   NodePosX=-32
-	   NodePosY=-4352
-	   NodeGuid=5AC649304096413924B030B0014FD935
-	   CustomProperties Pin (PinId=329EDA7E40EFEF35615E988D15DE170A,PinName="self",PinFriendlyName=NSLOCTEXT("K2Node", "Target", "Target"),PinToolTip="Target\nGameplay Statics Reference",PinType.PinCategory="object",PinType.PinSubCategory="",PinType.PinSubCategoryObject=Class'/Script/Engine.GameplayStatics',PinType.PinSubCategoryMemberReference=(),PinType.PinValueType=(),PinType.bIsMap=False,PinType.bIsSet=False,PinType.bIsArray=False,PinType.bIsReference=False,PinType.bIsConst=False,PinType.bIsWeakPointer=False,DefaultObject="/Script/Engine.Default__GameplayStatics",PersistentGuid=00000000000000000000000000000000,bHidden=True,bNotConnectable=False,bDefaultValueIsReadOnly=False,bDefaultValueIsIgnored=False,bAdvancedView=False,)
-	   CustomProperties Pin (PinId=89848130466D916E810B3F99D04CBF1F,PinName="WorldContextObject",PinToolTip="World Context Object\nObject Reference",PinType.PinCategory="object",PinType.PinSubCategory="",PinType.PinSubCategoryObject=Class'/Script/CoreUObject.Object',PinType.PinSubCategoryMemberReference=(),PinType.PinValueType=(),PinType.bIsMap=False,PinType.bIsSet=False,PinType.bIsArray=False,PinType.bIsReference=False,PinType.bIsConst=True,PinType.bIsWeakPointer=False,PersistentGuid=00000000000000000000000000000000,bHidden=True,bNotConnectable=False,bDefaultValueIsReadOnly=False,bDefaultValueIsIgnored=False,bAdvancedView=False,)
-	   CustomProperties Pin (PinId=02F269A847C5688444C9AEAEA78D5124,PinName="PlayerIndex",PinToolTip="Player Index\nInteger",PinType.PinCategory="int",PinType.PinSubCategory="",PinType.PinSubCategoryObject=None,PinType.PinSubCategoryMemberReference=(),PinType.PinValueType=(),PinType.bIsMap=False,PinType.bIsSet=False,PinType.bIsArray=False,PinType.bIsReference=False,PinType.bIsConst=False,PinType.bIsWeakPointer=False,DefaultValue="0",AutogeneratedDefaultValue="0",PersistentGuid=00000000000000000000000000000000,bHidden=False,bNotConnectable=False,bDefaultValueIsReadOnly=False,bDefaultValueIsIgnored=False,bAdvancedView=False,)
-	   CustomProperties Pin (PinId=A55AE1CA4B113E55AEA44FA5BD868FEB,PinName="ReturnValue",PinToolTip="Return Value\nPawn Reference",Direction="EGPD_Output",PinType.PinCategory="object",PinType.PinSubCategory="",PinType.PinSubCategoryObject=Class'/Script/Engine.Pawn',PinType.PinSubCategoryMemberReference=(),PinType.PinValueType=(),PinType.bIsMap=False,PinType.bIsSet=False,PinType.bIsArray=False,PinType.bIsReference=False,PinType.bIsConst=False,PinType.bIsWeakPointer=False,LinkedTo=(K2Node_CallFunction_147 05DD2A284290ED5C1DD79C95F86B089B,),PersistentGuid=00000000000000000000000000000000,bHidden=False,bNotConnectable=False,bDefaultValueIsReadOnly=False,bDefaultValueIsIgnored=False,bAdvancedView=False,)
-	End Object
-	Begin Object Class=K2Node_CallFunction Name="K2Node_CallFunction_147"
-	   FunctionReference=(MemberName="K2_TeleportTo",bSelfContext=True)
-	   NodePosX=400
-	   NodePosY=-4496
-	   NodeGuid=97C3515E434225B49FD8669033C40E7D
-	   CustomProperties Pin (PinId=0D4F8EEB4254987507E56590D0CCA616,PinName="execute",PinToolTip="\nExec",PinType.PinCategory="exec",PinType.PinSubCategory="",PinType.PinSubCategoryObject=None,PinType.PinSubCategoryMemberReference=(),PinType.PinValueType=(),PinType.bIsMap=False,PinType.bIsSet=False,PinType.bIsArray=False,PinType.bIsReference=False,PinType.bIsConst=False,PinType.bIsWeakPointer=False,LinkedTo=(K2Node_CallFunction_389 080DB6114DA99E281AD26280646A113F,),PersistentGuid=00000000000000000000000000000000,bHidden=False,bNotConnectable=False,bDefaultValueIsReadOnly=False,bDefaultValueIsIgnored=False,bAdvancedView=False,)
-	   CustomProperties Pin (PinId=DA74E61542E071192041BF8670B2B044,PinName="then",PinToolTip="\nExec",Direction="EGPD_Output",PinType.PinCategory="exec",PinType.PinSubCategory="",PinType.PinSubCategoryObject=None,PinType.PinSubCategoryMemberReference=(),PinType.PinValueType=(),PinType.bIsMap=False,PinType.bIsSet=False,PinType.bIsArray=False,PinType.bIsReference=False,PinType.bIsConst=False,PinType.bIsWeakPointer=False,PersistentGuid=00000000000000000000000000000000,bHidden=False,bNotConnectable=False,bDefaultValueIsReadOnly=False,bDefaultValueIsIgnored=False,bAdvancedView=False,)
-	   CustomProperties Pin (PinId=05DD2A284290ED5C1DD79C95F86B089B,PinName="self",PinFriendlyName=NSLOCTEXT("K2Node", "Target", "Target"),PinToolTip="Target\nActor Reference",PinType.PinCategory="object",PinType.PinSubCategory="",PinType.PinSubCategoryObject=Class'/Script/Engine.Actor',PinType.PinSubCategoryMemberReference=(),PinType.PinValueType=(),PinType.bIsMap=False,PinType.bIsSet=False,PinType.bIsArray=False,PinType.bIsReference=False,PinType.bIsConst=False,PinType.bIsWeakPointer=False,LinkedTo=(K2Node_CallFunction_146 A55AE1CA4B113E55AEA44FA5BD868FEB,),PersistentGuid=00000000000000000000000000000000,bHidden=False,bNotConnectable=False,bDefaultValueIsReadOnly=False,bDefaultValueIsIgnored=False,bAdvancedView=False,)
-	   CustomProperties Pin (PinId=7ADD16464BF350CD2D3D5DA206062506,PinName="DestLocation",PinToolTip="Dest Location\nVector \n\nThe target destination point",PinType.PinCategory="struct",PinType.PinSubCategory="",PinType.PinSubCategoryObject=ScriptStruct'/Script/CoreUObject.Vector',PinType.PinSubCategoryMemberReference=(),PinType.PinValueType=(),PinType.bIsMap=False,PinType.bIsSet=False,PinType.bIsArray=False,PinType.bIsReference=False,PinType.bIsConst=False,PinType.bIsWeakPointer=False,DefaultValue="0, 0, 0",AutogeneratedDefaultValue="0, 0, 0",LinkedTo=(K2Node_CallFunction_161 8EDCF5B4407E0D67A46F56B3A7EBDC6E,),PersistentGuid=00000000000000000000000000000000,bHidden=False,bNotConnectable=False,bDefaultValueIsReadOnly=False,bDefaultValueIsIgnored=False,bAdvancedView=False,)
-	   CustomProperties Pin (PinId=CADAD3724E0A28E00552BCA72427B545,PinName="DestRotation",PinToolTip="Dest Rotation\nRotator \n\nThe target rotation at the destination",PinType.PinCategory="struct",PinType.PinSubCategory="",PinType.PinSubCategoryObject=ScriptStruct'/Script/CoreUObject.Rotator',PinType.PinSubCategoryMemberReference=(),PinType.PinValueType=(),PinType.bIsMap=False,PinType.bIsSet=False,PinType.bIsArray=False,PinType.bIsReference=False,PinType.bIsConst=False,PinType.bIsWeakPointer=False,DefaultValue="0, 0, 0",AutogeneratedDefaultValue="0, 0, 0",LinkedTo=(K2Node_CallFunction_197 177F737D497F0B71A0D9548711E45134,),PersistentGuid=00000000000000000000000000000000,bHidden=False,bNotConnectable=False,bDefaultValueIsReadOnly=False,bDefaultValueIsIgnored=False,bAdvancedView=False,)
-	   CustomProperties Pin (PinId=150C2F4B41C32C40A5EE5F928C2473FE,PinName="ReturnValue",PinToolTip="Return Value\nBoolean\n\ntrue if the actor has been successfully moved, or false if it couldn\'t fit.",Direction="EGPD_Output",PinType.PinCategory="bool",PinType.PinSubCategory="",PinType.PinSubCategoryObject=None,PinType.PinSubCategoryMemberReference=(),PinType.PinValueType=(),PinType.bIsMap=False,PinType.bIsSet=False,PinType.bIsArray=False,PinType.bIsReference=False,PinType.bIsConst=False,PinType.bIsWeakPointer=False,DefaultValue="false",AutogeneratedDefaultValue="false",PersistentGuid=00000000000000000000000000000000,bHidden=False,bNotConnectable=False,bDefaultValueIsReadOnly=False,bDefaultValueIsIgnored=False,bAdvancedView=False,)
-	End Object
-	Begin Object Class=K2Node_Literal Name="K2Node_Literal_0"
-	   ObjectRef=PlayerStart'/BasicCity/Maps/BasicCityCityCenterDay.BasicCityCityCenterDay:PersistentLevel.PlayerStart2_3'
-	   NodePosX=-272
-	   NodePosY=-4192
-	   NodeGuid=58EDCAD94D643329430D2D879E37B3E5
-	   CustomProperties Pin (PinId=AFA007CD4AFA013315CE299F79589275,PinName="PlayerStart2",PinFriendlyName=NSLOCTEXT("", "C0C388B24DBCBE06EB1E1B8C7FFFE555", "PlayerStart2"),Direction="EGPD_Output",PinType.PinCategory="object",PinType.PinSubCategory="",PinType.PinSubCategoryObject=Class'/Script/Engine.PlayerStart',PinType.PinSubCategoryMemberReference=(),PinType.PinValueType=(),PinType.bIsMap=False,PinType.bIsSet=False,PinType.bIsArray=False,PinType.bIsReference=False,PinType.bIsConst=False,PinType.bIsWeakPointer=False,LinkedTo=(K2Node_CallFunction_161 76593FC04576E3C4E850799E8DD55EC5,K2Node_CallFunction_197 40BECA4844231CB62E8E729102B58F73,),PersistentGuid=00000000000000000000000000000000,bHidden=False,bNotConnectable=False,bDefaultValueIsReadOnly=False,bDefaultValueIsIgnored=False,bAdvancedView=False,)
-	End Object
-	Begin Object Class=K2Node_CallFunction Name="K2Node_CallFunction_161"
-	   bIsPureFunc=True
-	   bIsConstFunc=True
-	   FunctionReference=(MemberName="K2_GetActorLocation",bSelfContext=True)
-	   NodePosX=16
-	   NodePosY=-4272
-	   NodeGuid=72AB37F441B22C2E9CB0D3BFA6D9A48F
-	   CustomProperties Pin (PinId=76593FC04576E3C4E850799E8DD55EC5,PinName="self",PinFriendlyName=NSLOCTEXT("K2Node", "Target", "Target"),PinToolTip="Target\nActor Reference",PinType.PinCategory="object",PinType.PinSubCategory="",PinType.PinSubCategoryObject=Class'/Script/Engine.Actor',PinType.PinSubCategoryMemberReference=(),PinType.PinValueType=(),PinType.bIsMap=False,PinType.bIsSet=False,PinType.bIsArray=False,PinType.bIsReference=False,PinType.bIsConst=False,PinType.bIsWeakPointer=False,LinkedTo=(K2Node_Literal_0 AFA007CD4AFA013315CE299F79589275,),PersistentGuid=00000000000000000000000000000000,bHidden=False,bNotConnectable=False,bDefaultValueIsReadOnly=False,bDefaultValueIsIgnored=False,bAdvancedView=False,)
-	   CustomProperties Pin (PinId=8EDCF5B4407E0D67A46F56B3A7EBDC6E,PinName="ReturnValue",PinToolTip="Return Value\nVector ",Direction="EGPD_Output",PinType.PinCategory="struct",PinType.PinSubCategory="",PinType.PinSubCategoryObject=ScriptStruct'/Script/CoreUObject.Vector',PinType.PinSubCategoryMemberReference=(),PinType.PinValueType=(),PinType.bIsMap=False,PinType.bIsSet=False,PinType.bIsArray=False,PinType.bIsReference=False,PinType.bIsConst=False,PinType.bIsWeakPointer=False,AutogeneratedDefaultValue="0, 0, 0",LinkedTo=(K2Node_CallFunction_147 7ADD16464BF350CD2D3D5DA206062506,),PersistentGuid=00000000000000000000000000000000,bHidden=False,bNotConnectable=False,bDefaultValueIsReadOnly=False,bDefaultValueIsIgnored=False,bAdvancedView=False,)
-	End Object
-	Begin Object Class=K2Node_CallFunction Name="K2Node_CallFunction_197"
-	   bIsPureFunc=True
-	   bIsConstFunc=True
-	   FunctionReference=(MemberName="K2_GetActorRotation",bSelfContext=True)
-	   NodePosX=16
-	   NodePosY=-4176
-	   NodeGuid=AF5DAC314C10340F9EA01999BB42D085
-	   CustomProperties Pin (PinId=40BECA4844231CB62E8E729102B58F73,PinName="self",PinFriendlyName=NSLOCTEXT("K2Node", "Target", "Target"),PinToolTip="Target\nActor Reference",PinType.PinCategory="object",PinType.PinSubCategory="",PinType.PinSubCategoryObject=Class'/Script/Engine.Actor',PinType.PinSubCategoryMemberReference=(),PinType.PinValueType=(),PinType.bIsMap=False,PinType.bIsSet=False,PinType.bIsArray=False,PinType.bIsReference=False,PinType.bIsConst=False,PinType.bIsWeakPointer=False,LinkedTo=(K2Node_Literal_0 AFA007CD4AFA013315CE299F79589275,),PersistentGuid=00000000000000000000000000000000,bHidden=False,bNotConnectable=False,bDefaultValueIsReadOnly=False,bDefaultValueIsIgnored=False,bAdvancedView=False,)
-	   CustomProperties Pin (PinId=177F737D497F0B71A0D9548711E45134,PinName="ReturnValue",PinToolTip="Return Value\nRotator ",Direction="EGPD_Output",PinType.PinCategory="struct",PinType.PinSubCategory="",PinType.PinSubCategoryObject=ScriptStruct'/Script/CoreUObject.Rotator',PinType.PinSubCategoryMemberReference=(),PinType.PinValueType=(),PinType.bIsMap=False,PinType.bIsSet=False,PinType.bIsArray=False,PinType.bIsReference=False,PinType.bIsConst=False,PinType.bIsWeakPointer=False,AutogeneratedDefaultValue="0, 0, 0",LinkedTo=(K2Node_CallFunction_147 CADAD3724E0A28E00552BCA72427B545,),PersistentGuid=00000000000000000000000000000000,bHidden=False,bNotConnectable=False,bDefaultValueIsReadOnly=False,bDefaultValueIsIgnored=False,bAdvancedView=False,)
-	End Object
-[/EXCERPT:snippet]
-
-[/COMMENT]
\ No newline at end of file
