Copying //UE4/Dev-Documentation to //UE4/Samples-Main (Source: //UE4/Dev-Documentation @ 4091886)

diff --git a/Engine/Documentation/Source/Modding/RoboRecall/HowTo/SetupGameplay/EndLevelScript/EndLevelScript.INT.udn b/Engine/Documentation/Source/Modding/RoboRecall/HowTo/SetupGameplay/EndLevelScript/EndLevelScript.INT.udn
--- a/Engine/Documentation/Source/Modding/RoboRecall/HowTo/SetupGameplay/EndLevelScript/EndLevelScript.INT.udn
+++ b/Engine/Documentation/Source/Modding/RoboRecall/HowTo/SetupGameplay/EndLevelScript/EndLevelScript.INT.udn
@@ -54,7 +54,96 @@ When all the robots in the wave have been recalled and the level is complete, yo
 		![](console_command.png)(convert:false)
 		[/PARAM]
 		[PARAM:snippet]
-			[INCLUDE:#snippet]
+				Begin Object Class=K2Node_Event Name="K2Node_Event_52"
+				   EventReference=(MemberParent=Class'/Script/Engine.Actor',MemberName="ReceiveBeginPlay")
+				   bOverrideFunction=True
+				   NodePosX=-928
+				   NodePosY=-4480
+				   NodeGuid=6C849E7241A6108A36D5C1BEAAEAE94F
+				   CustomProperties Pin (PinId=B35535C443B530261C41B8A84FF2CC42,PinName="OutputDelegate",Direction="EGPD_Output",PinType.PinCategory="delegate",PinType.PinSubCategory="",PinType.PinSubCategoryObject=None,PinType.PinSubCategoryMemberReference=(MemberParent=Class'/Script/Engine.Actor',MemberName="ReceiveBeginPlay"),PinType.PinValueType=(),PinType.bIsMap=False,PinType.bIsSet=False,PinType.bIsArray=False,PinType.bIsReference=False,PinType.bIsConst=False,PinType.bIsWeakPointer=False,PersistentGuid=00000000000000000000000000000000,bHidden=False,bNotConnectable=False,bDefaultValueIsReadOnly=False,bDefaultValueIsIgnored=False,bAdvancedView=False,)
+				   CustomProperties Pin (PinId=5A6DBA7144DD625770231A9C725C0B49,PinName="then",Direction="EGPD_Output",PinType.PinCategory="exec",PinType.PinSubCategory="",PinType.PinSubCategoryObject=None,PinType.PinSubCategoryMemberReference=(),PinType.PinValueType=(),PinType.bIsMap=False,PinType.bIsSet=False,PinType.bIsArray=False,PinType.bIsReference=False,PinType.bIsConst=False,PinType.bIsWeakPointer=False,LinkedTo=(K2Node_DynamicCast_5 3228E8FA4B588A45DCEA3A839D739E9A,),PersistentGuid=00000000000000000000000000000000,bHidden=False,bNotConnectable=False,bDefaultValueIsReadOnly=False,bDefaultValueIsIgnored=False,bAdvancedView=False,)
+				End Object
+				Begin Object Class=K2Node_CallFunction Name="K2Node_CallFunction_280"
+				   bIsPureFunc=True
+				   FunctionReference=(MemberParent=Class'/Script/Engine.GameplayStatics',MemberName="GetGameMode")
+				   NodePosX=-720
+				   NodePosY=-4416
+				   NodeGuid=2CB9ABB94CF9F2C3BAEBB0B38ED0F003
+				   CustomProperties Pin (PinId=1AAC46094DE5E3D0CB1F8488BEDBFE3B,PinName="self",PinFriendlyName=NSLOCTEXT("K2Node", "Target", "Target"),PinToolTip="Target\nGameplay Statics Reference",PinType.PinCategory="object",PinType.PinSubCategory="",PinType.PinSubCategoryObject=Class'/Script/Engine.GameplayStatics',PinType.PinSubCategoryMemberReference=(),PinType.PinValueType=(),PinType.bIsMap=False,PinType.bIsSet=False,PinType.bIsArray=False,PinType.bIsReference=False,PinType.bIsConst=False,PinType.bIsWeakPointer=False,DefaultObject="/Script/Engine.Default__GameplayStatics",PersistentGuid=00000000000000000000000000000000,bHidden=True,bNotConnectable=False,bDefaultValueIsReadOnly=False,bDefaultValueIsIgnored=False,bAdvancedView=False,)
+				   CustomProperties Pin (PinId=539E882646828F110DA8539329878814,PinName="WorldContextObject",PinToolTip="World Context Object\nObject Reference",PinType.PinCategory="object",PinType.PinSubCategory="",PinType.PinSubCategoryObject=Class'/Script/CoreUObject.Object',PinType.PinSubCategoryMemberReference=(),PinType.PinValueType=(),PinType.bIsMap=False,PinType.bIsSet=False,PinType.bIsArray=False,PinType.bIsReference=False,PinType.bIsConst=True,PinType.bIsWeakPointer=False,PersistentGuid=00000000000000000000000000000000,bHidden=True,bNotConnectable=False,bDefaultValueIsReadOnly=False,bDefaultValueIsIgnored=False,bAdvancedView=False,)
+				   CustomProperties Pin (PinId=F81E36D8454550AAAEF10DB57E9FC71F,PinName="ReturnValue",PinToolTip="Return Value\nGame Mode Base Reference",Direction="EGPD_Output",PinType.PinCategory="object",PinType.PinSubCategory="",PinType.PinSubCategoryObject=Class'/Script/Engine.GameModeBase',PinType.PinSubCategoryMemberReference=(),PinType.PinValueType=(),PinType.bIsMap=False,PinType.bIsSet=False,PinType.bIsArray=False,PinType.bIsReference=False,PinType.bIsConst=False,PinType.bIsWeakPointer=False,LinkedTo=(K2Node_DynamicCast_5 8B0B2634480FE02DE82004ADB105AB79,),PersistentGuid=00000000000000000000000000000000,bHidden=False,bNotConnectable=False,bDefaultValueIsReadOnly=False,bDefaultValueIsIgnored=False,bAdvancedView=False,)
+				End Object
+				Begin Object Class=K2Node_DynamicCast Name="K2Node_DynamicCast_5"
+				   TargetType=BlueprintGeneratedClass'/Game/Prototype/Game/BP_OdinGameMode.BP_OdinGameMode_C'
+				   NodePosX=-512
+				   NodePosY=-4480
+				   NodeGuid=2A2928BD4CBEFB243D86D796866C7C1B
+				   CustomProperties Pin (PinId=3228E8FA4B588A45DCEA3A839D739E9A,PinName="execute",PinType.PinCategory="exec",PinType.PinSubCategory="",PinType.PinSubCategoryObject=None,PinType.PinSubCategoryMemberReference=(),PinType.PinValueType=(),PinType.bIsMap=False,PinType.bIsSet=False,PinType.bIsArray=False,PinType.bIsReference=False,PinType.bIsConst=False,PinType.bIsWeakPointer=False,LinkedTo=(K2Node_Event_52 5A6DBA7144DD625770231A9C725C0B49,),PersistentGuid=00000000000000000000000000000000,bHidden=False,bNotConnectable=False,bDefaultValueIsReadOnly=False,bDefaultValueIsIgnored=False,bAdvancedView=False,)
+				   CustomProperties Pin (PinId=C46EEAC7458FE7A6949E78BA06EE5B54,PinName="then",Direction="EGPD_Output",PinType.PinCategory="exec",PinType.PinSubCategory="",PinType.PinSubCategoryObject=None,PinType.PinSubCategoryMemberReference=(),PinType.PinValueType=(),PinType.bIsMap=False,PinType.bIsSet=False,PinType.bIsArray=False,PinType.bIsReference=False,PinType.bIsConst=False,PinType.bIsWeakPointer=False,LinkedTo=(K2Node_CallFunction_389 B9D0319843BB38F57642CEB7B91707F9,),PersistentGuid=00000000000000000000000000000000,bHidden=False,bNotConnectable=False,bDefaultValueIsReadOnly=False,bDefaultValueIsIgnored=False,bAdvancedView=False,)
+				   CustomProperties Pin (PinId=9EAA2A464990955344C076A409669F6D,PinName="CastFailed",Direction="EGPD_Output",PinType.PinCategory="exec",PinType.PinSubCategory="",PinType.PinSubCategoryObject=None,PinType.PinSubCategoryMemberReference=(),PinType.PinValueType=(),PinType.bIsMap=False,PinType.bIsSet=False,PinType.bIsArray=False,PinType.bIsReference=False,PinType.bIsConst=False,PinType.bIsWeakPointer=False,PersistentGuid=00000000000000000000000000000000,bHidden=False,bNotConnectable=False,bDefaultValueIsReadOnly=False,bDefaultValueIsIgnored=False,bAdvancedView=False,)
+				   CustomProperties Pin (PinId=8B0B2634480FE02DE82004ADB105AB79,PinName="Object",PinType.PinCategory="object",PinType.PinSubCategory="",PinType.PinSubCategoryObject=Class'/Script/CoreUObject.Object',PinType.PinSubCategoryMemberReference=(),PinType.PinValueType=(),PinType.bIsMap=False,PinType.bIsSet=False,PinType.bIsArray=False,PinType.bIsReference=False,PinType.bIsConst=False,PinType.bIsWeakPointer=False,LinkedTo=(K2Node_CallFunction_280 F81E36D8454550AAAEF10DB57E9FC71F,),PersistentGuid=00000000000000000000000000000000,bHidden=False,bNotConnectable=False,bDefaultValueIsReadOnly=False,bDefaultValueIsIgnored=False,bAdvancedView=False,)
+				   CustomProperties Pin (PinId=6B739B644EEBFA053240C1B460B2E558,PinName="AsBP Odin Game Mode",Direction="EGPD_Output",PinType.PinCategory="object",PinType.PinSubCategory="",PinType.PinSubCategoryObject=BlueprintGeneratedClass'/Game/Prototype/Game/BP_OdinGameMode.BP_OdinGameMode_C',PinType.PinSubCategoryMemberReference=(),PinType.PinValueType=(),PinType.bIsMap=False,PinType.bIsSet=False,PinType.bIsArray=False,PinType.bIsReference=False,PinType.bIsConst=False,PinType.bIsWeakPointer=False,LinkedTo=(K2Node_CallFunction_389 3117651A4CB6ACA2EAEB6BA4272EC84B,),PersistentGuid=00000000000000000000000000000000,bHidden=False,bNotConnectable=False,bDefaultValueIsReadOnly=False,bDefaultValueIsIgnored=False,bAdvancedView=False,)
+				   CustomProperties Pin (PinId=BEB2F603410D2A627B7D089246E322C7,PinName="bSuccess",Direction="EGPD_Output",PinType.PinCategory="bool",PinType.PinSubCategory="",PinType.PinSubCategoryObject=None,PinType.PinSubCategoryMemberReference=(),PinType.PinValueType=(),PinType.bIsMap=False,PinType.bIsSet=False,PinType.bIsArray=False,PinType.bIsReference=False,PinType.bIsConst=False,PinType.bIsWeakPointer=False,PersistentGuid=00000000000000000000000000000000,bHidden=True,bNotConnectable=False,bDefaultValueIsReadOnly=False,bDefaultValueIsIgnored=False,bAdvancedView=False,)
+				End Object
+				Begin Object Class=K2Node_CallFunction Name="K2Node_CallFunction_389"
+				   FunctionReference=(MemberParent=BlueprintGeneratedClass'/Game/Prototype/Game/BP_OdinGameMode.BP_OdinGameMode_C',MemberName="End Loading Sequence",MemberGuid=A1D38DD0498B32ACC9A936A032D370AC)
+				   NodePosX=-144
+				   NodePosY=-4496
+				   NodeGuid=4498E305455CCB64B47710BE1ECE74B6
+				   CustomProperties Pin (PinId=B9D0319843BB38F57642CEB7B91707F9,PinName="execute",PinToolTip="\nExec",PinType.PinCategory="exec",PinType.PinSubCategory="",PinType.PinSubCategoryObject=None,PinType.PinSubCategoryMemberReference=(),PinType.PinValueType=(),PinType.bIsMap=False,PinType.bIsSet=False,PinType.bIsArray=False,PinType.bIsReference=False,PinType.bIsConst=False,PinType.bIsWeakPointer=False,LinkedTo=(K2Node_DynamicCast_5 C46EEAC7458FE7A6949E78BA06EE5B54,),PersistentGuid=00000000000000000000000000000000,bHidden=False,bNotConnectable=False,bDefaultValueIsReadOnly=False,bDefaultValueIsIgnored=False,bAdvancedView=False,)
+				   CustomProperties Pin (PinId=080DB6114DA99E281AD26280646A113F,PinName="then",PinToolTip="\nExec",Direction="EGPD_Output",PinType.PinCategory="exec",PinType.PinSubCategory="",PinType.PinSubCategoryObject=None,PinType.PinSubCategoryMemberReference=(),PinType.PinValueType=(),PinType.bIsMap=False,PinType.bIsSet=False,PinType.bIsArray=False,PinType.bIsReference=False,PinType.bIsConst=False,PinType.bIsWeakPointer=False,LinkedTo=(K2Node_CallFunction_147 0D4F8EEB4254987507E56590D0CCA616,),PersistentGuid=00000000000000000000000000000000,bHidden=False,bNotConnectable=False,bDefaultValueIsReadOnly=False,bDefaultValueIsIgnored=False,bAdvancedView=False,)
+				   CustomProperties Pin (PinId=3117651A4CB6ACA2EAEB6BA4272EC84B,PinName="self",PinFriendlyName=NSLOCTEXT("K2Node", "Target", "Target"),PinToolTip="Target\nBP Odin Game Mode Reference",PinType.PinCategory="object",PinType.PinSubCategory="",PinType.PinSubCategoryObject=BlueprintGeneratedClass'/Game/Prototype/Game/BP_OdinGameMode.BP_OdinGameMode_C',PinType.PinSubCategoryMemberReference=(),PinType.PinValueType=(),PinType.bIsMap=False,PinType.bIsSet=False,PinType.bIsArray=False,PinType.bIsReference=False,PinType.bIsConst=False,PinType.bIsWeakPointer=False,LinkedTo=(K2Node_DynamicCast_5 6B739B644EEBFA053240C1B460B2E558,),PersistentGuid=00000000000000000000000000000000,bHidden=False,bNotConnectable=False,bDefaultValueIsReadOnly=False,bDefaultValueIsIgnored=False,bAdvancedView=False,)
+				End Object
+				Begin Object Class=K2Node_CallFunction Name="K2Node_CallFunction_146"
+				   bIsPureFunc=True
+				   FunctionReference=(MemberParent=Class'/Script/Engine.GameplayStatics',MemberName="GetPlayerPawn")
+				   NodePosX=-32
+				   NodePosY=-4352
+				   NodeGuid=5AC649304096413924B030B0014FD935
+				   CustomProperties Pin (PinId=329EDA7E40EFEF35615E988D15DE170A,PinName="self",PinFriendlyName=NSLOCTEXT("K2Node", "Target", "Target"),PinToolTip="Target\nGameplay Statics Reference",PinType.PinCategory="object",PinType.PinSubCategory="",PinType.PinSubCategoryObject=Class'/Script/Engine.GameplayStatics',PinType.PinSubCategoryMemberReference=(),PinType.PinValueType=(),PinType.bIsMap=False,PinType.bIsSet=False,PinType.bIsArray=False,PinType.bIsReference=False,PinType.bIsConst=False,PinType.bIsWeakPointer=False,DefaultObject="/Script/Engine.Default__GameplayStatics",PersistentGuid=00000000000000000000000000000000,bHidden=True,bNotConnectable=False,bDefaultValueIsReadOnly=False,bDefaultValueIsIgnored=False,bAdvancedView=False,)
+				   CustomProperties Pin (PinId=89848130466D916E810B3F99D04CBF1F,PinName="WorldContextObject",PinToolTip="World Context Object\nObject Reference",PinType.PinCategory="object",PinType.PinSubCategory="",PinType.PinSubCategoryObject=Class'/Script/CoreUObject.Object',PinType.PinSubCategoryMemberReference=(),PinType.PinValueType=(),PinType.bIsMap=False,PinType.bIsSet=False,PinType.bIsArray=False,PinType.bIsReference=False,PinType.bIsConst=True,PinType.bIsWeakPointer=False,PersistentGuid=00000000000000000000000000000000,bHidden=True,bNotConnectable=False,bDefaultValueIsReadOnly=False,bDefaultValueIsIgnored=False,bAdvancedView=False,)
+				   CustomProperties Pin (PinId=02F269A847C5688444C9AEAEA78D5124,PinName="PlayerIndex",PinToolTip="Player Index\nInteger",PinType.PinCategory="int",PinType.PinSubCategory="",PinType.PinSubCategoryObject=None,PinType.PinSubCategoryMemberReference=(),PinType.PinValueType=(),PinType.bIsMap=False,PinType.bIsSet=False,PinType.bIsArray=False,PinType.bIsReference=False,PinType.bIsConst=False,PinType.bIsWeakPointer=False,DefaultValue="0",AutogeneratedDefaultValue="0",PersistentGuid=00000000000000000000000000000000,bHidden=False,bNotConnectable=False,bDefaultValueIsReadOnly=False,bDefaultValueIsIgnored=False,bAdvancedView=False,)
+				   CustomProperties Pin (PinId=A55AE1CA4B113E55AEA44FA5BD868FEB,PinName="ReturnValue",PinToolTip="Return Value\nPawn Reference",Direction="EGPD_Output",PinType.PinCategory="object",PinType.PinSubCategory="",PinType.PinSubCategoryObject=Class'/Script/Engine.Pawn',PinType.PinSubCategoryMemberReference=(),PinType.PinValueType=(),PinType.bIsMap=False,PinType.bIsSet=False,PinType.bIsArray=False,PinType.bIsReference=False,PinType.bIsConst=False,PinType.bIsWeakPointer=False,LinkedTo=(K2Node_CallFunction_147 05DD2A284290ED5C1DD79C95F86B089B,),PersistentGuid=00000000000000000000000000000000,bHidden=False,bNotConnectable=False,bDefaultValueIsReadOnly=False,bDefaultValueIsIgnored=False,bAdvancedView=False,)
+				End Object
+				Begin Object Class=K2Node_CallFunction Name="K2Node_CallFunction_147"
+				   FunctionReference=(MemberName="K2_TeleportTo",bSelfContext=True)
+				   NodePosX=400
+				   NodePosY=-4496
+				   NodeGuid=97C3515E434225B49FD8669033C40E7D
+				   CustomProperties Pin (PinId=0D4F8EEB4254987507E56590D0CCA616,PinName="execute",PinToolTip="\nExec",PinType.PinCategory="exec",PinType.PinSubCategory="",PinType.PinSubCategoryObject=None,PinType.PinSubCategoryMemberReference=(),PinType.PinValueType=(),PinType.bIsMap=False,PinType.bIsSet=False,PinType.bIsArray=False,PinType.bIsReference=False,PinType.bIsConst=False,PinType.bIsWeakPointer=False,LinkedTo=(K2Node_CallFunction_389 080DB6114DA99E281AD26280646A113F,),PersistentGuid=00000000000000000000000000000000,bHidden=False,bNotConnectable=False,bDefaultValueIsReadOnly=False,bDefaultValueIsIgnored=False,bAdvancedView=False,)
+				   CustomProperties Pin (PinId=DA74E61542E071192041BF8670B2B044,PinName="then",PinToolTip="\nExec",Direction="EGPD_Output",PinType.PinCategory="exec",PinType.PinSubCategory="",PinType.PinSubCategoryObject=None,PinType.PinSubCategoryMemberReference=(),PinType.PinValueType=(),PinType.bIsMap=False,PinType.bIsSet=False,PinType.bIsArray=False,PinType.bIsReference=False,PinType.bIsConst=False,PinType.bIsWeakPointer=False,PersistentGuid=00000000000000000000000000000000,bHidden=False,bNotConnectable=False,bDefaultValueIsReadOnly=False,bDefaultValueIsIgnored=False,bAdvancedView=False,)
+				   CustomProperties Pin (PinId=05DD2A284290ED5C1DD79C95F86B089B,PinName="self",PinFriendlyName=NSLOCTEXT("K2Node", "Target", "Target"),PinToolTip="Target\nActor Reference",PinType.PinCategory="object",PinType.PinSubCategory="",PinType.PinSubCategoryObject=Class'/Script/Engine.Actor',PinType.PinSubCategoryMemberReference=(),PinType.PinValueType=(),PinType.bIsMap=False,PinType.bIsSet=False,PinType.bIsArray=False,PinType.bIsReference=False,PinType.bIsConst=False,PinType.bIsWeakPointer=False,LinkedTo=(K2Node_CallFunction_146 A55AE1CA4B113E55AEA44FA5BD868FEB,),PersistentGuid=00000000000000000000000000000000,bHidden=False,bNotConnectable=False,bDefaultValueIsReadOnly=False,bDefaultValueIsIgnored=False,bAdvancedView=False,)
+				   CustomProperties Pin (PinId=7ADD16464BF350CD2D3D5DA206062506,PinName="DestLocation",PinToolTip="Dest Location\nVector \n\nThe target destination point",PinType.PinCategory="struct",PinType.PinSubCategory="",PinType.PinSubCategoryObject=ScriptStruct'/Script/CoreUObject.Vector',PinType.PinSubCategoryMemberReference=(),PinType.PinValueType=(),PinType.bIsMap=False,PinType.bIsSet=False,PinType.bIsArray=False,PinType.bIsReference=False,PinType.bIsConst=False,PinType.bIsWeakPointer=False,DefaultValue="0, 0, 0",AutogeneratedDefaultValue="0, 0, 0",LinkedTo=(K2Node_CallFunction_161 8EDCF5B4407E0D67A46F56B3A7EBDC6E,),PersistentGuid=00000000000000000000000000000000,bHidden=False,bNotConnectable=False,bDefaultValueIsReadOnly=False,bDefaultValueIsIgnored=False,bAdvancedView=False,)
+				   CustomProperties Pin (PinId=CADAD3724E0A28E00552BCA72427B545,PinName="DestRotation",PinToolTip="Dest Rotation\nRotator \n\nThe target rotation at the destination",PinType.PinCategory="struct",PinType.PinSubCategory="",PinType.PinSubCategoryObject=ScriptStruct'/Script/CoreUObject.Rotator',PinType.PinSubCategoryMemberReference=(),PinType.PinValueType=(),PinType.bIsMap=False,PinType.bIsSet=False,PinType.bIsArray=False,PinType.bIsReference=False,PinType.bIsConst=False,PinType.bIsWeakPointer=False,DefaultValue="0, 0, 0",AutogeneratedDefaultValue="0, 0, 0",LinkedTo=(K2Node_CallFunction_197 177F737D497F0B71A0D9548711E45134,),PersistentGuid=00000000000000000000000000000000,bHidden=False,bNotConnectable=False,bDefaultValueIsReadOnly=False,bDefaultValueIsIgnored=False,bAdvancedView=False,)
+				   CustomProperties Pin (PinId=150C2F4B41C32C40A5EE5F928C2473FE,PinName="ReturnValue",PinToolTip="Return Value\nBoolean\n\ntrue if the actor has been successfully moved, or false if it couldn\'t fit.",Direction="EGPD_Output",PinType.PinCategory="bool",PinType.PinSubCategory="",PinType.PinSubCategoryObject=None,PinType.PinSubCategoryMemberReference=(),PinType.PinValueType=(),PinType.bIsMap=False,PinType.bIsSet=False,PinType.bIsArray=False,PinType.bIsReference=False,PinType.bIsConst=False,PinType.bIsWeakPointer=False,DefaultValue="false",AutogeneratedDefaultValue="false",PersistentGuid=00000000000000000000000000000000,bHidden=False,bNotConnectable=False,bDefaultValueIsReadOnly=False,bDefaultValueIsIgnored=False,bAdvancedView=False,)
+				End Object
+				Begin Object Class=K2Node_Literal Name="K2Node_Literal_0"
+				   ObjectRef=PlayerStart'/BasicCity/Maps/BasicCityCityCenterDay.BasicCityCityCenterDay:PersistentLevel.PlayerStart2_3'
+				   NodePosX=-272
+				   NodePosY=-4192
+				   NodeGuid=58EDCAD94D643329430D2D879E37B3E5
+				   CustomProperties Pin (PinId=AFA007CD4AFA013315CE299F79589275,PinName="PlayerStart2",PinFriendlyName=NSLOCTEXT("", "C0C388B24DBCBE06EB1E1B8C7FFFE555", "PlayerStart2"),Direction="EGPD_Output",PinType.PinCategory="object",PinType.PinSubCategory="",PinType.PinSubCategoryObject=Class'/Script/Engine.PlayerStart',PinType.PinSubCategoryMemberReference=(),PinType.PinValueType=(),PinType.bIsMap=False,PinType.bIsSet=False,PinType.bIsArray=False,PinType.bIsReference=False,PinType.bIsConst=False,PinType.bIsWeakPointer=False,LinkedTo=(K2Node_CallFunction_161 76593FC04576E3C4E850799E8DD55EC5,K2Node_CallFunction_197 40BECA4844231CB62E8E729102B58F73,),PersistentGuid=00000000000000000000000000000000,bHidden=False,bNotConnectable=False,bDefaultValueIsReadOnly=False,bDefaultValueIsIgnored=False,bAdvancedView=False,)
+				End Object
+				Begin Object Class=K2Node_CallFunction Name="K2Node_CallFunction_161"
+				   bIsPureFunc=True
+				   bIsConstFunc=True
+				   FunctionReference=(MemberName="K2_GetActorLocation",bSelfContext=True)
+				   NodePosX=16
+				   NodePosY=-4272
+				   NodeGuid=72AB37F441B22C2E9CB0D3BFA6D9A48F
+				   CustomProperties Pin (PinId=76593FC04576E3C4E850799E8DD55EC5,PinName="self",PinFriendlyName=NSLOCTEXT("K2Node", "Target", "Target"),PinToolTip="Target\nActor Reference",PinType.PinCategory="object",PinType.PinSubCategory="",PinType.PinSubCategoryObject=Class'/Script/Engine.Actor',PinType.PinSubCategoryMemberReference=(),PinType.PinValueType=(),PinType.bIsMap=False,PinType.bIsSet=False,PinType.bIsArray=False,PinType.bIsReference=False,PinType.bIsConst=False,PinType.bIsWeakPointer=False,LinkedTo=(K2Node_Literal_0 AFA007CD4AFA013315CE299F79589275,),PersistentGuid=00000000000000000000000000000000,bHidden=False,bNotConnectable=False,bDefaultValueIsReadOnly=False,bDefaultValueIsIgnored=False,bAdvancedView=False,)
+				   CustomProperties Pin (PinId=8EDCF5B4407E0D67A46F56B3A7EBDC6E,PinName="ReturnValue",PinToolTip="Return Value\nVector ",Direction="EGPD_Output",PinType.PinCategory="struct",PinType.PinSubCategory="",PinType.PinSubCategoryObject=ScriptStruct'/Script/CoreUObject.Vector',PinType.PinSubCategoryMemberReference=(),PinType.PinValueType=(),PinType.bIsMap=False,PinType.bIsSet=False,PinType.bIsArray=False,PinType.bIsReference=False,PinType.bIsConst=False,PinType.bIsWeakPointer=False,AutogeneratedDefaultValue="0, 0, 0",LinkedTo=(K2Node_CallFunction_147 7ADD16464BF350CD2D3D5DA206062506,),PersistentGuid=00000000000000000000000000000000,bHidden=False,bNotConnectable=False,bDefaultValueIsReadOnly=False,bDefaultValueIsIgnored=False,bAdvancedView=False,)
+				End Object
+				Begin Object Class=K2Node_CallFunction Name="K2Node_CallFunction_197"
+				   bIsPureFunc=True
+				   bIsConstFunc=True
+				   FunctionReference=(MemberName="K2_GetActorRotation",bSelfContext=True)
+				   NodePosX=16
+				   NodePosY=-4176
+				   NodeGuid=AF5DAC314C10340F9EA01999BB42D085
+				   CustomProperties Pin (PinId=40BECA4844231CB62E8E729102B58F73,PinName="self",PinFriendlyName=NSLOCTEXT("K2Node", "Target", "Target"),PinToolTip="Target\nActor Reference",PinType.PinCategory="object",PinType.PinSubCategory="",PinType.PinSubCategoryObject=Class'/Script/Engine.Actor',PinType.PinSubCategoryMemberReference=(),PinType.PinValueType=(),PinType.bIsMap=False,PinType.bIsSet=False,PinType.bIsArray=False,PinType.bIsReference=False,PinType.bIsConst=False,PinType.bIsWeakPointer=False,LinkedTo=(K2Node_Literal_0 AFA007CD4AFA013315CE299F79589275,),PersistentGuid=00000000000000000000000000000000,bHidden=False,bNotConnectable=False,bDefaultValueIsReadOnly=False,bDefaultValueIsIgnored=False,bAdvancedView=False,)
+				   CustomProperties Pin (PinId=177F737D497F0B71A0D9548711E45134,PinName="ReturnValue",PinToolTip="Return Value\nRotator ",Direction="EGPD_Output",PinType.PinCategory="struct",PinType.PinSubCategory="",PinType.PinSubCategoryObject=ScriptStruct'/Script/CoreUObject.Rotator',PinType.PinSubCategoryMemberReference=(),PinType.PinValueType=(),PinType.bIsMap=False,PinType.bIsSet=False,PinType.bIsArray=False,PinType.bIsReference=False,PinType.bIsConst=False,PinType.bIsWeakPointer=False,AutogeneratedDefaultValue="0, 0, 0",LinkedTo=(K2Node_CallFunction_147 CADAD3724E0A28E00552BCA72427B545,),PersistentGuid=00000000000000000000000000000000,bHidden=False,bNotConnectable=False,bDefaultValueIsReadOnly=False,bDefaultValueIsIgnored=False,bAdvancedView=False,)
+				End Object
 		[/PARAM]
 	[/OBJECT]
 
@@ -69,100 +158,3 @@ When the wave is complete, the loading sequence begins and the player is taken t
 ![](fade.png)(convert:false)
 
 [NAV]
-
-[COMMENT:none]
-
-[EXCERPT:snippet]
-	Begin Object Class=K2Node_Event Name="K2Node_Event_52"
-	   EventReference=(MemberParent=Class'/Script/Engine.Actor',MemberName="ReceiveBeginPlay")
-	   bOverrideFunction=True
-	   NodePosX=-928
-	   NodePosY=-4480
-	   NodeGuid=6C849E7241A6108A36D5C1BEAAEAE94F
-	   CustomProperties Pin (PinId=B35535C443B530261C41B8A84FF2CC42,PinName="OutputDelegate",Direction="EGPD_Output",PinType.PinCategory="delegate",PinType.PinSubCategory="",PinType.PinSubCategoryObject=None,PinType.PinSubCategoryMemberReference=(MemberParent=Class'/Script/Engine.Actor',MemberName="ReceiveBeginPlay"),PinType.PinValueType=(),PinType.bIsMap=False,PinType.bIsSet=False,PinType.bIsArray=False,PinType.bIsReference=False,PinType.bIsConst=False,PinType.bIsWeakPointer=False,PersistentGuid=00000000000000000000000000000000,bHidden=False,bNotConnectable=False,bDefaultValueIsReadOnly=False,bDefaultValueIsIgnored=False,bAdvancedView=False,)
-	   CustomProperties Pin (PinId=5A6DBA7144DD625770231A9C725C0B49,PinName="then",Direction="EGPD_Output",PinType.PinCategory="exec",PinType.PinSubCategory="",PinType.PinSubCategoryObject=None,PinType.PinSubCategoryMemberReference=(),PinType.PinValueType=(),PinType.bIsMap=False,PinType.bIsSet=False,PinType.bIsArray=False,PinType.bIsReference=False,PinType.bIsConst=False,PinType.bIsWeakPointer=False,LinkedTo=(K2Node_DynamicCast_5 3228E8FA4B588A45DCEA3A839D739E9A,),PersistentGuid=00000000000000000000000000000000,bHidden=False,bNotConnectable=False,bDefaultValueIsReadOnly=False,bDefaultValueIsIgnored=False,bAdvancedView=False,)
-	End Object
-	Begin Object Class=K2Node_CallFunction Name="K2Node_CallFunction_280"
-	   bIsPureFunc=True
-	   FunctionReference=(MemberParent=Class'/Script/Engine.GameplayStatics',MemberName="GetGameMode")
-	   NodePosX=-720
-	   NodePosY=-4416
-	   NodeGuid=2CB9ABB94CF9F2C3BAEBB0B38ED0F003
-	   CustomProperties Pin (PinId=1AAC46094DE5E3D0CB1F8488BEDBFE3B,PinName="self",PinFriendlyName=NSLOCTEXT("K2Node", "Target", "Target"),PinToolTip="Target\nGameplay Statics Reference",PinType.PinCategory="object",PinType.PinSubCategory="",PinType.PinSubCategoryObject=Class'/Script/Engine.GameplayStatics',PinType.PinSubCategoryMemberReference=(),PinType.PinValueType=(),PinType.bIsMap=False,PinType.bIsSet=False,PinType.bIsArray=False,PinType.bIsReference=False,PinType.bIsConst=False,PinType.bIsWeakPointer=False,DefaultObject="/Script/Engine.Default__GameplayStatics",PersistentGuid=00000000000000000000000000000000,bHidden=True,bNotConnectable=False,bDefaultValueIsReadOnly=False,bDefaultValueIsIgnored=False,bAdvancedView=False,)
-	   CustomProperties Pin (PinId=539E882646828F110DA8539329878814,PinName="WorldContextObject",PinToolTip="World Context Object\nObject Reference",PinType.PinCategory="object",PinType.PinSubCategory="",PinType.PinSubCategoryObject=Class'/Script/CoreUObject.Object',PinType.PinSubCategoryMemberReference=(),PinType.PinValueType=(),PinType.bIsMap=False,PinType.bIsSet=False,PinType.bIsArray=False,PinType.bIsReference=False,PinType.bIsConst=True,PinType.bIsWeakPointer=False,PersistentGuid=00000000000000000000000000000000,bHidden=True,bNotConnectable=False,bDefaultValueIsReadOnly=False,bDefaultValueIsIgnored=False,bAdvancedView=False,)
-	   CustomProperties Pin (PinId=F81E36D8454550AAAEF10DB57E9FC71F,PinName="ReturnValue",PinToolTip="Return Value\nGame Mode Base Reference",Direction="EGPD_Output",PinType.PinCategory="object",PinType.PinSubCategory="",PinType.PinSubCategoryObject=Class'/Script/Engine.GameModeBase',PinType.PinSubCategoryMemberReference=(),PinType.PinValueType=(),PinType.bIsMap=False,PinType.bIsSet=False,PinType.bIsArray=False,PinType.bIsReference=False,PinType.bIsConst=False,PinType.bIsWeakPointer=False,LinkedTo=(K2Node_DynamicCast_5 8B0B2634480FE02DE82004ADB105AB79,),PersistentGuid=00000000000000000000000000000000,bHidden=False,bNotConnectable=False,bDefaultValueIsReadOnly=False,bDefaultValueIsIgnored=False,bAdvancedView=False,)
-	End Object
-	Begin Object Class=K2Node_DynamicCast Name="K2Node_DynamicCast_5"
-	   TargetType=BlueprintGeneratedClass'/Game/Prototype/Game/BP_OdinGameMode.BP_OdinGameMode_C'
-	   NodePosX=-512
-	   NodePosY=-4480
-	   NodeGuid=2A2928BD4CBEFB243D86D796866C7C1B
-	   CustomProperties Pin (PinId=3228E8FA4B588A45DCEA3A839D739E9A,PinName="execute",PinType.PinCategory="exec",PinType.PinSubCategory="",PinType.PinSubCategoryObject=None,PinType.PinSubCategoryMemberReference=(),PinType.PinValueType=(),PinType.bIsMap=False,PinType.bIsSet=False,PinType.bIsArray=False,PinType.bIsReference=False,PinType.bIsConst=False,PinType.bIsWeakPointer=False,LinkedTo=(K2Node_Event_52 5A6DBA7144DD625770231A9C725C0B49,),PersistentGuid=00000000000000000000000000000000,bHidden=False,bNotConnectable=False,bDefaultValueIsReadOnly=False,bDefaultValueIsIgnored=False,bAdvancedView=False,)
-	   CustomProperties Pin (PinId=C46EEAC7458FE7A6949E78BA06EE5B54,PinName="then",Direction="EGPD_Output",PinType.PinCategory="exec",PinType.PinSubCategory="",PinType.PinSubCategoryObject=None,PinType.PinSubCategoryMemberReference=(),PinType.PinValueType=(),PinType.bIsMap=False,PinType.bIsSet=False,PinType.bIsArray=False,PinType.bIsReference=False,PinType.bIsConst=False,PinType.bIsWeakPointer=False,LinkedTo=(K2Node_CallFunction_389 B9D0319843BB38F57642CEB7B91707F9,),PersistentGuid=00000000000000000000000000000000,bHidden=False,bNotConnectable=False,bDefaultValueIsReadOnly=False,bDefaultValueIsIgnored=False,bAdvancedView=False,)
-	   CustomProperties Pin (PinId=9EAA2A464990955344C076A409669F6D,PinName="CastFailed",Direction="EGPD_Output",PinType.PinCategory="exec",PinType.PinSubCategory="",PinType.PinSubCategoryObject=None,PinType.PinSubCategoryMemberReference=(),PinType.PinValueType=(),PinType.bIsMap=False,PinType.bIsSet=False,PinType.bIsArray=False,PinType.bIsReference=False,PinType.bIsConst=False,PinType.bIsWeakPointer=False,PersistentGuid=00000000000000000000000000000000,bHidden=False,bNotConnectable=False,bDefaultValueIsReadOnly=False,bDefaultValueIsIgnored=False,bAdvancedView=False,)
-	   CustomProperties Pin (PinId=8B0B2634480FE02DE82004ADB105AB79,PinName="Object",PinType.PinCategory="object",PinType.PinSubCategory="",PinType.PinSubCategoryObject=Class'/Script/CoreUObject.Object',PinType.PinSubCategoryMemberReference=(),PinType.PinValueType=(),PinType.bIsMap=False,PinType.bIsSet=False,PinType.bIsArray=False,PinType.bIsReference=False,PinType.bIsConst=False,PinType.bIsWeakPointer=False,LinkedTo=(K2Node_CallFunction_280 F81E36D8454550AAAEF10DB57E9FC71F,),PersistentGuid=00000000000000000000000000000000,bHidden=False,bNotConnectable=False,bDefaultValueIsReadOnly=False,bDefaultValueIsIgnored=False,bAdvancedView=False,)
-	   CustomProperties Pin (PinId=6B739B644EEBFA053240C1B460B2E558,PinName="AsBP Odin Game Mode",Direction="EGPD_Output",PinType.PinCategory="object",PinType.PinSubCategory="",PinType.PinSubCategoryObject=BlueprintGeneratedClass'/Game/Prototype/Game/BP_OdinGameMode.BP_OdinGameMode_C',PinType.PinSubCategoryMemberReference=(),PinType.PinValueType=(),PinType.bIsMap=False,PinType.bIsSet=False,PinType.bIsArray=False,PinType.bIsReference=False,PinType.bIsConst=False,PinType.bIsWeakPointer=False,LinkedTo=(K2Node_CallFunction_389 3117651A4CB6ACA2EAEB6BA4272EC84B,),PersistentGuid=00000000000000000000000000000000,bHidden=False,bNotConnectable=False,bDefaultValueIsReadOnly=False,bDefaultValueIsIgnored=False,bAdvancedView=False,)
-	   CustomProperties Pin (PinId=BEB2F603410D2A627B7D089246E322C7,PinName="bSuccess",Direction="EGPD_Output",PinType.PinCategory="bool",PinType.PinSubCategory="",PinType.PinSubCategoryObject=None,PinType.PinSubCategoryMemberReference=(),PinType.PinValueType=(),PinType.bIsMap=False,PinType.bIsSet=False,PinType.bIsArray=False,PinType.bIsReference=False,PinType.bIsConst=False,PinType.bIsWeakPointer=False,PersistentGuid=00000000000000000000000000000000,bHidden=True,bNotConnectable=False,bDefaultValueIsReadOnly=False,bDefaultValueIsIgnored=False,bAdvancedView=False,)
-	End Object
-	Begin Object Class=K2Node_CallFunction Name="K2Node_CallFunction_389"
-	   FunctionReference=(MemberParent=BlueprintGeneratedClass'/Game/Prototype/Game/BP_OdinGameMode.BP_OdinGameMode_C',MemberName="End Loading Sequence",MemberGuid=A1D38DD0498B32ACC9A936A032D370AC)
-	   NodePosX=-144
-	   NodePosY=-4496
-	   NodeGuid=4498E305455CCB64B47710BE1ECE74B6
-	   CustomProperties Pin (PinId=B9D0319843BB38F57642CEB7B91707F9,PinName="execute",PinToolTip="\nExec",PinType.PinCategory="exec",PinType.PinSubCategory="",PinType.PinSubCategoryObject=None,PinType.PinSubCategoryMemberReference=(),PinType.PinValueType=(),PinType.bIsMap=False,PinType.bIsSet=False,PinType.bIsArray=False,PinType.bIsReference=False,PinType.bIsConst=False,PinType.bIsWeakPointer=False,LinkedTo=(K2Node_DynamicCast_5 C46EEAC7458FE7A6949E78BA06EE5B54,),PersistentGuid=00000000000000000000000000000000,bHidden=False,bNotConnectable=False,bDefaultValueIsReadOnly=False,bDefaultValueIsIgnored=False,bAdvancedView=False,)
-	   CustomProperties Pin (PinId=080DB6114DA99E281AD26280646A113F,PinName="then",PinToolTip="\nExec",Direction="EGPD_Output",PinType.PinCategory="exec",PinType.PinSubCategory="",PinType.PinSubCategoryObject=None,PinType.PinSubCategoryMemberReference=(),PinType.PinValueType=(),PinType.bIsMap=False,PinType.bIsSet=False,PinType.bIsArray=False,PinType.bIsReference=False,PinType.bIsConst=False,PinType.bIsWeakPointer=False,LinkedTo=(K2Node_CallFunction_147 0D4F8EEB4254987507E56590D0CCA616,),PersistentGuid=00000000000000000000000000000000,bHidden=False,bNotConnectable=False,bDefaultValueIsReadOnly=False,bDefaultValueIsIgnored=False,bAdvancedView=False,)
-	   CustomProperties Pin (PinId=3117651A4CB6ACA2EAEB6BA4272EC84B,PinName="self",PinFriendlyName=NSLOCTEXT("K2Node", "Target", "Target"),PinToolTip="Target\nBP Odin Game Mode Reference",PinType.PinCategory="object",PinType.PinSubCategory="",PinType.PinSubCategoryObject=BlueprintGeneratedClass'/Game/Prototype/Game/BP_OdinGameMode.BP_OdinGameMode_C',PinType.PinSubCategoryMemberReference=(),PinType.PinValueType=(),PinType.bIsMap=False,PinType.bIsSet=False,PinType.bIsArray=False,PinType.bIsReference=False,PinType.bIsConst=False,PinType.bIsWeakPointer=False,LinkedTo=(K2Node_DynamicCast_5 6B739B644EEBFA053240C1B460B2E558,),PersistentGuid=00000000000000000000000000000000,bHidden=False,bNotConnectable=False,bDefaultValueIsReadOnly=False,bDefaultValueIsIgnored=False,bAdvancedView=False,)
-	End Object
-	Begin Object Class=K2Node_CallFunction Name="K2Node_CallFunction_146"
-	   bIsPureFunc=True
-	   FunctionReference=(MemberParent=Class'/Script/Engine.GameplayStatics',MemberName="GetPlayerPawn")
-	   NodePosX=-32
-	   NodePosY=-4352
-	   NodeGuid=5AC649304096413924B030B0014FD935
-	   CustomProperties Pin (PinId=329EDA7E40EFEF35615E988D15DE170A,PinName="self",PinFriendlyName=NSLOCTEXT("K2Node", "Target", "Target"),PinToolTip="Target\nGameplay Statics Reference",PinType.PinCategory="object",PinType.PinSubCategory="",PinType.PinSubCategoryObject=Class'/Script/Engine.GameplayStatics',PinType.PinSubCategoryMemberReference=(),PinType.PinValueType=(),PinType.bIsMap=False,PinType.bIsSet=False,PinType.bIsArray=False,PinType.bIsReference=False,PinType.bIsConst=False,PinType.bIsWeakPointer=False,DefaultObject="/Script/Engine.Default__GameplayStatics",PersistentGuid=00000000000000000000000000000000,bHidden=True,bNotConnectable=False,bDefaultValueIsReadOnly=False,bDefaultValueIsIgnored=False,bAdvancedView=False,)
-	   CustomProperties Pin (PinId=89848130466D916E810B3F99D04CBF1F,PinName="WorldContextObject",PinToolTip="World Context Object\nObject Reference",PinType.PinCategory="object",PinType.PinSubCategory="",PinType.PinSubCategoryObject=Class'/Script/CoreUObject.Object',PinType.PinSubCategoryMemberReference=(),PinType.PinValueType=(),PinType.bIsMap=False,PinType.bIsSet=False,PinType.bIsArray=False,PinType.bIsReference=False,PinType.bIsConst=True,PinType.bIsWeakPointer=False,PersistentGuid=00000000000000000000000000000000,bHidden=True,bNotConnectable=False,bDefaultValueIsReadOnly=False,bDefaultValueIsIgnored=False,bAdvancedView=False,)
-	   CustomProperties Pin (PinId=02F269A847C5688444C9AEAEA78D5124,PinName="PlayerIndex",PinToolTip="Player Index\nInteger",PinType.PinCategory="int",PinType.PinSubCategory="",PinType.PinSubCategoryObject=None,PinType.PinSubCategoryMemberReference=(),PinType.PinValueType=(),PinType.bIsMap=False,PinType.bIsSet=False,PinType.bIsArray=False,PinType.bIsReference=False,PinType.bIsConst=False,PinType.bIsWeakPointer=False,DefaultValue="0",AutogeneratedDefaultValue="0",PersistentGuid=00000000000000000000000000000000,bHidden=False,bNotConnectable=False,bDefaultValueIsReadOnly=False,bDefaultValueIsIgnored=False,bAdvancedView=False,)
-	   CustomProperties Pin (PinId=A55AE1CA4B113E55AEA44FA5BD868FEB,PinName="ReturnValue",PinToolTip="Return Value\nPawn Reference",Direction="EGPD_Output",PinType.PinCategory="object",PinType.PinSubCategory="",PinType.PinSubCategoryObject=Class'/Script/Engine.Pawn',PinType.PinSubCategoryMemberReference=(),PinType.PinValueType=(),PinType.bIsMap=False,PinType.bIsSet=False,PinType.bIsArray=False,PinType.bIsReference=False,PinType.bIsConst=False,PinType.bIsWeakPointer=False,LinkedTo=(K2Node_CallFunction_147 05DD2A284290ED5C1DD79C95F86B089B,),PersistentGuid=00000000000000000000000000000000,bHidden=False,bNotConnectable=False,bDefaultValueIsReadOnly=False,bDefaultValueIsIgnored=False,bAdvancedView=False,)
-	End Object
-	Begin Object Class=K2Node_CallFunction Name="K2Node_CallFunction_147"
-	   FunctionReference=(MemberName="K2_TeleportTo",bSelfContext=True)
-	   NodePosX=400
-	   NodePosY=-4496
-	   NodeGuid=97C3515E434225B49FD8669033C40E7D
-	   CustomProperties Pin (PinId=0D4F8EEB4254987507E56590D0CCA616,PinName="execute",PinToolTip="\nExec",PinType.PinCategory="exec",PinType.PinSubCategory="",PinType.PinSubCategoryObject=None,PinType.PinSubCategoryMemberReference=(),PinType.PinValueType=(),PinType.bIsMap=False,PinType.bIsSet=False,PinType.bIsArray=False,PinType.bIsReference=False,PinType.bIsConst=False,PinType.bIsWeakPointer=False,LinkedTo=(K2Node_CallFunction_389 080DB6114DA99E281AD26280646A113F,),PersistentGuid=00000000000000000000000000000000,bHidden=False,bNotConnectable=False,bDefaultValueIsReadOnly=False,bDefaultValueIsIgnored=False,bAdvancedView=False,)
-	   CustomProperties Pin (PinId=DA74E61542E071192041BF8670B2B044,PinName="then",PinToolTip="\nExec",Direction="EGPD_Output",PinType.PinCategory="exec",PinType.PinSubCategory="",PinType.PinSubCategoryObject=None,PinType.PinSubCategoryMemberReference=(),PinType.PinValueType=(),PinType.bIsMap=False,PinType.bIsSet=False,PinType.bIsArray=False,PinType.bIsReference=False,PinType.bIsConst=False,PinType.bIsWeakPointer=False,PersistentGuid=00000000000000000000000000000000,bHidden=False,bNotConnectable=False,bDefaultValueIsReadOnly=False,bDefaultValueIsIgnored=False,bAdvancedView=False,)
-	   CustomProperties Pin (PinId=05DD2A284290ED5C1DD79C95F86B089B,PinName="self",PinFriendlyName=NSLOCTEXT("K2Node", "Target", "Target"),PinToolTip="Target\nActor Reference",PinType.PinCategory="object",PinType.PinSubCategory="",PinType.PinSubCategoryObject=Class'/Script/Engine.Actor',PinType.PinSubCategoryMemberReference=(),PinType.PinValueType=(),PinType.bIsMap=False,PinType.bIsSet=False,PinType.bIsArray=False,PinType.bIsReference=False,PinType.bIsConst=False,PinType.bIsWeakPointer=False,LinkedTo=(K2Node_CallFunction_146 A55AE1CA4B113E55AEA44FA5BD868FEB,),PersistentGuid=00000000000000000000000000000000,bHidden=False,bNotConnectable=False,bDefaultValueIsReadOnly=False,bDefaultValueIsIgnored=False,bAdvancedView=False,)
-	   CustomProperties Pin (PinId=7ADD16464BF350CD2D3D5DA206062506,PinName="DestLocation",PinToolTip="Dest Location\nVector \n\nThe target destination point",PinType.PinCategory="struct",PinType.PinSubCategory="",PinType.PinSubCategoryObject=ScriptStruct'/Script/CoreUObject.Vector',PinType.PinSubCategoryMemberReference=(),PinType.PinValueType=(),PinType.bIsMap=False,PinType.bIsSet=False,PinType.bIsArray=False,PinType.bIsReference=False,PinType.bIsConst=False,PinType.bIsWeakPointer=False,DefaultValue="0, 0, 0",AutogeneratedDefaultValue="0, 0, 0",LinkedTo=(K2Node_CallFunction_161 8EDCF5B4407E0D67A46F56B3A7EBDC6E,),PersistentGuid=00000000000000000000000000000000,bHidden=False,bNotConnectable=False,bDefaultValueIsReadOnly=False,bDefaultValueIsIgnored=False,bAdvancedView=False,)
-	   CustomProperties Pin (PinId=CADAD3724E0A28E00552BCA72427B545,PinName="DestRotation",PinToolTip="Dest Rotation\nRotator \n\nThe target rotation at the destination",PinType.PinCategory="struct",PinType.PinSubCategory="",PinType.PinSubCategoryObject=ScriptStruct'/Script/CoreUObject.Rotator',PinType.PinSubCategoryMemberReference=(),PinType.PinValueType=(),PinType.bIsMap=False,PinType.bIsSet=False,PinType.bIsArray=False,PinType.bIsReference=False,PinType.bIsConst=False,PinType.bIsWeakPointer=False,DefaultValue="0, 0, 0",AutogeneratedDefaultValue="0, 0, 0",LinkedTo=(K2Node_CallFunction_197 177F737D497F0B71A0D9548711E45134,),PersistentGuid=00000000000000000000000000000000,bHidden=False,bNotConnectable=False,bDefaultValueIsReadOnly=False,bDefaultValueIsIgnored=False,bAdvancedView=False,)
-	   CustomProperties Pin (PinId=150C2F4B41C32C40A5EE5F928C2473FE,PinName="ReturnValue",PinToolTip="Return Value\nBoolean\n\ntrue if the actor has been successfully moved, or false if it couldn\'t fit.",Direction="EGPD_Output",PinType.PinCategory="bool",PinType.PinSubCategory="",PinType.PinSubCategoryObject=None,PinType.PinSubCategoryMemberReference=(),PinType.PinValueType=(),PinType.bIsMap=False,PinType.bIsSet=False,PinType.bIsArray=False,PinType.bIsReference=False,PinType.bIsConst=False,PinType.bIsWeakPointer=False,DefaultValue="false",AutogeneratedDefaultValue="false",PersistentGuid=00000000000000000000000000000000,bHidden=False,bNotConnectable=False,bDefaultValueIsReadOnly=False,bDefaultValueIsIgnored=False,bAdvancedView=False,)
-	End Object
-	Begin Object Class=K2Node_Literal Name="K2Node_Literal_0"
-	   ObjectRef=PlayerStart'/BasicCity/Maps/BasicCityCityCenterDay.BasicCityCityCenterDay:PersistentLevel.PlayerStart2_3'
-	   NodePosX=-272
-	   NodePosY=-4192
-	   NodeGuid=58EDCAD94D643329430D2D879E37B3E5
-	   CustomProperties Pin (PinId=AFA007CD4AFA013315CE299F79589275,PinName="PlayerStart2",PinFriendlyName=NSLOCTEXT("", "C0C388B24DBCBE06EB1E1B8C7FFFE555", "PlayerStart2"),Direction="EGPD_Output",PinType.PinCategory="object",PinType.PinSubCategory="",PinType.PinSubCategoryObject=Class'/Script/Engine.PlayerStart',PinType.PinSubCategoryMemberReference=(),PinType.PinValueType=(),PinType.bIsMap=False,PinType.bIsSet=False,PinType.bIsArray=False,PinType.bIsReference=False,PinType.bIsConst=False,PinType.bIsWeakPointer=False,LinkedTo=(K2Node_CallFunction_161 76593FC04576E3C4E850799E8DD55EC5,K2Node_CallFunction_197 40BECA4844231CB62E8E729102B58F73,),PersistentGuid=00000000000000000000000000000000,bHidden=False,bNotConnectable=False,bDefaultValueIsReadOnly=False,bDefaultValueIsIgnored=False,bAdvancedView=False,)
-	End Object
-	Begin Object Class=K2Node_CallFunction Name="K2Node_CallFunction_161"
-	   bIsPureFunc=True
-	   bIsConstFunc=True
-	   FunctionReference=(MemberName="K2_GetActorLocation",bSelfContext=True)
-	   NodePosX=16
-	   NodePosY=-4272
-	   NodeGuid=72AB37F441B22C2E9CB0D3BFA6D9A48F
-	   CustomProperties Pin (PinId=76593FC04576E3C4E850799E8DD55EC5,PinName="self",PinFriendlyName=NSLOCTEXT("K2Node", "Target", "Target"),PinToolTip="Target\nActor Reference",PinType.PinCategory="object",PinType.PinSubCategory="",PinType.PinSubCategoryObject=Class'/Script/Engine.Actor',PinType.PinSubCategoryMemberReference=(),PinType.PinValueType=(),PinType.bIsMap=False,PinType.bIsSet=False,PinType.bIsArray=False,PinType.bIsReference=False,PinType.bIsConst=False,PinType.bIsWeakPointer=False,LinkedTo=(K2Node_Literal_0 AFA007CD4AFA013315CE299F79589275,),PersistentGuid=00000000000000000000000000000000,bHidden=False,bNotConnectable=False,bDefaultValueIsReadOnly=False,bDefaultValueIsIgnored=False,bAdvancedView=False,)
-	   CustomProperties Pin (PinId=8EDCF5B4407E0D67A46F56B3A7EBDC6E,PinName="ReturnValue",PinToolTip="Return Value\nVector ",Direction="EGPD_Output",PinType.PinCategory="struct",PinType.PinSubCategory="",PinType.PinSubCategoryObject=ScriptStruct'/Script/CoreUObject.Vector',PinType.PinSubCategoryMemberReference=(),PinType.PinValueType=(),PinType.bIsMap=False,PinType.bIsSet=False,PinType.bIsArray=False,PinType.bIsReference=False,PinType.bIsConst=False,PinType.bIsWeakPointer=False,AutogeneratedDefaultValue="0, 0, 0",LinkedTo=(K2Node_CallFunction_147 7ADD16464BF350CD2D3D5DA206062506,),PersistentGuid=00000000000000000000000000000000,bHidden=False,bNotConnectable=False,bDefaultValueIsReadOnly=False,bDefaultValueIsIgnored=False,bAdvancedView=False,)
-	End Object
-	Begin Object Class=K2Node_CallFunction Name="K2Node_CallFunction_197"
-	   bIsPureFunc=True
-	   bIsConstFunc=True
-	   FunctionReference=(MemberName="K2_GetActorRotation",bSelfContext=True)
-	   NodePosX=16
-	   NodePosY=-4176
-	   NodeGuid=AF5DAC314C10340F9EA01999BB42D085
-	   CustomProperties Pin (PinId=40BECA4844231CB62E8E729102B58F73,PinName="self",PinFriendlyName=NSLOCTEXT("K2Node", "Target", "Target"),PinToolTip="Target\nActor Reference",PinType.PinCategory="object",PinType.PinSubCategory="",PinType.PinSubCategoryObject=Class'/Script/Engine.Actor',PinType.PinSubCategoryMemberReference=(),PinType.PinValueType=(),PinType.bIsMap=False,PinType.bIsSet=False,PinType.bIsArray=False,PinType.bIsReference=False,PinType.bIsConst=False,PinType.bIsWeakPointer=False,LinkedTo=(K2Node_Literal_0 AFA007CD4AFA013315CE299F79589275,),PersistentGuid=00000000000000000000000000000000,bHidden=False,bNotConnectable=False,bDefaultValueIsReadOnly=False,bDefaultValueIsIgnored=False,bAdvancedView=False,)
-	   CustomProperties Pin (PinId=177F737D497F0B71A0D9548711E45134,PinName="ReturnValue",PinToolTip="Return Value\nRotator ",Direction="EGPD_Output",PinType.PinCategory="struct",PinType.PinSubCategory="",PinType.PinSubCategoryObject=ScriptStruct'/Script/CoreUObject.Rotator',PinType.PinSubCategoryMemberReference=(),PinType.PinValueType=(),PinType.bIsMap=False,PinType.bIsSet=False,PinType.bIsArray=False,PinType.bIsReference=False,PinType.bIsConst=False,PinType.bIsWeakPointer=False,AutogeneratedDefaultValue="0, 0, 0",LinkedTo=(K2Node_CallFunction_147 CADAD3724E0A28E00552BCA72427B545,),PersistentGuid=00000000000000000000000000000000,bHidden=False,bNotConnectable=False,bDefaultValueIsReadOnly=False,bDefaultValueIsIgnored=False,bAdvancedView=False,)
-	End Object
-[/EXCERPT:snippet]
-
-[/COMMENT]
\ No newline at end of file
diff --git a/Engine/Documentation/Source/Platforms/Android/AndroidManifestControl/AndroidManifestControl.INT.udn b/Engine/Documentation/Source/Platforms/Android/AndroidManifestControl/AndroidManifestControl.INT.udn
--- a/Engine/Documentation/Source/Platforms/Android/AndroidManifestControl/AndroidManifestControl.INT.udn
+++ b/Engine/Documentation/Source/Platforms/Android/AndroidManifestControl/AndroidManifestControl.INT.udn
@@ -84,9 +84,7 @@ You can add additional tags for the Application by clicking on the **Plus** sign
 
  Items that are input into the Extra Tags for Application section will be added to the **Application Definition** section of the AndroidManifest.xml file.
 
-	[COMMENT:none]
- Application Definition 
-[/COMMENT]
+	<!-- Application Definition -->
 	<application android:label="@string/app_name"
 	 android:icon="@drawable/icon" 
 	 android:hardwareAccelerated="True" 
diff --git a/Engine/Documentation/Source/Platforms/Linux/AdvancedLinuxDeveloper/LinuxCrossCompileLegacy/LinuxCrossCompileLegacy.INT.udn b/Engine/Documentation/Source/Platforms/Linux/AdvancedLinuxDeveloper/LinuxCrossCompileLegacy/LinuxCrossCompileLegacy.INT.udn
--- a/Engine/Documentation/Source/Platforms/Linux/AdvancedLinuxDeveloper/LinuxCrossCompileLegacy/LinuxCrossCompileLegacy.INT.udn
+++ b/Engine/Documentation/Source/Platforms/Linux/AdvancedLinuxDeveloper/LinuxCrossCompileLegacy/LinuxCrossCompileLegacy.INT.udn
@@ -13,10 +13,6 @@ related:Engine/Basics/Projects/Packaging
 related:Platforms/Linux/GettingStarted
 topic-image:LegacyRef_Topic.png
 
-[COMMENT:none]
-- %Globals:BinarySourceSelectVar% -
-[/COMMENT]
-
 [TOC(start:2 end:3)]
 
 This reference is an archive for users who need to set up their cross-compilation toolchains in Unreal Engine 4 (UE4) versions that were released prior to 4.14. 
@@ -35,13 +31,15 @@ Mac users currently have to resort to [native compilation](https://wiki.unrealen
 
 Please use the following table to download the appropriate toolchain: 
 
-| Toolchain                                                                                                                       | UE4 Version    | 
-| ------------------------------------------------------------------------------------------------------------------------------- | -------------- | 
-| **-v4** [clang-3.5.0-based](http://cdn.unrealengine.com/qfe/v4_clang-3.5.0_ld-2.24_glibc-2.12.2.zip)                            | 4.8 and earlier| 
-| **-v6** [clang-3.6.0-based](https://s3.amazonaws.com/unrealengine/qfe/v6_clang-3.6.0_ld-2.24_glibc-2.12.2.zip)                  | 4.9 and 4.10   | 
-| **-v7** [clang-3.7.0-based](https://s3.amazonaws.com/unrealengine/CrossToolchain_Linux/v7_clang-3.7.0_ld-2.24_glibc-2.12.2.zip) | 4.11 thru 4.13 | 
-| **-v8** [clang-3.9.0-based](http://cdn.unrealengine.com/qfe/v8_clang-3.9.0-centos7.zip)                                         | 4.14 and 4.15  |
-| **-v9** [clang-4.0.0-based](http://cdn.unrealengine.com/CrossToolchain_Linux/v9_clang-4.0.0-centos7.zip)                        | 4.16           |
+| Toolchain                                                                                                                        | UE4 Version    | 
+| -------------------------------------------------------------------------------------------------------------------------------- | -------------- | 
+| **-v11** [clang-5.0.0-based](http://cdn.unrealengine.com/CrossToolchain_Linux/v11_clang-5.0.0-centos7.zip)                       | 4.19           |
+| **-v10** [clang-5.0.0-based](http://cdn.unrealengine.com/CrossToolchain_Linux/v10_clang-5.0.0-centos7.zip)                       | 4.18           |
+| **-v9**  [clang-4.0.0-based](http://cdn.unrealengine.com/CrossToolchain_Linux/v9_clang-4.0.0-centos7.zip)                        | 4.16 and 4.17  |
+| **-v8**  [clang-3.9.0-based](http://cdn.unrealengine.com/qfe/v8_clang-3.9.0-centos7.zip)                                         | 4.14 and 4.15  |
+| **-v7**  [clang-3.7.0-based](https://s3.amazonaws.com/unrealengine/CrossToolchain_Linux/v7_clang-3.7.0_ld-2.24_glibc-2.12.2.zip) | 4.11 thru 4.13 | 
+| **-v6**  [clang-3.6.0-based](https://s3.amazonaws.com/unrealengine/qfe/v6_clang-3.6.0_ld-2.24_glibc-2.12.2.zip)                  | 4.9 and 4.10   | 
+| **-v4**  [clang-3.5.0-based](http://cdn.unrealengine.com/qfe/v4_clang-3.5.0_ld-2.24_glibc-2.12.2.zip)                            | 4.8 and earlier| 
 
 [REGION:note]
 The Setup and UnSetup Windows Batch Files (\*.bat) are not included in **-v9** because you only need these files if you're setting up AutoSDKs for UE4. 
diff --git a/Engine/Documentation/Source/Platforms/Linux/BeginnerLinuxDeveloper/SettingUpAnUnrealWorkflow/SettingUpAnUnrealWorkflow.INT.udn b/Engine/Documentation/Source/Platforms/Linux/BeginnerLinuxDeveloper/SettingUpAnUnrealWorkflow/SettingUpAnUnrealWorkflow.INT.udn
--- a/Engine/Documentation/Source/Platforms/Linux/BeginnerLinuxDeveloper/SettingUpAnUnrealWorkflow/SettingUpAnUnrealWorkflow.INT.udn
+++ b/Engine/Documentation/Source/Platforms/Linux/BeginnerLinuxDeveloper/SettingUpAnUnrealWorkflow/SettingUpAnUnrealWorkflow.INT.udn
@@ -16,8 +16,10 @@ topic-image:linuxQS_topic.png
 
 [REGION:imagetable]
 ![](CompiledEditor_Linux.png)(w:920)  
+[REGION:caption]
 At the end of this tutorial, you'll have Unreal Engine 4 (UE4) running on your Linux machine.
 [/REGION]
+[/REGION]
 
 ## Goal
 
diff --git a/Engine/Documentation/Source/Platforms/Linux/GettingStarted/LinuxCrossGettingStarted.INT.udn b/Engine/Documentation/Source/Platforms/Linux/GettingStarted/LinuxCrossGettingStarted.INT.udn
--- a/Engine/Documentation/Source/Platforms/Linux/GettingStarted/LinuxCrossGettingStarted.INT.udn
+++ b/Engine/Documentation/Source/Platforms/Linux/GettingStarted/LinuxCrossGettingStarted.INT.udn
@@ -15,9 +15,6 @@ related:Programming/Basics/ConfigurationFiles
 topic-image:GettingStarted_Topic.png
 prereq:Platforms/Linux/BeginnerLinuxDeveloper/SettingUpAnUnrealWorkflow
 
-[COMMENT:none]
-- %Globals:BinarySourceSelectVar% -
-[/COMMENT]
 
 [TOC(start:2 end:3)]
 
@@ -72,7 +69,7 @@ Internet connection, downloading the toolchain can take anywhere from a few minu
 		Click for full image.
 	[/REGION]
 
-1.	Open the `[ROOT]/Linux_CrossCompileToolChain/v8_clang-3.9.0-centos7` folder to view the extracted files.
+1.	Open the `[ROOT]/Linux_CrossCompileToolChain/v11_clang-5.0.0-centos7` folder to view the extracted files.
 
 	[REGION:lightbox]
 		[![](LinuxXC_Step4.png)(w:600)](LinuxXC_Step4.png)
@@ -109,7 +106,7 @@ Internet connection, downloading the toolchain can take anywhere from a few minu
 		Click for full image.
 	[/REGION]
 
-1.	At this point, a new **OutputEnvVars** text file should be located inside of the `v8_clang-3.9.0-centos7` folder. 
+1.	At this point, a new **OutputEnvVars** text file should be located inside of the `v11_clang-5.0.0-centos7` folder. 
 
 	[REGION:lightbox]
 		[![](LinuxXC_Step6.png)(w:600)](LinuxXC_Step6.png)
@@ -270,4 +267,4 @@ Congratulations! Moving forward, you should be able to see the cross-compile too
 [/REGION]
 
  Now, you're ready to [compile](Programming/Development/CompilingProjects), [build](Programming/Development/BuildConfigurations), and [package](Engine/Basics/Projects/Packaging) projects for a variety of Linux platforms from your Windows (host) computer. 
- 
\ No newline at end of file
+ 
diff --git a/Engine/Documentation/Source/Platforms/Linux/IntermediateLinuxDeveloper/AndroidOnLinux/1/AndroidOnLinux_1.INT.udn b/Engine/Documentation/Source/Platforms/Linux/IntermediateLinuxDeveloper/AndroidOnLinux/1/AndroidOnLinux_1.INT.udn
--- a/Engine/Documentation/Source/Platforms/Linux/IntermediateLinuxDeveloper/AndroidOnLinux/1/AndroidOnLinux_1.INT.udn
+++ b/Engine/Documentation/Source/Platforms/Linux/IntermediateLinuxDeveloper/AndroidOnLinux/1/AndroidOnLinux_1.INT.udn
@@ -110,11 +110,6 @@ Before continuing, please note that Android Vulkan on Linux is not supported at
 
 	[REGION:tip]
 	If you don't want to use Ubuntu's GUI, you can restart your computer by entering `sudo reboot` into the terminal. 
-	[COMMENT:none]
- 
-	![](LinuxAndroid_Step16.png)(w:900)
-	
-[/COMMENT]
 	[/REGION]
 [/EXCERPT:AndroidLinux_1]
 
diff --git a/Engine/Documentation/Source/Platforms/Mobile/InAppPurchases/In-AppPurchases.INT.udn b/Engine/Documentation/Source/Platforms/Mobile/InAppPurchases/In-AppPurchases.INT.udn
--- a/Engine/Documentation/Source/Platforms/Mobile/InAppPurchases/In-AppPurchases.INT.udn
+++ b/Engine/Documentation/Source/Platforms/Mobile/InAppPurchases/In-AppPurchases.INT.udn
@@ -20,6 +20,17 @@ See the apropriate platform-specific page below for details on configuring in-ap
 [/REGION]
 
 [EXCERPT:usage]
+[INCLUDE:#usage_android]
+
+## Restoring Purchases
+
+![image alt text](image_5.png)
+
+To restore purchases, use the **Restore In-App Purchases** Blueprint node (or associated C++ function call). It only takes in the Player Controller, and returns an array of all the purchase information associated with that Player Controller. You can then process that array to compare against the particular Identifier (or Identifiers) that your gameplay logic needs. 
+
+[/EXCERPT:usage]
+
+[COMMENT:none]
 [EXCERPT:usage_android]
 ## Reading Purchase Information
 
@@ -38,12 +49,4 @@ The **Make an In-App Purchase** node is latent, so any gameplay behavior you wan
 
 There is a non-latent version of this function (the Blueprint node will display without the clock). The exit execution pin here does not wait for a response from the online service, and you will generally want to use the latent version.
 [/EXCERPT:usage_android]
-
-## Restoring Purchases
-
-![image alt text](image_5.png)
-
-To restore purchases, use the **Restore In-App Purchases** Blueprint node (or associated C++ function call). It only takes in the Player Controller, and returns an array of all the purchase information associated with that Player Controller. You can then process that array to compare against the particular Identifier (or Identifiers) that your gameplay logic needs. 
-
-[/EXCERPT:usage]
-
+[/COMMENT]
\ No newline at end of file
diff --git a/Engine/Documentation/Source/Platforms/Mobile/MobilePackagingWizard/MobilePackagingWizard.INT.udn b/Engine/Documentation/Source/Platforms/Mobile/MobilePackagingWizard/MobilePackagingWizard.INT.udn
--- a/Engine/Documentation/Source/Platforms/Mobile/MobilePackagingWizard/MobilePackagingWizard.INT.udn
+++ b/Engine/Documentation/Source/Platforms/Mobile/MobilePackagingWizard/MobilePackagingWizard.INT.udn
@@ -343,8 +343,7 @@ Now that all of the options for the Mobile Packing Wizard have been set up and s
 	[/OBJECT]
 
 	[region:warning]
-	You must package your project's files before you build the DLC. <!--Trying to package the DLC first and then packing your project will not.
-[/COMMENT]
+	You must package your project's files before you build the DLC.
 	[/region]
 
 1. Once the **Launch this Profile** button has been pressed, the Project Launcher will begin the packaging process. You can see the progress of the packaging process in the Project Launcher window.
@@ -740,9 +739,7 @@ Now that all of the options for the Mobile Packing Wizard have been set up and s
 	[/OBJECT]
 
 	[region:warning]
-	You must package your project's files before you build the DLC. [COMMENT:none]
-Trying to package the DLC first and then packing your project will not.
-[/COMMENT]
+	You must package your project's files before you build the DLC.
 	[/region]
 
 1. Once the **Launch this Profile** button has been pressed, the Project Launcher will begin the packaging process. You can see the progress of the packaging process in the Project Launcher window.
@@ -904,4 +901,4 @@ mac
 Inside of the folders, you will find the various pieces of your project that have been set up to be downloaded after the users have downloaded and installed the application on their mobile device.
 
 [/EXCERPT:ios]
--->
\ No newline at end of file
+[/COMMENT]
\ No newline at end of file
diff --git a/Engine/Documentation/Source/Platforms/Mobile/Textures/MobileTextures.INT.udn b/Engine/Documentation/Source/Platforms/Mobile/Textures/MobileTextures.INT.udn
--- a/Engine/Documentation/Source/Platforms/Mobile/Textures/MobileTextures.INT.udn
+++ b/Engine/Documentation/Source/Platforms/Mobile/Textures/MobileTextures.INT.udn
@@ -17,9 +17,7 @@ when creating textures:
 
 * A maximum size of 2048 in either dimension. This is the largest size allowed by the PVR compressor.
 * Dimensions must both be a power of 2 (i.e. 2, 4, 8, 16, 32, 64, 128, 256, 512, 1024, or 2048).
-* You should use square textures whenever possible as they make more efficient use of memory. [COMMENT:none]
- Non-square textures are padded out to be square. (?)  
-[/COMMENT]
+* You should use square textures whenever possible as they make more efficient use of memory. 
 
 In addition, texture settings need to be set up properly to work on mobile platforms:
 
diff --git a/Engine/Documentation/Source/Programming/Automation/FBXTestBuilder/FBXTestBuilder.INT.udn b/Engine/Documentation/Source/Programming/Automation/FBXTestBuilder/FBXTestBuilder.INT.udn
--- a/Engine/Documentation/Source/Programming/Automation/FBXTestBuilder/FBXTestBuilder.INT.udn
+++ b/Engine/Documentation/Source/Programming/Automation/FBXTestBuilder/FBXTestBuilder.INT.udn
@@ -134,7 +134,7 @@ For each Test Plan, a list of **Expected Results** will be evaluated to determin
 [INCLUDE:#expectedpresettype]
 [EXCERPT:expectedpresettype]
 | Check Type | Description | Required Data Types |
-| --- | --- |
+| --- | --- | --- |
 | **Error Number** | Data should contain the number of error [int0]. | Integer | 
 | **Warning Number** | Data should contain the number of warning [int0]. | Integer |
 | **Created Static Mesh Number** | Data should contain the number of Static Mesh created [int0]. | Integer |
diff --git a/Engine/Documentation/Source/Programming/Development/BuildingUnrealEngine/BuildingUnrealEngine.INT.udn b/Engine/Documentation/Source/Programming/Development/BuildingUnrealEngine/BuildingUnrealEngine.INT.udn
--- a/Engine/Documentation/Source/Programming/Development/BuildingUnrealEngine/BuildingUnrealEngine.INT.udn
+++ b/Engine/Documentation/Source/Programming/Development/BuildingUnrealEngine/BuildingUnrealEngine.INT.udn
@@ -105,6 +105,22 @@ tags:Unreal Build System
 	[/PARAM]
 [/OBJECT]
 
+## Getting Started with Unreal Engine
+
+Learn how to use Unreal Engine by referring to the [](GettingStarted) documentation!
+
+If you're looking to quickly get started with UE4, refer to the following tutorials:
+*	[](Engine/Content/QuickStart)
+*	[](Programming/QuickStart)
+*	[](Engine/QuickStart)
+
+
+[REGION:note]
+UE4's in-editor help features are a great way to get your questions answered.
+[/REGION]
+
+[/EXCERPT:Main]
+
 [COMMENT:none]
 
 [EXCERPT:windowscompile]
@@ -173,10 +189,6 @@ the source code on a machine that has at least eight (8) gigabytes of RAM with a
 
 [/EXCERPT:linuxcompile]
 
-[/COMMENT]
-
-[COMMENT:none]
-
 [EXCERPT:windowsrun]
 
 1.	Set your startup project to **UE4**.
@@ -225,20 +237,4 @@ If you're unable to create an OpenGL context as you're working through these ste
 	![](CompiledEditor_Linux.png)
 [/EXCERPT:linuxrun]
 
-[/COMMENT]
-
-## Getting Started with Unreal Engine
-
-Learn how to use Unreal Engine by referring to the [](GettingStarted) documentation!
-
-If you're looking to quickly get started with UE4, refer to the following tutorials:
-*	[](Engine/Content/QuickStart)
-*	[](Programming/QuickStart)
-*	[](Engine/QuickStart)
-
-
-[REGION:note]
-UE4's in-editor help features are a great way to get your questions answered.
-[/REGION]
-
-[/EXCERPT:Main]
+[/COMMENT]
\ No newline at end of file
diff --git a/Engine/Documentation/Source/Programming/Development/CompilingProjects/CompilingProjects.INT.udn b/Engine/Documentation/Source/Programming/Development/CompilingProjects/CompilingProjects.INT.udn
--- a/Engine/Documentation/Source/Programming/Development/CompilingProjects/CompilingProjects.INT.udn
+++ b/Engine/Documentation/Source/Programming/Development/CompilingProjects/CompilingProjects.INT.udn
@@ -23,6 +23,11 @@ project created from a Blueprints Only template.
 
 UE4 and the UnrealBuildTool use different build configurations to determine how the engine is compiled. Which configuration you use will be determined by the purposes of the build you want to create.
 
+[INCLUDE:#winbuildconfig]
+
+[/EXCERPT:Main]
+
+[COMMENT:none]
 [EXCERPT:winbuildconfig]
 Each build configuration contains two keywords.  The first of these indicates the state of the engine and your game project. For instance, if you compile using a **Debug** configuration, you'll
 be able to debug your game's code. The second keyword indicates the target you're building for. For example, if you'd like to open a project in Unreal, you'd need to build with the **Editor**
@@ -43,8 +48,7 @@ target keyword.
 | **Client**                  | If you're working on a multiplayer project using UE4 networking features, this target designates the specified project as being a Client in UE4's client-server model for multiplayer games. If there is a `<Game>Client.Target.cs` file, the **Client** build configurations will be valid.                                                    |
 | **Server**                  | If you're working on a multiplayer project using UE4 networking features, this target designates the specified project as being a Server in UE4's client-server model for multiplayer games. If there is a `<Game>Server.Target.cs` file, the **Server** build configurations will be valid.                                                    |                                                                                                                                
 [/EXCERPT:winbuildconfig]
-
-[/EXCERPT:Main]
+[/COMMENT]
 
 [EXCERPT:VS]
 ## Building with Visual Studio
@@ -121,6 +125,17 @@ When building in Visual Studio, you are compiling the engine in addition to the
 
 You can now [run the engine](GettingStarted/RunningUnrealEngine) with your compiled project.  
 
+[INCLUDE:#MatchEXE]
+
+### Visual Studio Known Issues
+
+| Issue | Solution |
+| ----- | -------- |
+| "Project is out of date" message always appears | Visual Studio thinks the project is out of date when it really is up to date. You can safely suppress this message by checking the **Do not show this dialog again** box and pressing **No**. |
+| **No Debugging Information** window appears when using the _Debug_ configuration | The reason for this is that the UE4editor.exe was compiled using the _Development_ configuration. Your game project will still be debuggable when compiled using the _Debug_ configuration despite this warning. You can safely suppress this message by checking the **Do not show this dialog again** box and pressing **Yes**. |
+[/EXCERPT:VS]
+
+[COMMENT:none]
 [EXCERPT:MatchEXE]
 [REGION:note]
 When running the engine, it is important to use the Unreal Engine executable that matches the build configuration you rebuilt your project in.  For example, if you compiled your project in the **DebugGame Uncooked**
@@ -131,14 +146,7 @@ build configuration, you would run the `UE4-Win64-DebugGame.exe` executable with
 When running Unreal, it is important to add the `-game` flag if you rebuilt your project in any **Uncooked** configuration, and the `-debug` flag if you rebuilt your project in any **Debug** configuration.  
 [/REGION:note]
 [/EXCERPT:MatchEXE]
-
-### Visual Studio Known Issues
-
-| Issue | Solution |
-| ----- | -------- |
-| "Project is out of date" message always appears | Visual Studio thinks the project is out of date when it really is up to date. You can safely suppress this message by checking the **Do not show this dialog again** box and pressing **No**. |
-| **No Debugging Information** window appears when using the _Debug_ configuration | The reason for this is that the UE4editor.exe was compiled using the _Development_ configuration. Your game project will still be debuggable when compiled using the _Debug_ configuration despite this warning. You can safely suppress this message by checking the **Do not show this dialog again** box and pressing **Yes**. |
-[/EXCERPT:VS]
+[/COMMENT]
 
 [EXCERPT:Xcode]
 
@@ -154,11 +162,7 @@ The Unreal Engine Xcode project is configured to build Debug configuration when
 
 You can now [run the engine](GettingStarted/RunningUnrealEngine) with your compiled project.  
 
-[EXCERPT:MatchDMG]
-[REGION:note]
-When running the binary Editor, it is important to add the `-game` flag if you rebuilt your project in any **Uncooked** configuration, and the `-debug` flag if you rebuilt your project in any **Debug** configuration.  
-[/REGION:note]	
-[/EXCERPT:MatchDMG]
+[INCLUDE:#MatchDMG]
 
 ### Xcode Known Issues
 
@@ -167,6 +171,7 @@ When running the binary Editor, it is important to add the `-game` flag if you r
 | ----- | -------- |
 | Xcode does not stop at breakpoints | [INCLUDE:#XcodeDebug] |
 
+[/EXCERPT:Xcode]
 [COMMENT:none]
 
 [EXCERPT:XcodeDebug]
@@ -175,11 +180,10 @@ Xcode's LLDB debugger needs to be configured to correctly handle breakpoints in
 	settings set target.inline-breakpoint-strategy always
 [/EXCERPT:XcodeDebug]
 
-[/COMMENT]
-
-[/EXCERPT:Xcode]
+[EXCERPT:MatchDMG]
+[REGION:note]
+When running the binary Editor, it is important to add the `-game` flag if you rebuilt your project in any **Uncooked** configuration, and the `-debug` flag if you rebuilt your project in any **Debug** configuration.  
+[/REGION:note]	
+[/EXCERPT:MatchDMG]
 
-[COMMENT:none]
- The following was deprecated per LaurenR 29.7.15
-| **Uncooked**                | Projects built using configurations that have an **Uncooked** type should be opened in Unreal Editor with the "-game" flag.  This runs your game with uncooked content, in a new window, and is equivalent to **Play in > New Window at Default Player Start** in the editor.       | 
-[/COMMENT]
\ No newline at end of file
+[/COMMENT]
diff --git a/Engine/Documentation/Source/Programming/UnrealBuildSystem/ProjectFileGenerator/ProjectFileGenerator.INT.udn b/Engine/Documentation/Source/Programming/UnrealBuildSystem/ProjectFileGenerator/ProjectFileGenerator.INT.udn
--- a/Engine/Documentation/Source/Programming/UnrealBuildSystem/ProjectFileGenerator/ProjectFileGenerator.INT.udn
+++ b/Engine/Documentation/Source/Programming/UnrealBuildSystem/ProjectFileGenerator/ProjectFileGenerator.INT.udn
@@ -18,9 +18,7 @@ The following guide is meant for users who've downloaded and built UE4 from its
 
 If you download the latest Unreal Engine 4 (UE4) code, you might notice that there are no Visual Studio (VS) or Xcode project files included for compiling and running the engine or example games.  You will need to run a script that generates project files for you to load:
 
-* Run the GenerateProjectFiles.bat file which is located in the main UE4 directory[COMMENT:none]
-(mapped to '//depot/UE4/`)
-[/COMMENT].
+* Run the GenerateProjectFiles.bat file which is located in the main UE4 directory
 
 	[REGION:note]
 	You can run that script from any folder.  (Make a desktop shortcut if you want.)
@@ -127,9 +125,7 @@ This script is a simple wrapper around Unreal Build Tool, which is launched in a
 
 * Change your startup project to UnrealBuildTool.
 * Set the command-line parameters for debugging to:  -ProjectFiles
-* Set the working directory to your local path for:  [COMMENT:none]
-//depot/UE4
-[/COMMENT][Your UE4 Directory]/Engine/Source/
+* Set the working directory to your local path for:  [Your UE4 Directory]/Engine/Source/
 * Compile and debug as per normal.
 
 
diff --git a/Engine/Documentation/Source/Resources/ContentExamples/Lighting/2_1/2_1.INT.udn b/Engine/Documentation/Source/Resources/ContentExamples/Lighting/2_1/2_1.INT.udn
--- a/Engine/Documentation/Source/Resources/ContentExamples/Lighting/2_1/2_1.INT.udn
+++ b/Engine/Documentation/Source/Resources/ContentExamples/Lighting/2_1/2_1.INT.udn
@@ -77,7 +77,7 @@ On Brush surfaces, lightmap resolution is set via the _Lightmap Resolution_ prop
 	![Lightmap Resolution: 16.0](LightmapRes_16.png)(w:%2XSize%)
 	[/PARAM]
 [/OBJECT]
-
+[/EXCERPT:StaticLighting02]
 
 [EXCERPT:ILC]
 Static lights can also affect dynamic objects due to the _Indirect Lighting Cache Sample Grid_. This is a 3D grid of lighting samples that get baked into the scene and are used by the rendering engine to determine what effects static lighting will have on dynamic objects. The grid can be visualized in the viewport by typing the following command into the Console Bar in the upper-right corner of the editor:
diff --git a/Engine/Documentation/Source/Shared/Glossary/A/Glossary_A.INT.udn b/Engine/Documentation/Source/Shared/Glossary/A/Glossary_A.INT.udn
--- a/Engine/Documentation/Source/Shared/Glossary/A/Glossary_A.INT.udn
+++ b/Engine/Documentation/Source/Shared/Glossary/A/Glossary_A.INT.udn
@@ -10,18 +10,7 @@ Availability:NoPublish
 Actor
 	[/PARAM]
 	[PARAM:Definition]
-[COMMENT:none]
-
-**Actors** are the most important concept in Unreal Engine 4. Your game will be filled with them, as every interactive entity in your game is an Actor. They usually contain one or more components, can be placed in a **Level** or **Spawned** in during gameplay, and support network replication for multiplayer games.
-
-[/COMMENT]
-
-[EXCERPT:Glossary_Actor]
-An **Actor** is any object that can be placed into a level. Actors are a generic Class that support 3D transformations such as translation, rotation, and scale. Actors can be created (spawned) and destroyed through gameplay code (C++ or Blueprints). In C++, AActor is the base class of all Actors. 
-
-There are several different types of Actors, some examples include: StaticMeshActor, CameraActor, and PlayerStartActor. 
-[/EXCERPT:Glossary_Actor]
-
+		[INCLUDE:#Glossary_Actor]
 	[/PARAM]
 [/OBJECT]
 
@@ -31,11 +20,7 @@ There are several different types of Actors, some examples include: StaticMeshAc
 AI
 	[/PARAM]
 	[PARAM:Definition]
-
-[EXCERPT:Glossary_AI]
-**AI**, or Artificial intelligence, is behavior controlled by a set of programmed instructions, instead of by human input.  For example, non-player characters (NPCs) are controlled with AI.
-[/EXCERPT:Glossary_AI]
-
+		[INCLUDE:#Glossary_AI]
 	[/PARAM]
 [/OBJECT]
 
@@ -44,11 +29,7 @@ AI
 AIController
 	[/PARAM]
 	[PARAM:Definition]
-
-[EXCERPT:Glossary_AIController]
-Just as the PlayerController possesses a Pawn as a representation of the player in a game, an **AIController** possesses a Pawn to represent a non-player character (NPC) in a game. By default, Pawns and Characters will end up with a base AIController unless they are specifically possessed by a PlayerController or told not to create an AIController for themselves.
-[/EXCERPT:Glossary_AIController]
-
+		[INCLUDE:#Glossary_AIController]
 	[/PARAM]
 [/OBJECT]
 
@@ -58,13 +39,7 @@ Just as the PlayerController possesses a Pawn as a representation of the player
 animation
 	[/PARAM]
 	[PARAM:Definition]
-
-[EXCERPT:Glossary_animation]
-**Animation** is a change in a property value over time. Examples include changes in position, rotation, and/or scale, but can also extend to any other property value as well.
-
-"Animation" can refer specifically to sequences created in Unreal Matinee, or to sequences created for Skeletal Meshes in external 3D modeling and animation applications, and then imported into the Unreal Editor. It can even refer to a change over time in an Actor's property (or properties) created by C++ code or Blueprints.
-[/EXCERPT:Glossary_animation]
-
+		[INCLUDE:#Glossary_Animation]
 	[/PARAM]
 [/OBJECT]
 
@@ -74,11 +49,7 @@ animation
 Animation Sequence
 	[/PARAM]
 	[PARAM:Definition]
-
-[EXCERPT:Glossary_animation_sequence]
-An **Animation Sequence**, or **AnimSequence**, is an animation asset created in a 3D modeling and animation application for a Skeletal Mesh, and then imported into the Unreal Editor. You can work with Animation Sequences in Persona.
-[/EXCERPT:Glossary_animation_sequence]
-
+		[INCLUDE:#Glossary_Animation_Sequence]
 	[/PARAM]
 [/OBJECT]
 
@@ -88,14 +59,34 @@ An **Animation Sequence**, or **AnimSequence**, is an animation asset created in
 asset
 	[/PARAM]
 	[PARAM:Definition]
-
-[EXCERPT:Glossary_asset]
-An **asset** is an element used to build a game.   Examples include content imported into or created in the **Content Browser**, such as Materials, Static Meshes, Textures, Particle Systems, Blueprints, and Sound Cues.  Each asset is saved in an individual .uasset file.
-[/EXCERPT:Glossary_asset]
-
+		[INCLUDE:#Glossary_Asset]
 	[/PARAM]
 [/OBJECT]
 
 
 [/EXCERPT:Glossary_A]
 
+[COMMENT:none]
+[EXCERPT:Glossary_Actor]
+An **Actor** is any object that can be placed into a level. Actors are a generic Class that support 3D transformations such as translation, rotation, and scale. Actors can be created (spawned) and destroyed through gameplay code (C++ or Blueprints). In C++, AActor is the base class of all Actors. 
+
+There are several different types of Actors, some examples include: StaticMeshActor, CameraActor, and PlayerStartActor. 
+[/EXCERPT:Glossary_Actor]
+[EXCERPT:Glossary_AI]
+**AI**, or Artificial intelligence, is behavior controlled by a set of programmed instructions, instead of by human input.  For example, non-player characters (NPCs) are controlled with AI.
+[/EXCERPT:Glossary_AI]
+[EXCERPT:Glossary_AIController]
+Just as the PlayerController possesses a Pawn as a representation of the player in a game, an **AIController** possesses a Pawn to represent a non-player character (NPC) in a game. By default, Pawns and Characters will end up with a base AIController unless they are specifically possessed by a PlayerController or told not to create an AIController for themselves.
+[/EXCERPT:Glossary_AIController]
+[EXCERPT:Glossary_Animation]
+**Animation** is a change in a property value over time. Examples include changes in position, rotation, and/or scale, but can also extend to any other property value as well.
+
+"Animation" can refer specifically to sequences created in Unreal Matinee, or to sequences created for Skeletal Meshes in external 3D modeling and animation applications, and then imported into the Unreal Editor. It can even refer to a change over time in an Actor's property (or properties) created by C++ code or Blueprints.
+[/EXCERPT:Glossary_Animation]
+[EXCERPT:Glossary_Animation_Sequence]
+An **Animation Sequence**, or **AnimSequence**, is an animation asset created in a 3D modeling and animation application for a Skeletal Mesh, and then imported into the Unreal Editor. You can work with Animation Sequences in Persona.
+[/EXCERPT:Glossary_Animation_Sequence]
+[EXCERPT:Glossary_Asset]
+An **asset** is an element used to build a game.   Examples include content imported into or created in the **Content Browser**, such as Materials, Static Meshes, Textures, Particle Systems, Blueprints, and Sound Cues.  Each asset is saved in an individual .uasset file.
+[/EXCERPT:Glossary_Asset]
+[/COMMENT]
\ No newline at end of file
diff --git a/Engine/Documentation/Source/Shared/Glossary/B/Glossary_B.INT.udn b/Engine/Documentation/Source/Shared/Glossary/B/Glossary_B.INT.udn
--- a/Engine/Documentation/Source/Shared/Glossary/B/Glossary_B.INT.udn
+++ b/Engine/Documentation/Source/Shared/Glossary/B/Glossary_B.INT.udn
@@ -11,80 +11,65 @@ Availability:NoPublish
 Blueprints
 	[/PARAM]
 	[PARAM:Definition]
-
-[EXCERPT:Glossary_Blueprints]
-The **Blueprints Visual Scripting** system (or **Blueprints** for short) is a visual scripting system based on the concept of adding and connecting nodes together to create gameplay elements all from within the Unreal Editor. Flexible and powerful enough to create an entire game without ever having to touch a single line of code, Blueprints provide the ability to use virtually the full range of concepts and tools generally only available to programmers. 
-[/EXCERPT:Glossary_Blueprints]
-
+		[INCLUDE:#Glossary_Blueprints]
 	[/PARAM]
 [/OBJECT]
 
-
-
-
 [OBJECT:GlossaryItem]
 	[PARAM:Term]
 Blueprint Editor
 	[/PARAM]
 	[PARAM:Definition]
-
-[EXCERPT:Glossary_BlueprintEditor]
-The **Blueprint Editor** is a node-based graph editor. It also has several additional windows that serve targeted information and allow for additional modifications to be made to Blueprints and their graphs.
-[/EXCERPT:Glossary_BlueprintEditor]
-
+		[INCLUDE#Glossary_BlueprintEditor]
 	[/PARAM]
 [/OBJECT]
 
-
-
 [OBJECT:GlossaryItem]
 	[PARAM:Term]
 Blueprint Interface
 	[/PARAM]
 	[PARAM:Definition]
-
-[EXCERPT:Glossary_BlueprintInterface]
-A **Blueprint Interface** is a type of asset in which undefined functions are created. Any Blueprint can implement this Interface and then create its own definition of those functions. This serves as one way for Blueprints to communicate to one another, particularly when one Blueprint interacts with many other Blueprints.
-[/EXCERPT:Glossary_BlueprintInterface]
-
+		[INCLUDE:#Glossary_BlueprintInterface]
 	[/PARAM]
 [/OBJECT]
 
-
-
-
 [OBJECT:GlossaryItem]
 	[PARAM:Term]
 Blueprint Macro Library
 	[/PARAM]
 	[PARAM:Definition]
-
-[EXCERPT:Glossary_BlueprintMacroLibrary]
-A **Blueprint Macro Library** is a Blueprint container that holds a collection of macros or self-contained graphs that can be placed as nodes in other Blueprints. These can be time-savers, because they can store commonly used sequences of nodes, complete with inputs and outputs for both execution and data transfer. Macros are shared among all graphs that reference them, but they are auto-expanded into the graph as if they were a collapsed node during compiling. This means that Blueprint Macro Libraries do not need to be compiled, but changes to a macro are only reflected in graphs that reference that macro when the Blueprint containing those graphs is recompiled.
-[/EXCERPT:Glossary_BlueprintMacroLibrary]
-
+		[INCLUDE:#Glossary_BlueprintMacroLibrary]
 	[/PARAM]
 [/OBJECT]
 
-
-
-
 [OBJECT:GlossaryItem]
 	[PARAM:Term]
 Brush
 	[/PARAM]
 	[PARAM:Definition]
-
-[EXCERPT:Glossary_Brush]
-A **Brush** is an Actor that describes a 3D volume that is placed in a level in order to define level geometry (referred to as BSPs) and gameplay volumes. Typically you will use BSP Brushes to prototype or block-out your levels for gameplay testing. 
-
-Volumes on the other hand have several uses depending upon the effects attached to them such as: Blocking Volumes (which are invisible and used to prevent Actors from passing through them), Pain Causing Volumes (which causes damage over time to any Actor that overlaps it) or Trigger Volumes (which are used as a way to cause events when an Actor enters or exits them).  
-
-[/EXCERPT:Glossary_Brush]
- 
+		[INCLUDE:#Glossary_Brush]
 	[/PARAM]
 [/OBJECT]
 
 
 [/EXCERPT:Glossary_B]
 
+[COMMENT:none]
+[EXCERPT:Glossary_Blueprints]
+The **Blueprints Visual Scripting** system (or **Blueprints** for short) is a visual scripting system based on the concept of adding and connecting nodes together to create gameplay elements all from within the Unreal Editor. Flexible and powerful enough to create an entire game without ever having to touch a single line of code, Blueprints provide the ability to use virtually the full range of concepts and tools generally only available to programmers. 
+[/EXCERPT:Glossary_Blueprints]
+[EXCERPT:Glossary_BlueprintEditor]
+The **Blueprint Editor** is a node-based graph editor. It also has several additional windows that serve targeted information and allow for additional modifications to be made to Blueprints and their graphs.
+[/EXCERPT:Glossary_BlueprintEditor]
+[EXCERPT:Glossary_BlueprintInterface]
+A **Blueprint Interface** is a type of asset in which undefined functions are created. Any Blueprint can implement this Interface and then create its own definition of those functions. This serves as one way for Blueprints to communicate to one another, particularly when one Blueprint interacts with many other Blueprints.
+[/EXCERPT:Glossary_BlueprintInterface]
+[EXCERPT:Glossary_BlueprintMacroLibrary]
+A **Blueprint Macro Library** is a Blueprint container that holds a collection of macros or self-contained graphs that can be placed as nodes in other Blueprints. These can be time-savers, because they can store commonly used sequences of nodes, complete with inputs and outputs for both execution and data transfer. Macros are shared among all graphs that reference them, but they are auto-expanded into the graph as if they were a collapsed node during compiling. This means that Blueprint Macro Libraries do not need to be compiled, but changes to a macro are only reflected in graphs that reference that macro when the Blueprint containing those graphs is recompiled.
+[/EXCERPT:Glossary_BlueprintMacroLibrary]
+[EXCERPT:Glossary_Brush]
+A **Brush** is an Actor that describes a 3D volume that is placed in a level in order to define level geometry (referred to as BSPs) and gameplay volumes. Typically you will use BSP Brushes to prototype or block-out your levels for gameplay testing. 
+
+Volumes on the other hand have several uses depending upon the effects attached to them such as: Blocking Volumes (which are invisible and used to prevent Actors from passing through them), Pain Causing Volumes (which causes damage over time to any Actor that overlaps it) or Trigger Volumes (which are used as a way to cause events when an Actor enters or exits them).  
+[/EXCERPT:Glossary_Brush]
+[/COMMENT]
\ No newline at end of file
diff --git a/Engine/Documentation/Source/Shared/Glossary/C/Glossary_C.INT.udn b/Engine/Documentation/Source/Shared/Glossary/C/Glossary_C.INT.udn
--- a/Engine/Documentation/Source/Shared/Glossary/C/Glossary_C.INT.udn
+++ b/Engine/Documentation/Source/Shared/Glossary/C/Glossary_C.INT.udn
@@ -10,11 +10,7 @@
 Cascade
 	[/PARAM]
 	[PARAM:Definition]
-
-[EXCERPT:Glossary_Cascade]
-**Cascade** is the particle-system editor in the Unreal Editor. You can access Cascade by creating or opening a particle system in the **Content Browser**. In Cascade, you can add, modify, and remove Particle Emitters, add modules to Particle Emitters, and modify module properties.
-[/EXCERPT:Glossary_Cascade]
-
+		[INCLUDE:#Glossary_Cascade]
 	[/PARAM]
 [/OBJECT]
 
@@ -23,64 +19,117 @@
 Character
 	[/PARAM]
 	[PARAM:Definition]
+		[INCLUDE:#Glossary_Character]
+	[/PARAM]
+[/OBJECT]
 
+[OBJECT:GlossaryItem]
+	[PARAM:Term]
+cinematic sequence
+	[/PARAM]
+	[PARAM:Definition]
+		[INCLUDE:#Glossary_cinematic_sequence]
+	[/PARAM]
+[/OBJECT]
+
+[OBJECT:GlossaryItem]
+	[PARAM:Term]
+class
+	[/PARAM]
+	[PARAM:Definition]
+		[INCLUDE:#Glossary_class]
+	[/PARAM]
+[/OBJECT]
+
+[OBJECT:GlossaryItem]
+	[PARAM:Term]
+Blueprint Class
+	[/PARAM]
+	[PARAM:Definition]
+		[INCLUDE:#Glossary_Class_Blueprint]
+	[/PARAM]
+[/OBJECT]
+
+[OBJECT:GlossaryItem]
+	[PARAM:Term]
+collision
+	[/PARAM]
+	[PARAM:Definition]
+		[INCLUDE:#Glossary_collision]
+	[/PARAM]
+[/OBJECT]
+
+[OBJECT:GlossaryItem]
+	[PARAM:Term]
+component
+	[/PARAM]
+	[PARAM:Definition]
+		[INCLUDE:#Glossary_component]
+	[/PARAM]
+[/OBJECT]
+
+[OBJECT:GlossaryItem]
+	[PARAM:Term]
+Components Mode
+	[/PARAM]
+	[PARAM:Definition]
+		[INCLUDE:#Glossary_Components_Mode]
+
+		Components Mode is one of three modes; the others are Default Mode and Graph Mode.
+	[/PARAM]
+[/OBJECT]
+
+[OBJECT:GlossaryItem]
+	[PARAM:Term]
+Construction Script
+	[/PARAM]
+	[PARAM:Definition]
+		[INCLUDE:#Glossary_Construction_Script]
+
+		The Construction Script controls functionality while working with the object in the Editor, and the Event Graph controls functionality during gameplay; Construction Scripts do not execute during gameplay.
+	[/PARAM]
+[/OBJECT]
+
+[OBJECT:GlossaryItem]
+	[PARAM:Term]
+Content Browser
+	[/PARAM]
+	[PARAM:Definition]
+		[INCLUDE:#Glossary_Content_Browser]
+
+		[REGION:tip]
+		In the **Content Browser**, you can search and sort assets by type, and you can filter your view of the assets either by choosing specific types of assets in the **Filters** drop-down list, or by typing text in the **Search Assets** box.
+		[/REGION]
+	[/PARAM]
+[/OBJECT]
+
+[OBJECT:GlossaryItem]
+	[PARAM:Term]
+Controller Class
+	[/PARAM]
+	[PARAM:Definition]
+		[INCLUDE:#Glossary_Controller_Class]
+	[/PARAM]
+[/OBJECT]
+
+[/EXCERPT:Glossary_C]
+
+[COMMENT:none]
+[EXCERPT:Glossary_Cascade]
+**Cascade** is the particle-system editor in the Unreal Editor. You can access Cascade by creating or opening a particle system in the **Content Browser**. In Cascade, you can add, modify, and remove Particle Emitters, add modules to Particle Emitters, and modify module properties.
+[/EXCERPT:Glossary_Cascade]
 [EXCERPT:Glossary_Character]
 A **Character** is a subclass of a Pawn Actor that is intended to be used as a player character. The Character subclass includes a collision setup, input bindings for bipedal movement, and additional code for movement controlled by the player.
 [/EXCERPT:Glossary_Character]
-
-	[/PARAM]
-[/OBJECT]
-
-
-[OBJECT:GlossaryItem]
-	[PARAM:Term]
-cinematic sequence
-	[/PARAM]
-	[PARAM:Definition]
-
 [EXCERPT:Glossary_cinematic_sequence]
 A **cinematic sequence**, also just called a cinematic, is an animation that is played in the game. Some cinematic sequences halt gameplay and involve one or more changes in camera viewpoint; these are called **cutscenes**. A cinematic sequence can either be created and controlled by Matinee or imported as a pre-rendered video created in an external application.
 [/EXCERPT:Glossary_cinematic_sequence]
-
-	[/PARAM]
-[/OBJECT]
-
-
-[OBJECT:GlossaryItem]
-	[PARAM:Term]
-class
-	[/PARAM]
-	[PARAM:Definition]
-
 [EXCERPT:Glossary_class]
 A **Class** defines the behaviors and properties of a particular Actor or Object used in the creation of an Unreal Engine game. Classes are hierarchical, meaning a Class inherits information from its parent Classes (the Classes it was derived or "sub-classed" from) and passes that information to its children. Classes can be created in C++ code or in Blueprints.
 [/EXCERPT:Glossary_class]
-
-	[/PARAM]
-[/OBJECT]
-
-
-[OBJECT:GlossaryItem]
-	[PARAM:Term]
-Blueprint Class
-	[/PARAM]
-	[PARAM:Definition]
-
 [EXCERPT:Glossary_Class_Blueprint]
 A **Blueprint Class** is an asset that content creators can use to easily add functionality on top of existing gameplay classes. Often shortened as just _Blueprint_. Blueprints are created inside of Unreal Editor visually, instead of by typing code, and saved as assets in a content package. They essentially define a new class or type of Actor, which can then be placed into maps as instances that behave like any other type of Actor.
 [/EXCERPT:Glossary_Class_Blueprint]
-
-	[/PARAM]
-[/OBJECT]
-
-
-
-[OBJECT:GlossaryItem]
-	[PARAM:Term]
-collision
-	[/PARAM]
-	[PARAM:Definition]
-
 [EXCERPT:Glossary_collision]
 **Collision** is a programmatic way of preventing objects from overlapping during a physics simulation, giving the illusion of solidity in an object. May also refer to a Collision mesh, which is a simplified version of the geometry, generally existing as an invisible shell or hull around the mesh. This simplification aids in performance, allowing complex objects with high polygon counts to calculate their collisions with simpler geometry, which is faster to process.
 
@@ -90,59 +139,14 @@
 A Collision surface is required for an object to be simulated via physics calculations.
 [/REGION]
 [/EXCERPT:Glossary_collision]
-
-	[/PARAM]
-[/OBJECT]
-
-
-
-[OBJECT:GlossaryItem]
-	[PARAM:Term]
-component
-	[/PARAM]
-	[PARAM:Definition]
-
 [EXCERPT:Glossary_component]
 A **Component** is a piece of functionality that can be added to an Actor. Components cannot exist by themselves, however when added to an Actor, the Actor will have access to and can use functionality provided by the Component. 
 
 For example, a Spot Light Component will allow your Actor to emit light like a spot light, a Rotating Movement Component will make your Actor spin around, or an Audio Component will make your Actor able to play sounds. 
 [/EXCERPT:Glossary_component]
-
-[COMMENT:none]
-
-A **Component** is a special type of **Object** designed to be used as a sub-object within an **Actor**. Components are generally used in situations where it is necessary to have easily swappable parts in order to change the behavior or functionality of some particular aspect of the owning **Actor**. 
-
-For instance, a car is controlled and moves very differently from an aircraft, which is controlled and moves differently from a boat, and so on; yet all of these are vehicles that share other commonalities. By using a Component to handle the controls and movement, the same vehicle can easily be made to behave like any one of these specific types.
-
-[/COMMENT]
-
-	[/PARAM]
-[/OBJECT]
-
-
-
-[OBJECT:GlossaryItem]
-	[PARAM:Term]
-Components Mode
-	[/PARAM]
-	[PARAM:Definition]
-
 [EXCERPT:Glossary_Components_Mode]
 In **Components Mode**, you can add components to your Blueprint, manage the details of existing components, and arrange components. The list of components in Components Mode is hierarchical.
 [/EXCERPT:Glossary_Components_Mode]
-
-Components Mode is one of three modes; the others are Default Mode and Graph Mode.
-	[/PARAM]
-[/OBJECT]
-
-
-
-[OBJECT:GlossaryItem]
-	[PARAM:Term]
-Construction Script
-	[/PARAM]
-	[PARAM:Definition]
-
 [EXCERPT:Glossary_Construction_Script]
 The **Construction Script** runs following the Components list when an instance of a Blueprint Class 
 is created. It contains a node graph that is executed allowing the Blueprint Class instance to perform initialization operations. 
@@ -151,47 +155,10 @@
 upon and choose the correct mesh to use from a set of meshes or a fence Blueprint could perform traces extending out in 
 each direction to determine how long of a fence is needed to span the distance.
 [/EXCERPT:Glossary_Construction_Script]
-
-The Construction Script controls functionality while working with the object in the Editor, and the Event Graph controls functionality during gameplay; Construction Scripts do not execute during gameplay.
-
-	[/PARAM]
-[/OBJECT]
-
-
-
-[OBJECT:GlossaryItem]
-	[PARAM:Term]
-Content Browser
-	[/PARAM]
-	[PARAM:Definition]
-
 [EXCERPT:Glossary_Content_Browser]
 The **Content Browser** is the primary area of the Unreal Editor for creating, importing, organizing, viewing, and modifying content assets within Unreal Editor.  It also provides the ability to manage content folders and perform other useful operations on assets, such as renaming, moving, copying, and viewing references.  The **Content Browser** can search for and interact with all assets in the game.
 [/EXCERPT:Glossary_Content_Browser]
-
-[REGION:tip]
-In the **Content Browser**, you can search and sort assets by type, and you can filter your view of the assets either by choosing specific types of assets in the **Filters** drop-down list, or by typing text in the **Search Assets** box.
-[/REGION]
-
-	[/PARAM]
-[/OBJECT]
-
-
-
-[OBJECT:GlossaryItem]
-	[PARAM:Term]
-Controller Class
-	[/PARAM]
-	[PARAM:Definition]
-
 [EXCERPT:Glossary_Controller_Class]
 **Controllers** are non-physical Actors that can possess a Pawn (or Pawn-derived class like Character) to control its actions. A PlayerController is used by human players to control Pawns, while an AIController implements the artificial intelligence for the Pawns they control. Controllers take control of a Pawn with the Possess function, and give up control of the Pawn with the Unpossess function.
 [/EXCERPT:Glossary_Controller_Class]
-
-	[/PARAM]
-[/OBJECT]
-
-
-
-
-[/EXCERPT:Glossary_C]
+[/COMMENT]
diff --git a/Engine/Documentation/Source/Shared/Glossary/D/Glossary_D.INT.udn b/Engine/Documentation/Source/Shared/Glossary/D/Glossary_D.INT.udn
--- a/Engine/Documentation/Source/Shared/Glossary/D/Glossary_D.INT.udn
+++ b/Engine/Documentation/Source/Shared/Glossary/D/Glossary_D.INT.udn
@@ -11,38 +11,18 @@ Availability:NoPublish
 data pin
 	[/PARAM]
 	[PARAM:Definition]
-
-[EXCERPT:Glossary_data_pin]
-![Blueprint Data Pin Types](k2_pins_data_types.png)
-
-**Data pins** are used for taking data into a node or outputting data from a node. Data pins are type-specific and can be wired to variables of the same type (which have data pins of their own) or a data pin of the same type on another node. Like execution pins, data pins are displayed as an outline when not wired to anything, and solid when wired.
-
-![Blueprint Data Pins](k2_pins_data.png)
-
-Nodes can have any number of input or output data pins. The data pins of a Function Call node correspond to the parameters and return value of the corresponding function.
-[/EXCERPT:Glossary_data_pin]
-
+		[INCLUDE:#Glossary_data_pin]
 	[/PARAM]
 [/OBJECT]
 
-
-
-
 [OBJECT:GlossaryItem]
 	[PARAM:Term]
 Defaults Mode
 	[/PARAM]
 	[PARAM:Definition]
+		[INCLUDE:#Glossary_Defaults_Mode]
 
-[EXCERPT:Glossary_Defaults_Mode]
-The **Class Defaults** or **Defaults** mode displays all of the properties of the _Blueprint_, including variables 
-contained within the _Blueprint_ and components defined in native code as well as replication, navigation, 
-and collision properties. The values set for these properties in the Defaults mode will be 
-used when new instances of the _Blueprint_ are created.
-[/EXCERPT:Glossary_Defaults_Mode]
-
-Defaults Mode is one of three Blueprint Editor modes; the others are Components Mode and Graph Mode.
-
+		Defaults Mode is one of three Blueprint Editor modes; the others are Components Mode and Graph Mode.
 	[/PARAM]
 [/OBJECT]
 
@@ -51,11 +31,7 @@ Defaults Mode is one of three Blueprint Editor modes; the others are Components
 Details panel
 	[/PARAM]
 	[PARAM:Definition]
-
-[EXCERPT:Glossary_Details_panel]
-The **Details panel** is a panel in the Level Editor that contains information about the properties for various elements. Many of the tools and editors in UE4 contain Details panels: the Level Editor, Matinee, Persona, and so on.
-[/EXCERPT:Glossary_Details_panel]
-
+		[INCLUDE:#Glossary_Details_panel]
 	[/PARAM]
 [/OBJECT]
 
@@ -64,13 +40,32 @@ The **Details panel** is a panel in the Level Editor that contains information a
 Director Group
 	[/PARAM]
 	[PARAM:Definition]
-
-[EXCERPT:Glossary_Director_Group]
-The **Director Group** is a specific type of Matinee group that controls camera viewpoints. You can use the Director Group to switch camera viewpoints between shots. Each Matinee sequence can have only one Director Group.
-[/EXCERPT:Glossary_Director_Group]
-
+		[INCLUDE:#Glossary_Director_Group]
 	[/PARAM]
 [/OBJECT]
 
-
 [/EXCERPT:Glossary_D]
+
+[COMMENT:none]
+[EXCERPT:Glossary_data_pin]
+![Blueprint Data Pin Types](k2_pins_data_types.png)
+
+**Data pins** are used for taking data into a node or outputting data from a node. Data pins are type-specific and can be wired to variables of the same type (which have data pins of their own) or a data pin of the same type on another node. Like execution pins, data pins are displayed as an outline when not wired to anything, and solid when wired.
+
+![Blueprint Data Pins](k2_pins_data.png)
+
+Nodes can have any number of input or output data pins. The data pins of a Function Call node correspond to the parameters and return value of the corresponding function.
+[/EXCERPT:Glossary_data_pin]
+[EXCERPT:Glossary_Defaults_Mode]
+The **Class Defaults** or **Defaults** mode displays all of the properties of the _Blueprint_, including variables 
+contained within the _Blueprint_ and components defined in native code as well as replication, navigation, 
+and collision properties. The values set for these properties in the Defaults mode will be 
+used when new instances of the _Blueprint_ are created.
+[/EXCERPT:Glossary_Defaults_Mode]
+[EXCERPT:Glossary_Details_panel]
+The **Details panel** is a panel in the Level Editor that contains information about the properties for various elements. Many of the tools and editors in UE4 contain Details panels: the Level Editor, Matinee, Persona, and so on.
+[/EXCERPT:Glossary_Details_panel]
+[EXCERPT:Glossary_Director_Group]
+The **Director Group** is a specific type of Matinee group that controls camera viewpoints. You can use the Director Group to switch camera viewpoints between shots. Each Matinee sequence can have only one Director Group.
+[/EXCERPT:Glossary_Director_Group]
+[/COMMENT]
\ No newline at end of file
