Copying //UE4/Dev-Documentation to //UE4/Samples-Main (Source: //UE4/Dev-Documentation @ 4091886)

diff --git a/Engine/Documentation/Source/Shared/Glossary/E/Glossary_E.INT.udn b/Engine/Documentation/Source/Shared/Glossary/E/Glossary_E.INT.udn
--- a/Engine/Documentation/Source/Shared/Glossary/E/Glossary_E.INT.udn
+++ b/Engine/Documentation/Source/Shared/Glossary/E/Glossary_E.INT.udn
@@ -10,11 +10,7 @@ Availability:NoPublish
 event
 	[/PARAM]
 	[PARAM:Definition]
-
-[EXCERPT:Glossary_event]
-In Blueprints, an **event** is an entry point, or first executed node of visual scripting. This entry point is often triggered by an in-game occurrence. For instance, a player entering a volume can trigger an Overlap event, which can then be used to cause any functionality to happen within a Blueprint, such as opening a door, turning on a light, and so on.
-[/EXCERPT:Glossary_event]
-
+		[INCLUDE:#Glossary_event]
 	[/PARAM]
 [/OBJECT]
 
@@ -23,11 +19,7 @@ In Blueprints, an **event** is an entry point, or first executed node of visual
 EventGraph
 	[/PARAM]
 	[PARAM:Definition]
-
-[EXCERPT:Glossary_EventGraph]
-The **EventGraph** of a Blueprint contains a node graph that uses events and function calls to perform actions in response to gameplay events associated with the Blueprint. This is used to add functionality that is common to all instances of a Blueprint. This is where interactivity and dynamic responses are setup. For example, a light Blueprint could respond to a damage event by turning off its `LightComponent` and changing the material used by its mesh. This would automatically provide this behavior to all instances of the light Blueprint.
-[/EXCERPT:Glossary_EventGraph]
-
+		[INCLUDE:#Glossary_EventGraph]
 	[/PARAM]
 [/OBJECT]
 
@@ -37,14 +29,23 @@ The **EventGraph** of a Blueprint contains a node graph that uses events and fun
 execution pin
 	[/PARAM]
 	[PARAM:Definition]
-[EXCERPT:Glossary_execution_pin]
-![Blueprint Exec Pins](k2_pins_exec.png)
-
-**Execution pins** are used to connect nodes together to create a flow of execution. When an input execution pin is activated, the node is executed. Once execution of the node completes, it activates an output execution pin to continue the flow of execution. Execution pins are displayed as an outline when not wired, and solid when wired to another execution pin. Function Call nodes always have only a single input execution pin and a single output execution pin as functions only have one entry point and one exit point. Other types of nodes can have multiple input and output execution pins, allowing for different behavior dependent on which pin is activated.
-[/EXCERPT:Glossary_execution_pin]
-
+		[INCLUDE:#Glossary_execution_pin]
 	[/PARAM]
 [/OBJECT]
 
 
 [/EXCERPT:Glossary_E]
+
+[COMMENT:none]
+[EXCERPT:Glossary_event]
+In Blueprints, an **event** is an entry point, or first executed node of visual scripting. This entry point is often triggered by an in-game occurrence. For instance, a player entering a volume can trigger an Overlap event, which can then be used to cause any functionality to happen within a Blueprint, such as opening a door, turning on a light, and so on.
+[/EXCERPT:Glossary_event]
+[EXCERPT:Glossary_EventGraph]
+The **EventGraph** of a Blueprint contains a node graph that uses events and function calls to perform actions in response to gameplay events associated with the Blueprint. This is used to add functionality that is common to all instances of a Blueprint. This is where interactivity and dynamic responses are setup. For example, a light Blueprint could respond to a damage event by turning off its `LightComponent` and changing the material used by its mesh. This would automatically provide this behavior to all instances of the light Blueprint.
+[/EXCERPT:Glossary_EventGraph]
+[EXCERPT:Glossary_execution_pin]
+![Blueprint Exec Pins](k2_pins_exec.png)
+
+**Execution pins** are used to connect nodes together to create a flow of execution. When an input execution pin is activated, the node is executed. Once execution of the node completes, it activates an output execution pin to continue the flow of execution. Execution pins are displayed as an outline when not wired, and solid when wired to another execution pin. Function Call nodes always have only a single input execution pin and a single output execution pin as functions only have one entry point and one exit point. Other types of nodes can have multiple input and output execution pins, allowing for different behavior dependent on which pin is activated.
+[/EXCERPT:Glossary_execution_pin]
+[/COMMENT]
diff --git a/Engine/Documentation/Source/Shared/Glossary/F/Glossary_F.INT.udn b/Engine/Documentation/Source/Shared/Glossary/F/Glossary_F.INT.udn
--- a/Engine/Documentation/Source/Shared/Glossary/F/Glossary_F.INT.udn
+++ b/Engine/Documentation/Source/Shared/Glossary/F/Glossary_F.INT.udn
@@ -10,44 +10,41 @@ Availability:NoPublish
 FBX file
 	[/PARAM]
 	[PARAM:Definition]
-
-[EXCERPT:Glossary_FBX_file]
-**FBX** is a file format developed by Autodesk. This file format is used to import Static Meshes, Skeletal Meshes, and Animation Sequences into the Unreal Editor.
-[/EXCERPT:Glossary_FBX_file]
-
+		[INCLUDE:#Glossary_FBX_file]
 	[/PARAM]
 [/OBJECT]
 
-
-
 [OBJECT:GlossaryItem]
 	[PARAM:Term]
 Flow Control
 	[/PARAM]
 	[PARAM:Definition]
-
-[EXCERPT:Glossary_Flow_Control]
-**Flow Control** operations allow the flow of execution to be controlled explicitly in Blueprints. This control can be in the form of choosing one branch of the graph to execute based on whether some condition 
-is true, executing a specific branch multiple times, executing multiple branches in a specific order, 
-and so on. The default flow control operations include branches (if statements), loops (for and while), 
-gates, and sequences.
-[/EXCERPT:Glossary_Flow_Control]
-
+		[INCLUDE:#Glossary_Flow_Control]
 	[/PARAM]
 [/OBJECT]
 
-
 [OBJECT:GlossaryItem]
 	[PARAM:Term]
 framework
 	[/PARAM]
 	[PARAM:Definition]
-
-[EXCERPT:Glossary_framework]
-The **framework** systems form the foundation of gameplay in the Unreal Engine. The gameplay framework can be divided into three categories: gametype, player, and interface. Each of these is responsible for a particular aspect of the gameplay and each is comprised of one or more systems and classes that define the different parts and how they work together.
-[/EXCERPT:Glossary_framework]
-
+		[INCLUDE:#Glossary_framework]
 	[/PARAM]
 [/OBJECT]
 
 [/EXCERPT:Glossary_F]
+
+[COMMENT:none]
+[EXCERPT:Glossary_FBX_file]
+**FBX** is a file format developed by Autodesk. This file format is used to import Static Meshes, Skeletal Meshes, and Animation Sequences into the Unreal Editor.
+[/EXCERPT:Glossary_FBX_file]
+[EXCERPT:Glossary_Flow_Control]
+**Flow Control** operations allow the flow of execution to be controlled explicitly in Blueprints. This control can be in the form of choosing one branch of the graph to execute based on whether some condition 
+is true, executing a specific branch multiple times, executing multiple branches in a specific order, 
+and so on. The default flow control operations include branches (if statements), loops (for and while), 
+gates, and sequences.
+[/EXCERPT:Glossary_Flow_Control]
+[EXCERPT:Glossary_framework]
+The **framework** systems form the foundation of gameplay in the Unreal Engine. The gameplay framework can be divided into three categories: gametype, player, and interface. Each of these is responsible for a particular aspect of the gameplay and each is comprised of one or more systems and classes that define the different parts and how they work together.
+[/EXCERPT:Glossary_framework]
+[/COMMENT]
\ No newline at end of file
diff --git a/Engine/Documentation/Source/Shared/Glossary/G/Glossary_G.INT.udn b/Engine/Documentation/Source/Shared/Glossary/G/Glossary_G.INT.udn
--- a/Engine/Documentation/Source/Shared/Glossary/G/Glossary_G.INT.udn
+++ b/Engine/Documentation/Source/Shared/Glossary/G/Glossary_G.INT.udn
@@ -10,14 +10,7 @@ Availability:NoPublish
 GameMode
 	[/PARAM]
 	[PARAM:Definition]
-
-[EXCERPT:Glossary_GameMode]
-The **GameMode** Class is responsible for setting the rules of the game that is being played. The rules can include how players join the game, whether or not a game can be paused, and level transitions, as well as any game-specific behavior such as win conditions. 
-
-You can set the default GameMode in the **Project Settings**, but can override it on a per-Level basis. Regardless of how you choose to implement the GameMode, there is always only one GameMode present per-level. In a multiplayer game, the GameMode only exists on the server and the rules are replicated (sent) to each of the connected clients. 
-
-[/EXCERPT:Glossary_GameMode]
-
+		[INCLUDE:#Glossary_GameMode]
 	[/PARAM]
 [/OBJECT]
 
@@ -26,49 +19,53 @@ You can set the default GameMode in the **Project Settings**, but can override i
 GameState
 	[/PARAM]
 	[PARAM:Definition]
-
-[EXCERPT:Glossary_GameState]
-The **GameState** contains the information that you want replicated to every client in a game, or more simply it is 'The State of the Game' for everyone connected. 
-
-It often contains information about game score, whether a match has started or not, how many AI to spawn based upon the number of players in the world, and other game specific information. 
-
-For mutliplayer games, there is one instance of the GameState on each player's machine with the server's instance being the authoritative one (or the one that clients get their updated information from).
-[/EXCERPT:Glossary_GameState]
-
+		[INCLUDE:#Glossary_GameState]
 	[/PARAM]
 [/OBJECT]
 
-
 [OBJECT:GlossaryItem]
 	[PARAM:Term]
 Graph
 	[/PARAM]
 	[PARAM:Definition]
+		[INCLUDE:#Glossary_Graph_Mode]
 
-[EXCERPT:Glossary_Graph_Mode]
-The **Graph** is a visual representation of a particular graph of nodes as it shows 
-all of the nodes contained in the graph as well as the connections between them. It provides 
-editing capabilities for adding and removing nodes, arranging nodes, and creating links between 
-nodes. Breakpoints can also be set in the Graph to aid in [debugging](Engine/Blueprints\UserGuide\Debugging) Blueprints.
-[/EXCERPT:Glossary_Graph_Mode]
-
-[INCLUDE:Globals#bpModesRemoved]
+		[INCLUDE:Globals#bpModesRemoved]
 	[/PARAM]
 [/OBJECT]
 
-
-
 [OBJECT:GlossaryItem]
 	[PARAM:Term]
 group
 	[/PARAM]
 	[PARAM:Definition]
-
-[EXCERPT:Glossary_group]
-In Matinee, a **group** is a container for animation and visual effects for a specific Actor. In the Level Editor, a **group** is an organizational object in which selected Actors can be combined together so they can be transformed as a single unit.
-[/EXCERPT:Glossary_group]
-
+		[INCLUDE:#Glossary_group]
 	[/PARAM]
 [/OBJECT]
 
 [/EXCERPT:Glossary_G]
+
+[COMMENT:none]
+[EXCERPT:Glossary_GameMode]
+The **GameMode** Class is responsible for setting the rules of the game that is being played. The rules can include how players join the game, whether or not a game can be paused, and level transitions, as well as any game-specific behavior such as win conditions. 
+
+You can set the default GameMode in the **Project Settings**, but can override it on a per-Level basis. Regardless of how you choose to implement the GameMode, there is always only one GameMode present per-level. In a multiplayer game, the GameMode only exists on the server and the rules are replicated (sent) to each of the connected clients. 
+
+[/EXCERPT:Glossary_GameMode]
+[EXCERPT:Glossary_GameState]
+The **GameState** contains the information that you want replicated to every client in a game, or more simply it is 'The State of the Game' for everyone connected. 
+
+It often contains information about game score, whether a match has started or not, how many AI to spawn based upon the number of players in the world, and other game specific information. 
+
+For mutliplayer games, there is one instance of the GameState on each player's machine with the server's instance being the authoritative one (or the one that clients get their updated information from).
+[/EXCERPT:Glossary_GameState]
+[EXCERPT:Glossary_Graph_Mode]
+The **Graph** is a visual representation of a particular graph of nodes as it shows 
+all of the nodes contained in the graph as well as the connections between them. It provides 
+editing capabilities for adding and removing nodes, arranging nodes, and creating links between 
+nodes. Breakpoints can also be set in the Graph to aid in [debugging](Engine/Blueprints\UserGuide\Debugging) Blueprints.
+[/EXCERPT:Glossary_Graph_Mode]
+[EXCERPT:Glossary_group]
+In Matinee, a **group** is a container for animation and visual effects for a specific Actor. In the Level Editor, a **group** is an organizational object in which selected Actors can be combined together so they can be transformed as a single unit.
+[/EXCERPT:Glossary_group]
+[/COMMENT]
\ No newline at end of file
diff --git a/Engine/Documentation/Source/Shared/Glossary/K/Glossary_K.INT.udn b/Engine/Documentation/Source/Shared/Glossary/K/Glossary_K.INT.udn
--- a/Engine/Documentation/Source/Shared/Glossary/K/Glossary_K.INT.udn
+++ b/Engine/Documentation/Source/Shared/Glossary/K/Glossary_K.INT.udn
@@ -10,11 +10,7 @@ Availability:NoPublish
 key, keyframe
 	[/PARAM]
 	[PARAM:Definition]
-
-[EXCERPT:Glossary_Key]
-In animation, a **key** or **keyframe** is a record of a given value or values at a specific time. Curves are then used to interpolate between these keyframes, resulting in animated change over time.
-[/EXCERPT:Glossary_Key]
-
+		[INCLUDE:#Glossary_Key]
 	[/PARAM]
 [/OBJECT]
 
@@ -23,28 +19,29 @@ In animation, a **key** or **keyframe** is a record of a given value or values a
 keybindings
 	[/PARAM]
 	[PARAM:Definition]
-
-[EXCERPT:Glossary_keybindings]
-**Keybindings** are custom-defined keyboard shortcuts. In the Unreal Engine, you can bind specific keys or key combinations to execute specific commands or actions.
-[/EXCERPT:Glossary_keybindings]
-
+		[INCLUDE:#Glossary_keybindings]
 	[/PARAM]
 [/OBJECT]
 
-
-
 [OBJECT:GlossaryItem]
 	[PARAM:Term]
 	Kill Z
 	[/PARAM]
 	[PARAM:Definition]
-
-[EXCERPT:Glossary_Kill_Z]
-The **Kill Z** is the coordinate along the Z axis at which an Actor is automatically destroyed or players are killed. Setting a Kill Z helps prevent player characters from getting stuck and being unable to respawn when they fall off an edge.
-[/EXCERPT:Glossary_Kill_Z]
-
+		[INCLUDE:#Glossary_Kill_Z]
 	[/PARAM]
 [/OBJECT]
 
-
 [/EXCERPT:Glossary_K]
+
+[COMMENT:none]
+[EXCERPT:Glossary_Key]
+In animation, a **key** or **keyframe** is a record of a given value or values at a specific time. Curves are then used to interpolate between these keyframes, resulting in animated change over time.
+[/EXCERPT:Glossary_Key]
+[EXCERPT:Glossary_keybindings]
+**Keybindings** are custom-defined keyboard shortcuts. In the Unreal Engine, you can bind specific keys or key combinations to execute specific commands or actions.
+[/EXCERPT:Glossary_keybindings]
+[EXCERPT:Glossary_Kill_Z]
+The **Kill Z** is the coordinate along the Z axis at which an Actor is automatically destroyed or players are killed. Setting a Kill Z helps prevent player characters from getting stuck and being unable to respawn when they fall off an edge.
+[/EXCERPT:Glossary_Kill_Z]
+[/COMMENT]
\ No newline at end of file
diff --git a/Engine/Documentation/Source/Shared/Glossary/L/Glossary_L.INT.udn b/Engine/Documentation/Source/Shared/Glossary/L/Glossary_L.INT.udn
--- a/Engine/Documentation/Source/Shared/Glossary/L/Glossary_L.INT.udn
+++ b/Engine/Documentation/Source/Shared/Glossary/L/Glossary_L.INT.udn
@@ -10,18 +10,14 @@ Availability:NoPublish
 Level
 	[/PARAM]
 	[PARAM:Definition]
+		[INCLUDE:#Glossary_level]
+	[/PARAM]
+[/OBJECT]
 
-[COMMENT:none]
-
-A **level** is a defined area of gameplay. Also called a map. Levels are created, viewed, and modified mainly by placing, transforming, and editing the properties of Actors. In the Unreal Editor, each level is saved as a separate .umap file, which is separate from the project file (.uproject).
-
-[/COMMENT]
+[/EXCERPT:Glossary_L]
 
+[COMMENT:none]
 [EXCERPT:Glossary_level]
 A **Level** is a user defined area of gameplay. Levels are created, viewed, and modified mainly by placing, transforming, and editing the properties of the Actors it contains. In the Unreal Editor, each Level is saved as a separate .umap file, which is also why you will sometimes see them referred to as Maps. 
 [/EXCERPT:Glossary_level]
-
-	[/PARAM]
-[/OBJECT]
-
-[/EXCERPT:Glossary_L]
+[/COMMENT]
\ No newline at end of file
diff --git a/Engine/Documentation/Source/Shared/Glossary/M/Glossary_M.INT.udn b/Engine/Documentation/Source/Shared/Glossary/M/Glossary_M.INT.udn
--- a/Engine/Documentation/Source/Shared/Glossary/M/Glossary_M.INT.udn
+++ b/Engine/Documentation/Source/Shared/Glossary/M/Glossary_M.INT.udn
@@ -11,7 +11,76 @@
 Material
 	[/PARAM]
 	[PARAM:Definition]
+		[INCLUDE:#Glossary_Material]
+	[/PARAM]
+[/OBJECT]
 
+[OBJECT:GlossaryItem]
+	[PARAM:Term]
+Matinee
+	[/PARAM]
+	[PARAM:Definition]
+		[INCLUDE:#Glossary_Matinee]
+	[/PARAM]
+[/OBJECT]
+
+[OBJECT:GlossaryItem]
+	[PARAM:Term]
+Matinee Actor
+	[/PARAM]
+	[PARAM:Definition]
+		[INCLUDE:#Glossary_Matinee_Actor]
+	[/PARAM]
+[/OBJECT]
+
+[OBJECT:GlossaryItem]
+	[PARAM:Term]
+Migrate Tool
+	[/PARAM]
+	[PARAM:Definition]
+		[INCLUDE:#Glossary_Migrate_Tool]
+	[/PARAM]
+[/OBJECT]
+
+[OBJECT:GlossaryItem]
+	[PARAM:Term]
+Mobility
+	[/PARAM]
+	[PARAM:Definition]
+		[INCLUDE:#Glossary_Mobility]
+	[/PARAM]
+[/OBJECT]
+
+[OBJECT:GlossaryItem]
+	[PARAM:Term]
+Modes panel
+	[/PARAM]
+	[PARAM:Definition]
+		[INCLUDE:#Glossary_Modes_panel]
+	[/PARAM]
+[/OBJECT]
+
+[OBJECT:GlossaryItem]
+	[PARAM:Term]
+Movable
+	[/PARAM]
+	[PARAM:Definition]
+		[INCLUDE:#Glossary_Movable]
+	[/PARAM]
+[/OBJECT]
+
+[OBJECT:GlossaryItem]
+	[PARAM:Term]
+Move tool
+	[/PARAM]
+	[PARAM:Definition]
+		[INCLUDE:#Glossary_Move_tool]
+	[/PARAM]
+[/OBJECT]
+
+[/EXCERPT:Glossary_M]
+
+[COMMENT:none]
 [EXCERPT:Glossary_Material]
 A **Material** is an asset that can be applied to a mesh to control the visual look of the scene. At a high level, it is probably easiest to think of a Material as the "paint" that is applied to an object. But even that can be a little misleading, since a Material literally defines the type of surface from which your object appears to be made. You can define its color, how shiny it is, whether you can see through the object, and much more. 
 
@@ -19,64 +88,15 @@
 
 Unreal Engine 4 utilizes a _physically-based shading model_. This means that rather than defining a Material using arbitrary properties (such as Diffuse Color and Specular Power), you instead use properties more easily relatable to the real world. These include Base Color, Metallic, Specular, and Roughness.
 [/EXCERPT:Glossary_Material]
-
-	[/PARAM]
-[/OBJECT]
-
-
-
-[OBJECT:GlossaryItem]
-	[PARAM:Term]
-Matinee
-	[/PARAM]
-	[PARAM:Definition]
-
 [EXCERPT:Glossary_Matinee]
 The **Matinee** animation tool provides the ability to animate the properties of Actors over time, to create either dynamic gameplay or cinematic in-game sequences. The system is based on the use of specialized animation tracks in which you can place keyframes to set the values of certain properties of the Actors in the level. The Matinee Editor is similar to the non-linear editors used for video editing, making it familiar to video professionals.
 [/EXCERPT:Glossary_Matinee]
-
-	[/PARAM]
-[/OBJECT]
-
-
-
-
-[OBJECT:GlossaryItem]
-	[PARAM:Term]
-Matinee Actor
-	[/PARAM]
-	[PARAM:Definition]
-
 [EXCERPT:Glossary_Matinee_Actor]
 A **Matinee Actor** is a specific instance of a Matinee animation, created as an Actor and placed in your scene.
 [/EXCERPT:Glossary_Matinee_Actor]
-
-	[/PARAM]
-[/OBJECT]
-
-
-
-[OBJECT:GlossaryItem]
-	[PARAM:Term]
-Migrate Tool
-	[/PARAM]
-	[PARAM:Definition]
-
 [EXCERPT:Glossary_Migrate_Tool]
 The **Content Browser** contains the **Migrate Tool**, which allows assets to be copied from one project to another, along with any dependent assets. For instance, you can migrate a Material, and any Texture assets that define that Material will automatically be copied along with it. This is useful when you need to merge or fork projects, or when transitioning from a test environment into a production project.
 [/EXCERPT:Glossary_Migrate_Tool]
-
-	[/PARAM]
-[/OBJECT]
-
-
-
-[OBJECT:GlossaryItem]
-	[PARAM:Term]
-Mobility
-	[/PARAM]
-	[PARAM:Definition]
-
 [EXCERPT:Glossary_Mobility]
 **Mobility** is a property of Static Mesh and Light Actors. The Mobility property, which is set in the **Details** panel, determines whether an Actor can move during gameplay. In addition, for Light Actors, the Mobility property determines the types of light and shadows cast.  Mobility property settings include Static, Stationary (Light Actors only), and Movable.
 
@@ -84,62 +104,24 @@
 
 Actors with a Movable setting have the most versatility but the highest performance hit; Actors with a Static setting have the least versatility but the lowest performance hit.
 [/EXCERPT:Glossary_Mobility]
-
-	[/PARAM]
-[/OBJECT]
-
-
-[OBJECT:GlossaryItem]
-	[PARAM:Term]
-Modes panel
-	[/PARAM]
-	[PARAM:Definition]
-
-
 [EXCERPT:Glossary_Modes_panel]
 The **Modes** panel contains a selection of various tool modes for the Editor. These change the primary behavior of the Level Editor for a specialized task, such as placing new assets into the world, creating geometry brushes and volumes, painting on meshes, generating foliage, and sculpting landscapes.
 [/EXCERPT:Glossary_Modes_panel]
-
-	[/PARAM]
-[/OBJECT]
-
-
-[OBJECT:GlossaryItem]
-	[PARAM:Term]
-Movable
-	[/PARAM]
-	[PARAM:Definition]
-
 [EXCERPT:Glossary_Movable]
 **Movable** is one possible Mobility property settings. This setting can be applied to both Static Mesh and Light Actors.
 
 **Movable Static Mesh Actors** can be moved during gameplay, and can interact with other Actors.
 
+[INCLUDE:#Glossary_Movable_Lights]
+
+Movable Lights can change position, rotation, color, brightness, falloff, radius, and just about every other property they have.
+[/EXCERPT:Glossary_Movable]
 [EXCERPT:Glossary_Movable_Lights]
 **Movable Lights** are lights that can be moved and changed in game. Typical uses would be headlights on a car, flashlights, light fixtures that sway in the wind, and so on. Movable Lights are completely dynamic, and do not support indirect lighting at the moment.  Movable lights use dynamic shadowing for the entire scene.  As a result, shadow casting movable lights can cost 20x more than a Stationary light.
 
 Of the three light mobilities, Movable lights tend to have medium quality, highest mutability, and the highest performance cost.
 [/EXCERPT:Glossary_Movable_Lights]
-
-Movable Lights can change position, rotation, color, brightness, falloff, radius, and just about every other property they have.
-[/EXCERPT:Glossary_Movable]
-
-	[/PARAM]
-[/OBJECT]
-
-
-
-[OBJECT:GlossaryItem]
-	[PARAM:Term]
-Move tool
-	[/PARAM]
-	[PARAM:Definition]
-
 [EXCERPT:Glossary_Move_tool]
 The **Move tool** is a transformation widget that you can use to move (translate) Actors in your scene along the X, Y, and Z axes.
 [/EXCERPT:Glossary_Move_tool]
-
-	[/PARAM]
-[/OBJECT]
-
-[/EXCERPT:Glossary_M]
+[/COMMENT]
diff --git a/Engine/Documentation/Source/Shared/Glossary/N/Glossary_N.INT.udn b/Engine/Documentation/Source/Shared/Glossary/N/Glossary_N.INT.udn
--- a/Engine/Documentation/Source/Shared/Glossary/N/Glossary_N.INT.udn
+++ b/Engine/Documentation/Source/Shared/Glossary/N/Glossary_N.INT.udn
@@ -11,16 +11,14 @@ Availability:NoPublish
 node
 	[/PARAM]
 	[PARAM:Definition]
-
-[EXCERPT:Glossary_node]
-A **node** is a single block of visual scripting code, generally seen as a shape into which various data and/or execution wires are connected.
-[/EXCERPT:Glossary_node]
-
+		[INCLUDE:#Glossary_node]
 	[/PARAM]
 [/OBJECT]
 
-
-
 [/EXCERPT:Glossary_N]
 
-
+[COMMENT:none]
+[EXCERPT:Glossary_node]
+A **node** is a single block of visual scripting code, generally seen as a shape into which various data and/or execution wires are connected.
+[/EXCERPT:Glossary_node]
+[/COMMENT]
\ No newline at end of file
diff --git a/Engine/Documentation/Source/Shared/Glossary/O/Glossary_O.INT.udn b/Engine/Documentation/Source/Shared/Glossary/O/Glossary_O.INT.udn
--- a/Engine/Documentation/Source/Shared/Glossary/O/Glossary_O.INT.udn
+++ b/Engine/Documentation/Source/Shared/Glossary/O/Glossary_O.INT.udn
@@ -10,23 +10,15 @@ Availability:NoPublish
 Object
 	[/PARAM]
 	[PARAM:Definition]
-
-[COMMENT:none]
-
-**Objects** are instances of classes that inherit from the UObject class; the base class of all objects in Unreal Engine, including Actors. 
-
-[/COMMENT]
-
-[EXCERPT:Glossary_Object]
-The base building blocks in the Unreal Engine are called **Objects** and contain a lot of the essential "under the hood" functionality for your game assets. Just about everything in Unreal Engine 4 inherits (or gets some functionality) from an Object. In C++, `UObject` is the base class of all objects; it implements features such as garbage collections, metadata (`UProperty`) support for exposing variables to the Unreal Editor, and serialization for loading and saving.
-[/EXCERPT:Glossary_Object]
-
+		[INCLUDE:#Glossary_Object]
 	[/PARAM]
 [/OBJECT]
 
 
-
-
-
 [/EXCERPT:Glossary_O]
 
+[COMMENT:none]
+[EXCERPT:Glossary_Object]
+The base building blocks in the Unreal Engine are called **Objects** and contain a lot of the essential "under the hood" functionality for your game assets. Just about everything in Unreal Engine 4 inherits (or gets some functionality) from an Object. In C++, `UObject` is the base class of all objects; it implements features such as garbage collections, metadata (`UProperty`) support for exposing variables to the Unreal Editor, and serialization for loading and saving.
+[/EXCERPT:Glossary_Object]
+[/COMMENT]
\ No newline at end of file
diff --git a/Engine/Documentation/Source/Shared/Glossary/P/Glossary_P.INT.udn b/Engine/Documentation/Source/Shared/Glossary/P/Glossary_P.INT.udn
--- a/Engine/Documentation/Source/Shared/Glossary/P/Glossary_P.INT.udn
+++ b/Engine/Documentation/Source/Shared/Glossary/P/Glossary_P.INT.udn
@@ -10,104 +10,63 @@
 panel
 	[/PARAM]
 	[PARAM:Definition]
-
-[EXCERPT:Glossary_panel]
-A **panel** is a part of the UI that contains information about your level and can be moved, docked, or dragged into a separate window. Examples include the **Details** panel and the **Content Browser**.
-[/EXCERPT:Glossary_panel]
-
+		[INCLUDE:#Glossary_panel]
 	[/PARAM]
 [/OBJECT]
 
-
-
 [OBJECT:GlossaryItem]
 	[PARAM:Term]
 Pawn 
 	[/PARAM]
 	[PARAM:Definition]
-
-[EXCERPT:Glossary_pawn]
-**Pawns** are a subclass of Actor and serve as an in-game avatar or persona, for example the characters in a game. Pawns can be controlled by a player or by the game's AI, in the form of non-player characters (NPCs). 
-
-When a Pawn is controlled by a human or AI player, it is considered as _Possessed_. Conversely, when a Pawn is not controlled by a human or AI player it is considered as _Unpossessed_. 
-[/EXCERPT:Glossary_pawn]
-
+		[INCLUDE:#Glossary_pawn]
 	[/PARAM]
 [/OBJECT]
 
-
-
 [OBJECT:GlossaryItem]
 	[PARAM:Term]
 Persona
 	[/PARAM]
 	[PARAM:Definition]
+		[INCLUDE:#Glossary_Persona]
 
-[EXCERPT:Glossary_Persona]
-**Persona** is the animation editing toolset within Unreal. It is a robust system for editing Skeletons, Skeletal Meshes, Animation Blueprints, and various other animation assets. Most (if not all) of your animation work within Unreal will take place in this editor. This is a multi-purpose tool that contains panels for editing skeletal [Sockets](Engine/Content\Types\SkeletalMeshes\Sockets), previewing [Animation Sequences](Engine/Animation/Sequences), setting up animation [Blend Spaces](Engine/Animation/Blendspaces) and [Montages](Engine/Animation/AnimMontage), editing [Animation Blueprints](Engine/Animation/AnimBlueprints), and much more.
-[/EXCERPT:Glossary_Persona]
-
-You can open Persona by **double-clicking** any animation-related asset, such as Skeletal Meshes, Blend Spaces, Animation Blueprints, and so on.
-
+		You can open Persona by **double-clicking** any animation-related asset, such as Skeletal Meshes, Blend Spaces, Animation Blueprints, and so on.
 	[/PARAM]
 [/OBJECT]
 
-
-
 [OBJECT:GlossaryItem]
 	[PARAM:Term]
 PhAT
 	[/PARAM]
 	[PARAM:Definition]
-
-[EXCERPT:Glossary_PhAT]
-**PhAT**, the Physics Asset editing tool in the Unreal Editor, is used to modify physics and collision properties for Skeletal Mesh Actors.
-[/EXCERPT:Glossary_PhAT]
-
+		[INCLUDE:#Glossary_PhAT]
 	[/PARAM]
 [/OBJECT]
 
-
-
 [OBJECT:GlossaryItem]
 	[PARAM:Term]
 Physical Material
 	[/PARAM]
 	[PARAM:Definition]
-
-[EXCERPT:Glossary_Physical_Material]
-**Physical Materials** are used to define the response of a physical object when interacting dynamically with the world.  Physical Materials are fairly easy to use.  Creating one will give you a set of default values, identical to the default physical material that is applied to all physics objects.  Examples of this would be a character's dead body (ragdoll), a movable crate, and so on.
-[/EXCERPT:Glossary_Physical_Material]
-
+		[INCLUDE:#Glossary_Physical_Material]
 	[/PARAM]
 [/OBJECT]
 
-
-
 [OBJECT:GlossaryItem]
 	[PARAM:Term]
 pin
 	[/PARAM]
 	[PARAM:Definition]
-
-[EXCERPT:Glossary_pin]
-A **pin** is part of a visual scripting node. Pins connect to other nodes via wires. Execution pins must be hooked up in order for a node to execute. Data pins are color-coded according to type, and are used to input data from or output data to other nodes.
-[/EXCERPT:Glossary_pin]
-
+		[INCLUDE:#Glossary_pin]
 	[/PARAM]
 [/OBJECT]
 
-
 [OBJECT:GlossaryItem]
 	[PARAM:Term]
 pivot point
 	[/PARAM]
 	[PARAM:Definition]
-
-[EXCERPT:Glossary_pivot_point]
-The **pivot point** (or origin point) indicates the location on a Static Mesh Actor or a Skeletal Mesh Actor that has been defined in the 3D modeling program in which it was created as the object's base for alignment purposes. The pivot point can be in the center of an object, at its base, or even outside it.
-[/EXCERPT:Glossary_pivot_point]
-
+		[INCLUDE:#Glossary_pivot_point]
 	[/PARAM]
 [/OBJECT]
 
@@ -116,37 +75,71 @@
 Play In Editor (PIE)
 	[/PARAM]
 	[PARAM:Definition]
+		[INCLUDE:#Glossary_Play_In_Editor]
+	[/PARAM]
+[/OBJECT]
 
+[OBJECT:GlossaryItem]
+	[PARAM:Term]
+PlayerController
+	[/PARAM]
+	[PARAM:Definition]
+		[INCLUDE:#Glossary_PlayerController]
+	[/PARAM]
+[/OBJECT]
+
+[OBJECT:GlossaryItem]
+	[PARAM:Term]
+PlayerState
+	[/PARAM]
+	[PARAM:Definition]
+		[INCLUDE:#Glossary_PlayerState]
+	[/PARAM]
+[/OBJECT]
+
+[OBJECT:GlossaryItem]
+	[PARAM:Term]
+project
+	[/PARAM]
+	[PARAM:Definition]
+		[INCLUDE:#Glossary_project]
+	[/PARAM]
+[/OBJECT]
+
+[/EXCERPT:Glossary_P]
+
+[COMMENT:none]
+[EXCERPT:Glossary_panel]
+A **panel** is a part of the UI that contains information about your level and can be moved, docked, or dragged into a separate window. Examples include the **Details** panel and the **Content Browser**.
+[/EXCERPT:Glossary_panel]
+[EXCERPT:Glossary_pawn]
+**Pawns** are a subclass of Actor and serve as an in-game avatar or persona, for example the characters in a game. Pawns can be controlled by a player or by the game's AI, in the form of non-player characters (NPCs). 
+
+When a Pawn is controlled by a human or AI player, it is considered as _Possessed_. Conversely, when a Pawn is not controlled by a human or AI player it is considered as _Unpossessed_. 
+[/EXCERPT:Glossary_pawn]
+[EXCERPT:Glossary_Persona]
+**Persona** is the animation editing toolset within Unreal. It is a robust system for editing Skeletons, Skeletal Meshes, Animation Blueprints, and various other animation assets. Most (if not all) of your animation work within Unreal will take place in this editor. This is a multi-purpose tool that contains panels for editing skeletal [Sockets](Engine/Content\Types\SkeletalMeshes\Sockets), previewing [Animation Sequences](Engine/Animation/Sequences), setting up animation [Blend Spaces](Engine/Animation/Blendspaces) and [Montages](Engine/Animation/AnimMontage), editing [Animation Blueprints](Engine/Animation/AnimBlueprints), and much more.
+[/EXCERPT:Glossary_Persona]
+[EXCERPT:Glossary_PhAT]
+**PhAT**, the Physics Asset editing tool in the Unreal Editor, is used to modify physics and collision properties for Skeletal Mesh Actors.
+[/EXCERPT:Glossary_PhAT]
+[EXCERPT:Glossary_Physical_Material]
+**Physical Materials** are used to define the response of a physical object when interacting dynamically with the world.  Physical Materials are fairly easy to use.  Creating one will give you a set of default values, identical to the default physical material that is applied to all physics objects.  Examples of this would be a character's dead body (ragdoll), a movable crate, and so on.
+[/EXCERPT:Glossary_Physical_Material]
+[EXCERPT:Glossary_pin]
+A **pin** is part of a visual scripting node. Pins connect to other nodes via wires. Execution pins must be hooked up in order for a node to execute. Data pins are color-coded according to type, and are used to input data from or output data to other nodes.
+[/EXCERPT:Glossary_pin]
+[EXCERPT:Glossary_pivot_point]
+The **pivot point** (or origin point) indicates the location on a Static Mesh Actor or a Skeletal Mesh Actor that has been defined in the 3D modeling program in which it was created as the object's base for alignment purposes. The pivot point can be in the center of an object, at its base, or even outside it.
+[/EXCERPT:Glossary_pivot_point]
 [EXCERPT:Glossary_Play_In_Editor]
 **Play In Editor (PIE)** is a mode in which you can play your level directly in the Unreal Editor window, starting in a location you specify, for testing purposes. In this mode, you can play the current level exactly as it would play in-game, so that you can test gameplay functionality, including player controls and level events triggered by players' actions.
 [/EXCERPT:Glossary_Play_In_Editor]
-
-	[/PARAM]
-[/OBJECT]
-
-
-[OBJECT:GlossaryItem]
-	[PARAM:Term]
-PlayerController
-	[/PARAM]
-	[PARAM:Definition]
-
 [EXCERPT:Glossary_PlayerController]
 The **PlayerController** Class is used to take player input and translate that into interactions in the game and every game has at least one PlayerController in it. A PlayerController often possesses a Pawn or Character as a representation of the player in a game. 
 
 The PlayerController is also the primary network interaction point for multiplayer games. During multiplayer play, the server has one instance of a PlayerController for every player in the game since it must be able to make network function calls to each player. Each client only has the PlayerController that corresponds to their player and can only use their PlayerController to communicate with the server.
 [/EXCERPT:Glossary_PlayerController]
-
-	[/PARAM]
-[/OBJECT]
-
-
-[OBJECT:GlossaryItem]
-	[PARAM:Term]
-PlayerState
-	[/PARAM]
-	[PARAM:Definition]
-
 [EXCERPT:Glossary_PlayerState]
 A **PlayerState** is the state of a participant in the game, such as a human player or a bot that is simulating a player. Non-player AI that exists as part of the game world would not have a PlayerState. 
 
@@ -155,18 +148,6 @@
 For multiplayer games, PlayerStates for all players exist on all machines (unlike PlayerControllers) and can replicate data from the server to the client to keep things in sync.
 
 [/EXCERPT:Glossary_PlayerState]
-
-	[/PARAM]
-[/OBJECT]
-
-
-
-[OBJECT:GlossaryItem]
-	[PARAM:Term]
-project
-	[/PARAM]
-	[PARAM:Definition]
-
 [EXCERPT:Glossary_project]
 A **Project** is a self-contained unit that holds all the content and code that make up an individual game and coincides with a set of directories on your disk. For example, in the image below the Hierarchy Tree of the **Content Browser** contains the same directory structure found inside your Project folder on your disk. 
 
@@ -208,8 +189,4 @@
 
 You can create any number of different Projects which can all be maintained and developed in parallel. Both the Engine (and Editor) can easily switch between them which will allow you to work on multiple games at once or have several test projects in addition to your main game Project.
 [/EXCERPT:Glossary_project]
-
-	[/PARAM]
-[/OBJECT]
-
-[/EXCERPT:Glossary_P]
+[/COMMENT]
diff --git a/Engine/Documentation/Source/Shared/Glossary/R/Glossary_R.INT.udn b/Engine/Documentation/Source/Shared/Glossary/R/Glossary_R.INT.udn
--- a/Engine/Documentation/Source/Shared/Glossary/R/Glossary_R.INT.udn
+++ b/Engine/Documentation/Source/Shared/Glossary/R/Glossary_R.INT.udn
@@ -10,12 +10,14 @@ Availability:NoPublish
 Rotate tool 
 	[/PARAM]
 	[PARAM:Definition]
-
-[EXCERPT:Glossary_Rotate_tool]
-The **Rotate tool** is a transformation widget that you can use to rotate Actors in your scene along the X, Y, and Z axes.
-[/EXCERPT:Glossary_Rotate_tool]
-
+		[INCLUDE:#Glossary_Rotate_tool]
 	[/PARAM]
 [/OBJECT]
 
 [/EXCERPT:Glossary_R]
+
+[COMMENT:none]
+[EXCERPT:Glossary_Rotate_tool]
+The **Rotate tool** is a transformation widget that you can use to rotate Actors in your scene along the X, Y, and Z axes.
+[/EXCERPT:Glossary_Rotate_tool]
+[/COMMENT]
\ No newline at end of file
diff --git a/Engine/Documentation/Source/Shared/Glossary/S/Glossary_S.INT.udn b/Engine/Documentation/Source/Shared/Glossary/S/Glossary_S.INT.udn
--- a/Engine/Documentation/Source/Shared/Glossary/S/Glossary_S.INT.udn
+++ b/Engine/Documentation/Source/Shared/Glossary/S/Glossary_S.INT.udn
@@ -10,11 +10,7 @@
 safe frame
 	[/PARAM]
 	[PARAM:Definition]
-
-[EXCERPT:Glossary_safe_area]
-In Matinee, the **safe frame** is the view, defined by the viewpoint camera's field of view, that will be seen during your cinematic.
-[/EXCERPT:Glossary_safe_area]
-
+		[INCLUDE:#Glossary_safe_area]
 	[/PARAM]
 [/OBJECT]
 
@@ -23,11 +19,7 @@
 Scale tool
 	[/PARAM]
 	[PARAM:Definition]
-
-[EXCERPT:Glossary_Scale_tool]
-The **Scale tool** is a transformation widget that you can use to scale Actors in your scene along the X, Y, and Z axes.
-[/EXCERPT:Glossary_Scale_tool]
-
+		[INCLUDE:#Glossary_Scale_tool]
 	[/PARAM]
 [/OBJECT]
 
@@ -36,14 +28,7 @@
 Simulate In Editor (SIE)
 	[/PARAM]
 	[PARAM:Definition]
-
-[EXCERPT:Glossary_Simulate_In_Editor]
-Using the **Simulate** button begins a **Simulate In Editor** session in the currently active viewport.  While simulating, gameplay begins, including the execution of _Blueprints_ and C++ code that do not rely on a player's interaction with the game. 
-While simulating, you have full access to the Editor's tools, so you can modify the scene and its contents, or even place new Actors.   You can also select and inspect _Pawns_ controlled by AI right as they are performing actions, 
-and quickly debug and tweak gameplay behaviors. However, because you are not using a PlayerController while simulating, you cannot enter game controls. 
- You can save certain changes made in a _Simulate In Editor_ session using [Keep Simulation Changes](Engine/UI/LevelEditor\InEditorTesting\KeepSimulationChanges).
-[/EXCERPT:Glossary_Simulate_In_Editor]
-
+		[INCLUDE:#Glossary_Simulate_In_Editor]
 	[/PARAM]
 [/OBJECT]
 
@@ -52,11 +37,7 @@
 Skeletal Mesh Actor
 	[/PARAM]
 	[PARAM:Definition]
-
-[EXCERPT:Glossary_Skeletal_Mesh_Actor]
-A **Skeletal Mesh Actor** is a 3D model asset that can be animated in an external 3D modeling program and then imported into the Unreal Editor. A Skeletal Mesh Actor consists of two parts: the mesh itself and a hierarchical skeleton whose bones can be used to create animation by deforming the mesh according to rules defined in the 3D modeling program. You can import Skeletal Mesh Actors as assets in the **Content Browser**.
-[/EXCERPT:Glossary_Skeletal_Mesh_Actor]
-
+		[INCLUDE:#Glossary_Skeletal_Mesh_Actor]
 	[/PARAM]
 [/OBJECT]
 
@@ -65,11 +46,7 @@
 socket
 	[/PARAM]
 	[PARAM:Definition]
-
-[EXCERPT:Glossary_socket]
-Commonly in games, you will want to attach an object to the bone of a character. This might be a weapon attached to a hand, or a hat attached to the head. For this purpose, Unreal Engine allows you to create **Sockets** in its animation toolset - [Persona](Engine/Animation/Persona) - which are offset from a bone within the Skeletal Mesh. Sockets can then be translated, rotated, and scaled relative to the bone. Static Meshes and/or Skeletal Meshes can also be previewed attached to sockets. This allows content creators to easily set up sockets for a Skeletal Mesh, and then tell the programmer the names of the sockets to attach objects to.
-[/EXCERPT:Glossary_socket]
-
+		[INCLUDE:#Glossary_socket]
 	[/PARAM]
 [/OBJECT]
 
@@ -78,74 +55,86 @@
 sound cue
 	[/PARAM]
 	[PARAM:Definition]
+		[INCLUDE:#Glossary_sound_cue]
+	[/PARAM]
+[/OBJECT]
 
+[OBJECT:GlossaryItem]
+	[PARAM:Term]
+Static
+	[/PARAM]
+	[PARAM:Definition]
+		[INCLUDE:#Glossary_Static]
+	[/PARAM]
+[/OBJECT]
+
+[OBJECT:GlossaryItem]
+	[PARAM:Term]
+Static Mesh
+	[/PARAM]
+	[PARAM:Definition]
+		INCLUDE:#Glossary_Static_Mesh
+	[/PARAM]
+[/OBJECT]
+
+[OBJECT:GlossaryItem]
+	[PARAM:Term]
+Stationary
+	[/PARAM]
+	[PARAM:Definition]
+		[INCLUDE:#Glossary_Stationary]
+	[/PARAM]
+[/OBJECT]
+
+[/EXCERPT:Glossary_S]
+
+[COMMENT:none]
+[EXCERPT:Glossary_safe_area]
+In Matinee, the **safe frame** is the view, defined by the viewpoint camera's field of view, that will be seen during your cinematic.
+[/EXCERPT:Glossary_safe_area]
+[EXCERPT:Glossary_Scale_tool]
+The **Scale tool** is a transformation widget that you can use to scale Actors in your scene along the X, Y, and Z axes.
+[/EXCERPT:Glossary_Scale_tool]
+[EXCERPT:Glossary_Simulate_In_Editor]
+Using the **Simulate** button begins a **Simulate In Editor** session in the currently active viewport.  While simulating, gameplay begins, including the execution of _Blueprints_ and C++ code that do not rely on a player's interaction with the game. 
+While simulating, you have full access to the Editor's tools, so you can modify the scene and its contents, or even place new Actors.   You can also select and inspect _Pawns_ controlled by AI right as they are performing actions, 
+and quickly debug and tweak gameplay behaviors. However, because you are not using a PlayerController while simulating, you cannot enter game controls. 
+ You can save certain changes made in a _Simulate In Editor_ session using [Keep Simulation Changes](Engine/UI/LevelEditor\InEditorTesting\KeepSimulationChanges).
+[/EXCERPT:Glossary_Simulate_In_Editor]
+[EXCERPT:Glossary_Skeletal_Mesh_Actor]
+A **Skeletal Mesh Actor** is a 3D model asset that can be animated in an external 3D modeling program and then imported into the Unreal Editor. A Skeletal Mesh Actor consists of two parts: the mesh itself and a hierarchical skeleton whose bones can be used to create animation by deforming the mesh according to rules defined in the 3D modeling program. You can import Skeletal Mesh Actors as assets in the **Content Browser**.
+[/EXCERPT:Glossary_Skeletal_Mesh_Actor]
+[EXCERPT:Glossary_socket]
+Commonly in games, you will want to attach an object to the bone of a character. This might be a weapon attached to a hand, or a hat attached to the head. For this purpose, Unreal Engine allows you to create **Sockets** in its animation toolset - [Persona](Engine/Animation/Persona) - which are offset from a bone within the Skeletal Mesh. Sockets can then be translated, rotated, and scaled relative to the bone. Static Meshes and/or Skeletal Meshes can also be previewed attached to sockets. This allows content creators to easily set up sockets for a Skeletal Mesh, and then tell the programmer the names of the sockets to attach objects to.
+[/EXCERPT:Glossary_socket]
 [EXCERPT:Glossary_sound_cue]
 A **sound cue** is an audio asset that contains a reference to an audio file created in an external program that has been imported into a project. You can modify the audio output based on the imported audio file in the Sound Cue Editor.
 [/EXCERPT:Glossary_sound_cue]
-
-	[/PARAM]
-[/OBJECT]
-
-
-
-[OBJECT:GlossaryItem]
-	[PARAM:Term]
-Static
-	[/PARAM]
-	[PARAM:Definition]
-
 [EXCERPT:Glossary_Static]
 **Static** is one of three possible Mobility property settings. This setting can be applied to both Static Mesh and Light Components.
 
+[INCLUDE:#Glossary_Static_Lights]
+
+[/EXCERPT:Glossary_Static]
 [EXCERPT:Glossary_Static_Lights]
 **Static Lights** are lights that cannot be changed or moved in any way at runtime.  They are calculated only within Lightmaps, and once processed, have no further impact on performance.  Movable objects cannot integrate with static lights, so the usefulness of static lights is limited.
 
 Of the three light mobilities, Static lights tend to have medium quality, lowest mutability, and the lowest performance cost.
 [/EXCERPT:Glossary_Static_Lights]
-
-[/EXCERPT:Glossary_Static]
-
-
-	[/PARAM]
-[/OBJECT]
-
-
-
-[OBJECT:GlossaryItem]
-	[PARAM:Term]
-Static Mesh
-	[/PARAM]
-	[PARAM:Definition]
-
 [EXCERPT:Glossary_Static_Mesh]
 A **Static Mesh** is a piece of geometry that consists of a set of polygons that can be cached in video memory and rendered by the graphics card. This allows them to be rendered efficiently, meaning they can be much more complex than other types of geometry such as Brushes. Since they are cached in video memory, Static Meshes can be translated, rotated, and scaled, but they cannot have their vertices animated in any way.
 
 Static Meshes are the basic unit used to create world geometry for levels created in the Unreal Engine. These are 3D models created in external modeling applications (such as 3dsMax, Maya, Softimage, and so on) that are imported into the Unreal Editor through the **Content Browser**, saved in packages, and then used in various ways to create renderable elements. The vast majority of any map in a game made with Unreal will consist of Static Meshes, generally in the form of Static Mesh Actors. Other uses of Static Meshes are for creating movers such as doors or lifts, rigid body physics objects, foliage and terrain decorations, procedurally created buildings, game objectives, and many more visual elements.
 [/EXCERPT:Glossary_Static_Mesh]
-
-	[/PARAM]
-[/OBJECT]
-
-
-
-[OBJECT:GlossaryItem]
-	[PARAM:Term]
-Stationary
-	[/PARAM]
-	[PARAM:Definition]
-
 [EXCERPT:Glossary_Stationary]
 **Stationary** is one of three possible Mobility property settings. This setting can be applied only to Light Components.
 
+[INCLUDE:#Glossary_Stationary_Lights]
+
+[/EXCERPT:Glossary_Stationary]
 [EXCERPT:Glossary_Stationary_Lights]
 **Stationary Lights** are lights that are intended to stay in one position, but are able to change in other ways, such as their brightness and color. This is the primary way in which they differ from Static Lights, which cannot change in any way during gameplay. However, it should be noted that runtime changes to brightness only affect the direct lighting. Indirect (bounced) lighting, since it is pre-calculated by Lightmass, will not change.
 
 Of the three light mobilities, Stationary lights tend to have the highest quality, medium mutability, and medium performance cost.
 [/EXCERPT:Glossary_Stationary_Lights]
-
-[/EXCERPT:Glossary_Stationary]
-
-	[/PARAM]
-[/OBJECT]
-
-[/EXCERPT:Glossary_S]
+[/COMMENT]
diff --git a/Engine/Documentation/Source/Shared/Glossary/T/Glossary_T.INT.udn b/Engine/Documentation/Source/Shared/Glossary/T/Glossary_T.INT.udn
--- a/Engine/Documentation/Source/Shared/Glossary/T/Glossary_T.INT.udn
+++ b/Engine/Documentation/Source/Shared/Glossary/T/Glossary_T.INT.udn
@@ -10,76 +10,65 @@ Availability:NoPublish
 tangent handle
 	[/PARAM]
 	[PARAM:Definition]
-
-[EXCERPT:Glossary_tangent_handles]
-A **tangent handle** is a tool used in keyframe editing to adjust the shape and curve of the lines connecting two keyframes. You can find tangent handles attached to the keys in Curve Editors.
-[/EXCERPT:Glossary_tangent_handles]
-
+		[INCLUDE:#Glossary_tangent_handles]
 	[/PARAM]
 [/OBJECT]
 
-
-
 [OBJECT:GlossaryItem]
 	[PARAM:Term]
 Texture
 	[/PARAM]
 	[PARAM:Definition]
-
-[EXCERPT:Glossary_Texture]
-**Textures** are images that are used in Materials. They are mapped to the surfaces the Material is applied to. Either Textures are applied directly - for example, for Base Color textures - or the values of the Texture's pixels (or _texels_) are used within the Material as masks or for other calculations. In some instances, Textures may also be used directly, outside of materials, such as for drawing to the HUD. For the most part, Textures are created externally within an image-editing application, such as Photoshop, and then imported into Unreal Editor through the **Content Browser**. However, some Textures are generated within Unreal, such as Render Textures. These generally take some information from the scene and render it to a Texture to be used elsewhere.
-[/EXCERPT:Glossary_Texture]
-
+		[INCLUDE:#Glossary_Texture]
 	[/PARAM]
 [/OBJECT]
 
-
-
 [OBJECT:GlossaryItem]
 	[PARAM:Term]
 time cursor
 	[/PARAM]
 	[PARAM:Definition]
-
-[EXCERPT:Glossary_time_cursor]
-The **time cursor** is the black vertical bar at the bottom of the timeline that indicates the position in time of the animation's playback. You can drag the time cursor back and forth to preview the animation in the 3D viewport.
-[/EXCERPT:Glossary_time_cursor]
-
+		[INCLUDE:#Glossary_time_cursor]
 	[/PARAM]
 [/OBJECT]
 
-
-
 [OBJECT:GlossaryItem]
 	[PARAM:Term]
 timeline
 	[/PARAM]
 	[PARAM:Definition]
-
-[EXCERPT:Glossary_timeline]
-A **timeline** is a visual representation of time in Matinee, which you can use to set and modify keys for keyframe animation, and to preview your cinematic. Different from the Timeline node in Blueprints.
-[/EXCERPT:Glossary_timeline]
-
+		[INCLUDE:#Glossary_timeline]
 	[/PARAM]
 [/OBJECT]
 
-
-
 [OBJECT:GlossaryItem]
 	[PARAM:Term]
 Timeline node
 	[/PARAM]
 	[PARAM:Definition]
+		[INCLUDE:#Glossary_Timeline_node]
+	[/PARAM]
+[/OBJECT]
+
+[/EXCERPT:Glossary_T]
 
+[COMMENT:none]
+[EXCERPT:Glossary_tangent_handles]
+A **tangent handle** is a tool used in keyframe editing to adjust the shape and curve of the lines connecting two keyframes. You can find tangent handles attached to the keys in Curve Editors.
+[/EXCERPT:Glossary_tangent_handles]
+[EXCERPT:Glossary_Texture]
+**Textures** are images that are used in Materials. They are mapped to the surfaces the Material is applied to. Either Textures are applied directly - for example, for Base Color textures - or the values of the Texture's pixels (or _texels_) are used within the Material as masks or for other calculations. In some instances, Textures may also be used directly, outside of materials, such as for drawing to the HUD. For the most part, Textures are created externally within an image-editing application, such as Photoshop, and then imported into Unreal Editor through the **Content Browser**. However, some Textures are generated within Unreal, such as Render Textures. These generally take some information from the scene and render it to a Texture to be used elsewhere.
+[/EXCERPT:Glossary_Texture]
+[EXCERPT:Glossary_time_cursor]
+The **time cursor** is the black vertical bar at the bottom of the timeline that indicates the position in time of the animation's playback. You can drag the time cursor back and forth to preview the animation in the 3D viewport.
+[/EXCERPT:Glossary_time_cursor]
+[EXCERPT:Glossary_timeline]
+A **timeline** is a visual representation of time in Matinee, which you can use to set and modify keys for keyframe animation, and to preview your cinematic. Different from the Timeline node in Blueprints.
+[/EXCERPT:Glossary_timeline]
 [EXCERPT:Glossary_Timeline_node]
 **Timeline nodes** are special nodes within Blueprints that allow for simple time-based animation to be quickly designed and played back based on in-game events. Timelines are somewhat like simple Matinee sequences 
 in that they allow simple values to be animated and events to be fired off over time. These can be edited 
 directly inside the Blueprint editor by **Double-clicking** on the Timeline in the Graph tab or in the My Blueprint tab. They are specifically built for handling simple, non-cinematic tasks such as opening doors, altering 
 lights, or performing other time-centric manipulations to Actors within a scene.
 [/EXCERPT:Glossary_Timeline_node]
-
-	[/PARAM]
-[/OBJECT]
-
-
-[/EXCERPT:Glossary_T]
+[/COMMENT]
\ No newline at end of file
diff --git a/Engine/Documentation/Source/Shared/Glossary/U/Glossary_U.INT.udn b/Engine/Documentation/Source/Shared/Glossary/U/Glossary_U.INT.udn
--- a/Engine/Documentation/Source/Shared/Glossary/U/Glossary_U.INT.udn
+++ b/Engine/Documentation/Source/Shared/Glossary/U/Glossary_U.INT.udn
@@ -10,15 +10,17 @@ Availability:NoPublish
 Unreal Frontend
 	[/PARAM]
 	[PARAM:Definition]
+		[INCLUDE:#Glossary_Unreal_FrontEnd]
+	[/PARAM]
+[/OBJECT]
+
 
+[/EXCERPT:Glossary_U]
+
+[COMMENT:none]
 [EXCERPT:Glossary_Unreal_FrontEnd]
 **UnrealFrontend** (UFE) is a tool intended to simplify and speed up daily video game development and testing 
 tasks, such as preparing game builds, deploying them to a device, and launching them. UFE is designed to 
 be the central interface for all game deployment, profiling, and testing tasks.
 [/EXCERPT:Glossary_Unreal_FrontEnd]
-
-	[/PARAM]
-[/OBJECT]
-
-
-[/EXCERPT:Glossary_U]
+[/COMMENT]
\ No newline at end of file
diff --git a/Engine/Documentation/Source/Shared/Glossary/V/Glossary_V.INT.udn b/Engine/Documentation/Source/Shared/Glossary/V/Glossary_V.INT.udn
--- a/Engine/Documentation/Source/Shared/Glossary/V/Glossary_V.INT.udn
+++ b/Engine/Documentation/Source/Shared/Glossary/V/Glossary_V.INT.udn
@@ -10,14 +10,15 @@ Availability:NoPublish
 viewport
 	[/PARAM]
 	[PARAM:Definition]
-
-[EXCERPT:Glossary_viewport]
-The **viewport** is the visual representation of your level in the Level Editor. While working on your level, you can use a 3D viewport, 2D viewports, or a combination. Viewports contain a variety of tools to help you see exactly the data you need.
-[/EXCERPT:Glossary_viewport]
-
-
+		[INCLUDE:#Glossary_viewport]
 	[/PARAM]
 [/OBJECT]
 
 
 [/EXCERPT:Glossary_V]
+
+[COMMENT:none]
+[EXCERPT:Glossary_viewport]
+The **viewport** is the visual representation of your level in the Level Editor. While working on your level, you can use a 3D viewport, 2D viewports, or a combination. Viewports contain a variety of tools to help you see exactly the data you need.
+[/EXCERPT:Glossary_viewport]
+[/COMMENT]
\ No newline at end of file
diff --git a/Engine/Documentation/Source/Shared/Glossary/W/Glossary_W.INT.udn b/Engine/Documentation/Source/Shared/Glossary/W/Glossary_W.INT.udn
--- a/Engine/Documentation/Source/Shared/Glossary/W/Glossary_W.INT.udn
+++ b/Engine/Documentation/Source/Shared/Glossary/W/Glossary_W.INT.udn
@@ -10,12 +10,7 @@ Availability:NoPublish
 widget
 	[/PARAM]
 	[PARAM:Definition]
-
-
-[EXCERPT:Glossary_widget]
-A **widget** is a visual tool used in the Level Editor. Widgets typically have a visual representation in the viewport. For example, the Move tool has red, green, and blue arrows indicating the X, Y, and Z axes.
-[/EXCERPT:Glossary_widget]
-
+		[INCLUDE:#Glossary_widget]
 	[/PARAM]
 [/OBJECT]
 
@@ -24,11 +19,7 @@ A **widget** is a visual tool used in the Level Editor. Widgets typically have a
 wire
 	[/PARAM]
 	[PARAM:Definition]
-
-[EXCERPT:Glossary_wire]
-A **wire** is a line connecting two nodes in the graph section of a visual script, such as a Blueprint. You create wires by dragging from the pins on nodes. You can drag a wire from a pin on one node to a pin on another node, or you can drag from a pin to an empty space, and use the context menu that is created to select another node to add to your graph. Like pins, wires are color-coded.
-[/EXCERPT:Glossary_wire]
-
+		[INCLUDE:#Glossary_wire]
 	[/PARAM]
 [/OBJECT]
 
@@ -37,18 +28,23 @@ A **wire** is a line connecting two nodes in the graph section of a visual scrip
 World
 	[/PARAM]
 	[PARAM:Definition]
-
-[EXCERPT:Glossary_world]
-A **World** contains a list of Levels that are loaded. It handles the streaming of Levels and the spawning (creation) of dynamic Actors. 
-
-Direct interaction with a World is not necessary, but it does help provide a specific reference point within the game structure (ie: mentioning World directly means you are not speaking about Levels, Maps or the game). 
-
-[/EXCERPT:Glossary_world]
-
+		[INCLUDE:#Glossary_world]
 	[/PARAM]
 [/OBJECT]
 
 [/EXCERPT:Glossary_W]
 
+[COMMENT:none]
+[EXCERPT:Glossary_widget]
+A **widget** is a visual tool used in the Level Editor. Widgets typically have a visual representation in the viewport. For example, the Move tool has red, green, and blue arrows indicating the X, Y, and Z axes.
+[/EXCERPT:Glossary_widget]
+[EXCERPT:Glossary_wire]
+A **wire** is a line connecting two nodes in the graph section of a visual script, such as a Blueprint. You create wires by dragging from the pins on nodes. You can drag a wire from a pin on one node to a pin on another node, or you can drag from a pin to an empty space, and use the context menu that is created to select another node to add to your graph. Like pins, wires are color-coded.
+[/EXCERPT:Glossary_wire]
+[EXCERPT:Glossary_world]
+A **World** contains a list of Levels that are loaded. It handles the streaming of Levels and the spawning (creation) of dynamic Actors. 
 
+Direct interaction with a World is not necessary, but it does help provide a specific reference point within the game structure (ie: mentioning World directly means you are not speaking about Levels, Maps or the game). 
 
+[/EXCERPT:Glossary_world]
+[/COMMENT]
\ No newline at end of file
diff --git a/Engine/Documentation/Source/Shared/GraphNodes/Material/Material.INT.udn b/Engine/Documentation/Source/Shared/GraphNodes/Material/Material.INT.udn
--- a/Engine/Documentation/Source/Shared/GraphNodes/Material/Material.INT.udn
+++ b/Engine/Documentation/Source/Shared/GraphNodes/Material/Material.INT.udn
@@ -205,7 +205,7 @@ Engine/Rendering/Materials/ExpressionReference/Coordinates#lightmapuvs
 [/EXCERPT:UMaterialExpressionLightmapUVs]
 
 [EXCERPT:UMaterialExpressionObjectOrientation]
-The ObjectOrientation expression outputs the world-space up vector of the object. In other words, this is the direction the local positive Z-axis of the object the material is applied to is pointing. 
+The ObjectOrientation expression outputs the world-space up vector of the object. In other words, this is the direction the local positive Z-axis of the object the material is applied to is pointing. For deferred decal, material domain returns decal projection direction (X-axis). 
 
 [VAR:ToolTipFullLink]
 Engine/Rendering/Materials/ExpressionReference/Coordinates#objectorientation
diff --git a/Engine/Documentation/Source/Shared/Tutorials/inEditorTutorialCamera/inEditorTutorialCamera.INT.udn b/Engine/Documentation/Source/Shared/Tutorials/inEditorTutorialCamera/inEditorTutorialCamera.INT.udn
--- a/Engine/Documentation/Source/Shared/Tutorials/inEditorTutorialCamera/inEditorTutorialCamera.INT.udn
+++ b/Engine/Documentation/Source/Shared/Tutorials/inEditorTutorialCamera/inEditorTutorialCamera.INT.udn
@@ -9,15 +9,17 @@ Crumbs:
 Viewport Introduction
 [/VAR]
 
-[EXCERPT:ViewportIntro]
-The **Viewports** are your window into the worlds you create in Unreal. They can be navigated just as you would in a game, or can be used in a more schematic design sense as you would for an architectural blueprint. The Unreal Editor viewports contain a variety of tools and visualizers to help you see exactly the data you need. 
-[/EXCERPT:ViewportIntro]
+[INCLUDE:#ViewportIntro]
 
 ![](viewports.png)
 
 [/EXCERPT:Stage3]
 
-
+[COMMENT:none]
+[EXCERPT:ViewportIntro]
+The **Viewports** are your window into the worlds you create in Unreal. They can be navigated just as you would in a game, or can be used in a more schematic design sense as you would for an architectural blueprint. The Unreal Editor viewports contain a variety of tools and visualizers to help you see exactly the data you need. 
+[/EXCERPT:ViewportIntro]
+[/COMMENT]
 
 
 [EXCERPT:Stage3.1]
diff --git a/Engine/Documentation/Source/Support/Builds/ReleaseNotes/2016/4_11/ReleaseNotes_4_11.INT.udn b/Engine/Documentation/Source/Support/Builds/ReleaseNotes/2016/4_11/ReleaseNotes_4_11.INT.udn
--- a/Engine/Documentation/Source/Support/Builds/ReleaseNotes/2016/4_11/ReleaseNotes_4_11.INT.udn
+++ b/Engine/Documentation/Source/Support/Builds/ReleaseNotes/2016/4_11/ReleaseNotes_4_11.INT.udn
@@ -32,11 +32,11 @@ Thousands of small optimizations throughout the engine have resulted in a big in
 
 **Faster garbage collection**.  We now support garbage collection "clusters", which allows the engine to treat groups of objects as a single unit, drastically reducing the number of objects that need to be considered. Currently, only subobjects for Materials and Particle Systems are clustered. Additionally, the mark and destroy phases are more cache-coherent, resulting in a 9x reduction in time, and memory churn has been reduced during reachability analysis.
 
-**Multi-threaded animation.  **Animation Graph updates can now run on worker threads allowing the number of animated characters to scale with the number of cores. Check out the Upgrade Notes as we've deprecated many animation related APIs and there are limitations on which animation graphs can run in worker threads.
+**Multi-threaded animation.** Animation Graph updates can now run on worker threads allowing the number of animated characters to scale with the number of cores. Check out the Upgrade Notes as we've deprecated many animation related APIs and there are limitations on which animation graphs can run in worker threads.
 
-**Instant animation variable access.  ** We've added a 'fast-path' for variable access in the Animation Graph update. This allows us to simply copy parameters internally rather than executing Blueprint code. The compiler can currently optimize the following constructs: member variables, negated boolean member variables, and members of a nested structure.
+**Instant animation variable access.** We've added a 'fast-path' for variable access in the Animation Graph update. This allows us to simply copy parameters internally rather than executing Blueprint code. The compiler can currently optimize the following constructs: member variables, negated boolean member variables, and members of a nested structure.
 
-**Additive animation 'baking'.  **We now have an option to turn on **baked additive animations**. This makes using **additive animations **roughly **3x faster.** The work involved in calculating the delta pose for an additive animation is done at cook time rather than run time. This saves not only the calculation work involved in creating the additive deltas but also the memory accesses and allocations involved with decompressing the base animation which is no longer needed at runtime. This feature increases cook times, so you'll need to enable it by setting the cvar "a.UseBakedAdditiveAnimations" to 1. Future versions of the engine will rework animation cooking to avoid this cost and permanently enable this optimization.
+**Additive animation 'baking'.** We now have an option to turn on **baked additive animations**. This makes using **additive animations **roughly **3x faster.** The work involved in calculating the delta pose for an additive animation is done at cook time rather than run time. This saves not only the calculation work involved in creating the additive deltas but also the memory accesses and allocations involved with decompressing the base animation which is no longer needed at runtime. This feature increases cook times, so you'll need to enable it by setting the cvar "a.UseBakedAdditiveAnimations" to 1. Future versions of the engine will rework animation cooking to avoid this cost and permanently enable this optimization.
 
 ## New: Realistic Hair Shading
 
@@ -786,7 +786,7 @@ You can turn on Sequencer by opening the Plugins panel and enabling the "Level S
 
 ![image alt text](image_52.png)
 
-* **New: Per-Animation ****Compression**** Settings**
+* **New: Per-Animation **Compression** Settings
 
     * You can now set **animation compression **on a **per-animation** basis. 
 
diff --git a/Engine/Documentation/Source/Support/Builds/ReleaseNotes/2016/4_14/ReleaseNotes_4_14.INT.udn b/Engine/Documentation/Source/Support/Builds/ReleaseNotes/2016/4_14/ReleaseNotes_4_14.INT.udn
--- a/Engine/Documentation/Source/Support/Builds/ReleaseNotes/2016/4_14/ReleaseNotes_4_14.INT.udn
+++ b/Engine/Documentation/Source/Support/Builds/ReleaseNotes/2016/4_14/ReleaseNotes_4_14.INT.udn
@@ -263,7 +263,7 @@ When you use Reflection Captures, the engine mixes the indirect specular from th
 
 Mixing works well for rough surfaces, but for smooth surfaces the reflections from Reflection Captures no longer match reflections from other methods, like Screen Space Reflections or Planar Reflections.  
 
-**Lightmap mixing is no longer done on very smooth surfaces**. ** **A surface with roughness .3 will get full lightmap mixing, fading out to no lightmap mixing by Roughness .1 and below.  This allows Reflection Captures and SSR to match much better and it's harder to spot transitions.  
+**Lightmap mixing is no longer done on very smooth surfaces.**  A surface with roughness .3 will get full lightmap mixing, fading out to no lightmap mixing by Roughness .1 and below.  This allows Reflection Captures and SSR to match much better and it's harder to spot transitions.  
 
 The below shot shows mirror surface reflections before and after.  Note the difference in the reflection of the wall between SSR and reflection captures.  The artifact is especially noticeable in motion, because it will move with your camera due to SSR limitations.
 
diff --git a/Engine/Documentation/Source/Support/Builds/ReleaseNotes/4_15/ReleaseNotes_4_15.INT.udn b/Engine/Documentation/Source/Support/Builds/ReleaseNotes/4_15/ReleaseNotes_4_15.INT.udn
--- a/Engine/Documentation/Source/Support/Builds/ReleaseNotes/4_15/ReleaseNotes_4_15.INT.udn
+++ b/Engine/Documentation/Source/Support/Builds/ReleaseNotes/4_15/ReleaseNotes_4_15.INT.udn
@@ -317,7 +317,7 @@ Three new viewmodes are available that make it much easier to ensure the texture
 
 * **Primitive Distance Accuracy** - This shows the accuracy of the distance computed by the streamer as compared to the distance computed by the GPU. The streamer computes the distance from the viewpoint to the texture instance axis-aligned bounding boxes (AABB). This bounding box is computed in the build process by summing up the component lod-section AABBs that are using the texture.
 
-* **Mesh UV Densities Accuracy** -** **This view mode shows the effect of the UV densities used by the streamer for each mesh. If the UV densities are wrong, then the textures on those meshes risk being under-streamed (blurry) or over-streamed (wasting memory).
+* **Mesh UV Densities Accuracy** - This view mode shows the effect of the UV densities used by the streamer for each mesh. If the UV densities are wrong, then the textures on those meshes risk being under-streamed (blurry) or over-streamed (wasting memory).
 
 * **Material Texture Scales Accuracy** - Some textures used in a material are sampled using a scaled value of one of the mesh UV channels. The texture streaming build tries to compute which texture coordinate is used and what scale is applied to each sampled texture. This could fail for many reasons, in which case the streamer fallbacks on assuming the texture was sampled with texcoord 0 and a scale of 1.
 
