Copying //UE4/Dev-Documentation to //UE4/Samples-Main (Source: //UE4/Dev-Documentation @ 4091886)

diff --git a/Engine/Documentation/Source/Support/Builds/ReleaseNotes/4_19/4_19ReleaseNotes.INT.udn b/Engine/Documentation/Source/Support/Builds/ReleaseNotes/4_19/4_19ReleaseNotes.INT.udn
new file mode 100644
--- /dev/null
+++ b/Engine/Documentation/Source/Support/Builds/ReleaseNotes/4_19/4_19ReleaseNotes.INT.udn
@@ -0,0 +1,3816 @@
+Availability:Public
+Title: Unreal Engine 4.19 Release Notes
+Description: Release Notes for Unreal Engine 4.19
+Crumbs:%ROOT%, Engine
+Parent:Support/Builds
+Order:19
+reference-image:image_7.gif
+reference-image:image_30.gif
+reference-image:image_43.gif
+
+Unreal Engine 4.19 Release Notes
+
+![image alt text](419Banner.png)(w:1000 convert:false)
+
+# What’s New
+
+Unreal Engine 4.19 enables you to **step inside the creative process** - the tools become almost transparent so you can spend more time creating. Improvements to rendering, physics, Landscape terrain, and many more systems mean you can **build worlds that run faster than ever before**. The quality of life for the developers using our tools is always top of mind, so we continue to look at areas we can improve to **put the power into the developers’ hands**.
+
+Whether you are creating games, linear media, architectural visualizations, or product design tools, Unreal Engine 4.19 enables you to know exactly what the finished product will look like every step of the way. The new Live Link plugin seamlessly blends workflows between external content creation tools and Unreal Engine so you can see updates as you make changes to source content. And with the continued improvements to Sequencer, you can be the director with even more control of your scenes in real time.
+
+When it comes to bringing the worlds in your imagination to life, the sky's the limit. Create breathtaking vistas in large, open worlds thanks to Landscape rendering optimizations. With the new Dynamic Resolution feature that adjusts the resolution as needed to achieve desired frame rates, those worlds will run smoother than ever before on PlayStation 4 and Xbox One.
+
+It wouldn’t be a complete release without a mountain of workflow and usability improvements, and Unreal Engine 4.19 does not disappoint in this respect. Working with Material layers and parameters is easier and more intuitive. Features for debugging Blueprints are more powerful with the ability to step over, step into, and step out. You can now save content folders as favorites. Animation tools have been improved with pinnable commands, ability to have multiple viewports, and lots more.
+
+In addition to all of the updates from Epic, this release includes **128**** improvements ****submitted by the incredible community of Unreal Engine developers on GitHub**! Thanks to each of these contributors to Unreal Engine 4.19:
+
+ 
+
+[REDACTED_CONTRIBUTORS]
+
+# Major Features 
+
+## New: Temporal Upsampling
+
+We have added a new upscaling method called Temporal Upsample that performs both a temporal accumulation of the frame at lower resolution and primary spatial upscale, reducing output blur.
+
+[OBJECT:EmbeddedVideo]
+	[PARAMLITERAL:videoid]
+	Z4fa3EdnG5k
+	[/PARAMLITERAL]
+	[PARAMLITERAL:width]
+	78
+	[/PARAMLITERAL]
+	[PARAMLITERAL:height]
+	56.25
+	[/PARAMLITERAL]
+	[PARAMLITERAL:units]
+	%
+	[/PARAMLITERAL]
+	[PARAMLITERAL:type]
+	simple
+	[/PARAMLITERAL]
+[/OBJECT]
+
+*1 - Temporal Upsample Enabled; 2 - Temporal Upsample Disabled *
+
+In previous versions, the idea of Screen Percentage enabled you to render the 3D scene at lower resolution, and spatially upscale it to the output resolution, before the UI gets drawn on top. This was a very flexible method to hit a GPU budget on less powerful hardware, but it required you to make a tradeoff between output sharpness and GPU performance.
+
+To replace that single screen percentage, we now offer two separate screen percentages used for upscaling:
+
+* Primary screen percentage that by default will use the spatial upscale pass as before;
+
+* Secondary screen percentage that is a static, spatial only upscale at the very end of post processing, before the UI draws.
+
+![image alt text](image_1.png)(w:1000 convert:false)
+
+The temporal upscaler that happens in the Temporal Anti-Aliasing (TAA) pass enables consistent geometry sharpness from varying primary screen percentages from `50%` up to `200%`. In effect, even though the screen percentage may be lowered, with TAAU enabled, geometry in the background that would usually become muddy or blend together can now maintain its detail and complexity, like the fence and telephone pole in the example above.
+
+For additional information, see the [Screen Percentage with Temporal Upsample](https://docs.unrealengine.com/en-us/Engine/Rendering/ScreenPercentage) page.
+
+## New: Dynamic Resolution
+
+We now have support for Dynamic Resolution, which **adjusts the resolution as needed to achieve a desired framerate**, for games on PlayStation 4 and Xbox One! This works by using a heuristic to set the primary screen percentage based on the previous frames GPU workload.
+
+![image alt text](image_2.png)(w:1000 convert:false)
+
+Support for additional platforms will be coming in a future release. For additional information, see the [Dynamic Resolution](https://docs.unrealengine.com/en-us/Engine/Rendering/DynamicResolution) page.
+
+## New: Unified Unreal AR Framework
+
+With this release of Unreal Engine, the Unified Unreal Augmented Reality Framework provides a rich, unified framework for building Augmented Reality (AR) apps for both Apple and Google handheld platforms. The framework provides a single path of development, allowing developers to build AR apps for both platforms using a single code path. The Unified Unreal AR Framework includes functions supporting Alignment, Light Estimation, Pinning, Session State, Trace Results, and Tracking.
+
+![image alt text](image_3.png)(w:1000 convert:false)
+
+## New: Unified Unreal AR Framework Project Template
+
+Also new is the Blueprint template HandheldAR, which provides a complete example project demonstrating the new functionality.
+
+![image alt text](image_4.png)(w:849 h:590 convert:false)
+
+## New: Physical Light Units
+
+New to Unreal Engine 4.19, all light units are now defined using physically based units. Some light units were already well defined, but others were using undefined, engine specific units. The unit selection of a light is done through a drop-down menu (where applicable). For compatibility reasons, the default light units are kept compatible with previous versions of the engine. The new light unit property can be edited per light, changing how the engine interprets the "Intensity" property when doing lighting related computations.
+
+![image alt text](image_5.png)(w:1000 convert:false)
+
+*Project courtesy of Litrix*
+
+## New: Live Link Plugin Improvements
+
+The **Maya Live Link Plugin** is now available and can be used to establish a connection between Maya and UE4 enabling you to **preview changes made in Maya in real-time inside the UE4 Editor**. 
+
+[OBJECT:EmbeddedVideo]
+	[PARAMLITERAL:videoid]
+	q_h-HxeVaJA
+	[/PARAMLITERAL]
+	[PARAMLITERAL:width]
+	78
+	[/PARAMLITERAL]
+	[PARAMLITERAL:height]
+	56.25
+	[/PARAMLITERAL]
+	[PARAMLITERAL:units]
+	%
+	[/PARAMLITERAL]
+	[PARAMLITERAL:type]
+	simple
+	[/PARAMLITERAL]
+[/OBJECT]
+
+Once enabled, you will need to copy the binary files associated with your version of Maya to your Maya Plugins folder and enable it through the Maya Plugins Manager. 
+
+You can find pre-built binaries of the Maya Live Link Plugin inside your Engine installation folder in the following location: Engine\Extras\MayaLiveLink\LiveLinkMaya.zip. Inside the zip file are binaries for Maya 2016/2017 and 2018 for Windows. If you require binaries for other versions, the source code for the Plugin can be found in the Engine\Source\Programs\MayaLiveLinkPlugin folder which can be used to build the binaries. 
+
+The Maya Live Link UI window can be enabled through the MEL console with the command **MayaLiveLinkU**. At the top right is a display that shows whether this instance of Maya is connected to an Unreal client. Below Unreal Engine Live Link is a list of all the subjects currently being streamed (in the image above, only one is being streamed) and in the lower window, controls for adding and removing subjects from streaming is available. 
+
+### Motionbuilder Live Link Plugin
+
+**
+**The Motionbuilder Plugin offers the same functionality as the Maya Plugin and shows up in the Editor as a connection in a similar way. It also has a custom UI for managing streaming:
+
+Objects can be selected from the the current scene and added to the streamed list (as shown above). From there, their names can be set in the **Subject Name** column and their **Stream Type** (Camera, Skeleton, etc.) can be set. Streaming on the subject can also be enabled and disabled from here. 
+
+### Stream Active Camera as Editor Active Camera
+
+Controlling the active camera within Maya will now manipulate the Editor’s active camera. We’ve also reworked the Editor update hook so that streaming is more robust and facial rigs/leaf bones now correctly update. 
+
+### Added Virtual Subjects to Live Link
+
+Virtual subjects are a way to combine multiple subjects coming into Live Link into one subject that can then be used by the Editor. Virtual Subjects are created within the client and contain the bones of multiple real subjects, all tied to a common root. 
+
+Below is the virtual subject being applied to a skeletal mesh in the editor. 
+
+[REGION:asyncgif]
+![image alt text](image_7.gif)(w:624 convert:false)
+[/REGION]
+
+Here we have created a new virtual subject (using **Add > Add Virtual Subject**) and set it to read subjects from the Maya Live Link source. The virtual subject is set to consist of the three Item subjects being sent from Maya. 
+
+### Live Link Plugin Development
+
+The purpose of Live Link is to provide a common interface for streaming and consuming any kind of animation data from sources outside of UE4 (for example, DDC tools and Motion Capture servers). It is designed to be extensible via Unreal Plugins, allowing third parties to develop new features with no need to make and then maintain Engine changes. 
+
+There are two paths for integrating in Live Link:
+
+1. Building an Unreal Engine Plugin that exposes a new Source to Live Link. This is the recommended approach for anyone that already has their own streaming protocol.
+
+2. Integrating a Message Bus end-point in third party software to allow it to act as a data transmitter for the built-in Message Bus Source. This is the approach we have taken for our Maya and MotionBuilder Plugins.
+
+For an overview of general Plugin development, see the [Plugins](https://docs.unrealengine.com/latest/INT/Programming/Plugins/) documentation pages. 
+
+### Live Link Motion Controller Support
+
+Live Link can now be used with Motion Controllers. The motion source property of a Motion Controller can be set to a subject within Live Link. When set in this way, the position of the Motion Controller Component is governed by the first transform of the subject.
+
+The Motion Controller integration can also access custom parameters on the Live Link subject. These are passed via the curve support built into Live Link subjects. 
+
+To access the values, it is necessary to derive a new Blueprint from **MotionControllerComponent **and override the **OnMotionControllerUpdated **function. During **OnMotionControllerUpdated**, it is valid to call **GetParameterValue **on the Motion Controller.
+
+Here is an example of a possible way to drive a light components Intensity from an Intensity parameter in Live Link:
+
+![image alt text](image_8.png)(w:925 h:289 convert:false)
+
+### Miscellaneous Improvements
+
+* **Added ability for Live Link Sources to define their own custom settings - **This was requested by partners like IKinema building Live Link support
+
+* **Live Link Sources pushing skeletons to Live Link Client can now pass source GUID as well** - Message bus source now pushes GUID when sending skeleton.
+
+* **Added virtual Initialization function and Update DeltaTime parameter to Live Link Retargeter API**
+
+## New: Sequencer Improvements
+
+We continue to make improvements to the functionality and workflow of Sequencer to make it more powerful and increase efficiency.
+
+### Copy/Paste/Duplicate Object Tracks
+
+You can now copy/paste/duplicate object tracks and their child tracks from the right-click context menu. You can copy a spawnable to the same Level Sequence or another Sequence or a possessable from one Level Sequence to another and the object will be bound to the same object in the other Level Sequence. You can now also copy from one UMG animation to another.
+
+[OBJECT:EmbeddedVideo]
+	[PARAMLITERAL:videoid]
+	e2ixsUw0S6k
+	[/PARAMLITERAL]
+	[PARAMLITERAL:width]
+	78
+	[/PARAMLITERAL]
+	[PARAMLITERAL:height]
+	56.25
+	[/PARAMLITERAL]
+	[PARAMLITERAL:units]
+	%
+	[/PARAMLITERAL]
+	[PARAMLITERAL:type]
+	simple
+	[/PARAMLITERAL]
+[/OBJECT]
+
+### Level Sequence Dynamic Transform Origin
+
+You can now offset the Actors controlled by Sequence Transform tracks with a global offset at runtime. This allows you to also reuse a Level Sequence in different coordinate spaces. To use this feature, inside the Details panel for your Level Sequence, enable **Override Instance Data** then assign a **Transform Origin Actor. **
+
+The **Transform Origin **section specifies a transform from which all absolute transform sections inside the Sequence should be added to (Scale is ignored). For the best results, keyframe your Actor’s starting transform at 0,0,0 inside your Level Sequence and let the Transform Origin Actor you define drive the position from which to start in the world (depicted below). 
+
+[OBJECT:EmbeddedVideo]
+	[PARAMLITERAL:videoid]
+	hf82P3SmWqU
+	[/PARAMLITERAL]
+	[PARAMLITERAL:width]
+	78
+	[/PARAMLITERAL]
+	[PARAMLITERAL:height]
+	56.25
+	[/PARAMLITERAL]
+	[PARAMLITERAL:units]
+	%
+	[/PARAMLITERAL]
+	[PARAMLITERAL:type]
+	simple
+	[/PARAMLITERAL]
+[/OBJECT]
+
+Whenever we press a key (see above), the cube moves and our Actor continues to walk along a path, starting from the location of the **Transform Origin Actor**.
+
+### Sequencer Anim BP Weight Control
+
+Sequencer weight blending now works with Animation Blueprints. You can use the same slot for the animation, and control weight by curve. Assign the same Slot Name under Properties for each animation to blend. Then keyframe the Weight values you desire. Inside your Anim BP, use the Slot Node. The AnimBP will use whatever Weight values you provided when blending.
+
+[OBJECT:EmbeddedVideo]
+	[PARAMLITERAL:videoid]
+	4c8S4BEFUa4
+	[/PARAMLITERAL]
+	[PARAMLITERAL:width]
+	78
+	[/PARAMLITERAL]
+	[PARAMLITERAL:height]
+	56.25
+	[/PARAMLITERAL]
+	[PARAMLITERAL:units]
+	%
+	[/PARAMLITERAL]
+	[PARAMLITERAL:type]
+	simple
+	[/PARAMLITERAL]
+[/OBJECT]
+
+### Improved Sequencer Editor Performance
+
+* Compile On The Fly Logic - Sequencer is now able to compile partially or completely out-of-date evaluation templates from the source data as it needs. This affords much more efficient compilation when working within Sequencer in the editor.
+
+* Details Panel Updates - The details panel now defers updates while scrubbing/playing in the Sequencer Editor.
+
+### Minor Improvements
+
+* New .ini setting for the default behavior of "When Finished". 
+
+    * For level sequences, it continues to default to RestoreState. For UMG, it’s now set to Keep State.
+	
+	[/Script/LevelSequence.LevelSequence]
+DefaultCompletionMode=RestoreState
+
+* Always create a camera cut track when adding a camera so that there’s less confusion as to why a camera doesn’t take control when sequencer is activated.
+
+* Rotations will no longer flip when autokeying/adding a new key. 
+
+* Option to show selected nodes in the tree view only. This can be useful when working with a lot of actors in your level sequence to limit scrolling up and down the view. 
+
+* Option to bake transform track to keys at frame intervals.
+
+* Option to create a camera anim from a transform track.
+
+* Ability to bind (path track, camera cut track or attach track) to spawnables in subsequences. 
+
+* Can now drag sections up a row, instead of only just down.
+
+## New: Landscape Rendering Optimization
+
+The Landscape level of detail (LOD) system has been changed so that instead of being distance-based it now uses screen size to determine detail for a component, similar to how the Static Mesh LOD system works. The LOD distribution now gives more coherent sizes on distant triangles based on their screen size, which can maintain detail that was not possible before.
+
+![image alt text](image_12.png)(w:1000 convert:false)
+
+![image alt text](image_13.png)(w:490 convert:false)    ![image alt text](image_14.png)(w:490 convert:false)
+
+*(Left, old method; Right, new method)*
+
+You can visualize the Landscape LODs by going to **Viewport > Visualizers > LOD**. 
+
+The following settings are available to control where LODs transition based on screen size. Additionally, material tessellation can now be controlled, which gives the added performance benefit while using a Directional light with dynamic shadowing enabled. 
+
+![image alt text](image_15.png)(w:441 convert:false)
+
+## New: Proxy LOD System (Experimental)
+
+The new Proxy LOD system is an experimental Plugin for producing low poly LOD with baked materials for multiple meshes. The new system is used by HLOD and is a replacement for Simplygon. 
+
+[REGION:note]
+**Note:** Currently, only Windows builds are supported.
+[/REGION]
+
+To enable the tool, look for the "Hierarchical LOD Mesh Simplification" option under Editor in the Project Settings dialog. From the Hierarchical LOD Mesh Reduction Plugin, you should be able to select the new module "ProxyLODMeshReduction".  After the prompted editor restart, the Plugin will replace the third party Simplygon tool for static mesh merging LODs.  This new Plugin is accessed in two ways: The HLOD Outliner, and the Merge Actors dialog.
+
+![image alt text](image_16.png)(w:490 convert:false) ![image alt text](image_17.png)(w:490 convert:false)
+
+
+*Source Geometry: 11 geometry pieces, 38,360 tris, and 27,194 verts | Result Geometry Proxy: 1 geometry piece, 8,095 tris, and 5,663 verts*
+
+
+
+## New: Material Parameters Editing and Saving
+
+You can now **save parameter values to a new Child Instance or Sibling Instance** in the Material Editor and the Material Instance Editor!
+
+ ![image alt text](image_18.png)(w:1000 convert:false)
+
+* **Save to Sibling Instance** enables you to save the current parameters and overrides to a new Material Instance with the same Parent Material.
+
+* **Save to Child Instance** enables you to save the current parameters and overrides to a new Material Instance with the current instance as the Parent Material.
+
+The Material Editor now features a **Parameter Defaults** panel. Here, you’ll have access to the default values set in your Material Graph for any parameter. You can easily change any of your parameter default values here.
+
+![image alt text](image_19.png)(w:1000 convert:false)
+
+## New: Material Layering (Experimental)
+
+Material Layering enables you to **combine your Materials in a stack**, using the new Material Layer and Material Layer Blend assets! This enables you to build the correct Material Graph without building sections of nodes by hand. This functionality is similar to Material Functions, but supports the creation of child instances.
+
+[OBJECT:EmbeddedVideo]
+	[PARAMLITERAL:videoid]
+	vc5Jg-iqPgk
+	[/PARAMLITERAL]
+	[PARAMLITERAL:width]
+	78
+	[/PARAMLITERAL]
+	[PARAMLITERAL:height]
+	56.25
+	[/PARAMLITERAL]
+	[PARAMLITERAL:units]
+	%
+	[/PARAMLITERAL]
+	[PARAMLITERAL:type]
+	simple
+	[/PARAMLITERAL]
+[/OBJECT]
+
+Material Layering is experimental and can be enabled in the **Project Settings** **> Rendering > Experimental** section by setting **Support Material Layers** to true.
+
+For additional information, see our post on the [forums](https://forums.unrealengine.com/development-discussion/rendering/1414193-material-layering-feedback-for-4-19) to get started and leave us feedback!
+
+## New: Upgraded 'Auto Convex' Tool
+
+You can now specify **Max Hulls **instead of accuracy when generating simplified collision for a Static Mesh using the Auto Convex Tool. Computation is now a background task in the Static Mesh Editor. The Auto Convex Tool also uses a newer version of V-HACD library that should give more accurate results. 
+
+![image alt text](image_21.png)(w:1000 convert:false)
+
+## New: Improved Facial Animation Sharing using Curve only animations
+
+You can share facial animations between characters by animating a curve to drive a pose without having any bone transform. Any bone transform, even with one mesh’s ref pose, won’t work for other meshes, so it is important to remove them and start with the mesh’s own ref pose if you want to share the curve between different meshes. This allows you to share the facial curves between different faces. 
+
+We’ve added the ability to allow users to delete all bone tracks, leaving only animation curves. You can access the **Remove All Bone Tracks **option from the **Asset **menu under **Animation. **
+
+[OBJECT:EmbeddedVideo]
+	[PARAMLITERAL:videoid]
+	m54mNljjiAs
+	[/PARAMLITERAL]
+	[PARAMLITERAL:width]
+	78
+	[/PARAMLITERAL]
+	[PARAMLITERAL:height]
+	56.25
+	[/PARAMLITERAL]
+	[PARAMLITERAL:units]
+	%
+	[/PARAMLITERAL]
+	[PARAMLITERAL:type]
+	simple
+	[/PARAMLITERAL]
+[/OBJECT]
+
+You can also use Animation Modifier to remove certain joints only if you want to share body animations but drive curves with face. Added the functions **Remove Bone Animation **to remove an Animation Curve by Name from the given Animation Sequence (Raw Animation Curves [Names] may not be removed from the Skeleton) and **Remove All Bone Animation **to remove all Animation Bone Track Data from the given animation Sequence. 
+
+![image alt text](image_23.png)(w:1000 convert:false)
+
+We’ve also allowed the opting out of importing bone tracks (curve only animation) during the import process.
+
+We’ve fixed Pose Preview to select weights that the asset contains. You can see what bone tracks the animation contains inside the Asset Details panel under **Animation Track Names.**
+
+## New: Cloth Updates
+
+The new "Anim Drive" feature added to clothing assets and runtime interaction for simulations adds a new mask target for "Anim Drive Multipliers" that is used to drive springs that pull clothing towards its animated position, giving control to animators in cinematics or animation driven scenes.
+
+[OBJECT:EmbeddedVideo]
+	[PARAMLITERAL:videoid]
+	iXY2jqgYZ3o
+	[/PARAMLITERAL]
+	[PARAMLITERAL:width]
+	78
+	[/PARAMLITERAL]
+	[PARAMLITERAL:height]
+	56.25
+	[/PARAMLITERAL]
+	[PARAMLITERAL:units]
+	%
+	[/PARAMLITERAL]
+	[PARAMLITERAL:type]
+	simple
+	[/PARAMLITERAL]
+[/OBJECT]
+
+This can be driven at runtime using the "Interactor" object that Blueprints can access on skeletal mesh components. The value set at runtime is multiplicatively combined with the painted values to get the final spring and damper strengths.
+
+Other improvements include:
+
+* **Gravity override support for clothing simulations - **Using the interactor object on a skeletal mesh component - an arbitrary gravity override can be specified.
+
+* **Added auto-range feature to cloth view ranges - **Also added a way for tools to extend those ranges when necessary
+
+* **Improved Clothing Visualizations - **Can now be enabled when Clothing Paint Mode is active.
+
+* **Clothing Create Menu now defaults to use the Skeletal Mesh Physics Asset**
+
+## New: Animation Blueprint Improvements
+
+We have made several quality of life improvements to the Animation Blueprints system:
+
+* You now have the ability to control whether or not the post process instance runs on a Skeletal Mesh Component from both Blueprints and the Animation Editor suite 
+
+* **Copy Pose node now supports copying curves as well as Bone transform - **When using this, make sure to enable **bCopyCurves** if you want to copy curves (by default this is false). When enabled, it will copy the curves that exist in the current Skeleton. If the curve list has changed, it will require the animation system to be reinitialized.
+
+* **Added additive warning when playing an animation on an additive pose - **When a non-additive animation is connected to the additive pin of an additive animation node (such as Apply Additive), we will now output a warning in the message log with the Animation and Blueprint name.
+
+* **Native AnimInstance selection on SkeletalMeshComponent**** - **Allow native Animation Instances to be set on Skeletal Mesh Components. Previously, only Anim Blueprints could be set on Skeletal Mesh Components.
+
+## New: Animation Tools Improvements
+
+**Added Pinnable Command List to animation editors**
+
+Shift-clicking a menu option in either the skeleton tree filter or the viewport menus will 'pin' the command outside of the menu. This allows for easier access to commands that are frequently used.
+
+[OBJECT:EmbeddedVideo]
+	[PARAMLITERAL:videoid]
+	jy3VPBRus9k
+	[/PARAMLITERAL]
+	[PARAMLITERAL:width]
+	78
+	[/PARAMLITERAL]
+	[PARAMLITERAL:height]
+	56.25
+	[/PARAMLITERAL]
+	[PARAMLITERAL:units]
+	%
+	[/PARAMLITERAL]
+	[PARAMLITERAL:type]
+	simple
+	[/PARAMLITERAL]
+[/OBJECT]
+
+**Added up to 4 separate viewports to all Animation Editor viewports - **All animation-related editors now have the ability to open up to 4 separate viewports onto the same scene, each with their own settings.
+
+![image alt text](image_26.png)(w:1000 convert:false)
+
+**Added the ability to follow and orbit specified bones** 
+
+[OBJECT:EmbeddedVideo]
+	[PARAMLITERAL:videoid]
+	Z-WVH9XwH_c
+	[/PARAMLITERAL]
+	[PARAMLITERAL:width]
+	78
+	[/PARAMLITERAL]
+	[PARAMLITERAL:height]
+	56.25
+	[/PARAMLITERAL]
+	[PARAMLITERAL:units]
+	%
+	[/PARAMLITERAL]
+	[PARAMLITERAL:type]
+	simple
+	[/PARAMLITERAL]
+[/OBJECT]
+
+**Additional Skeletal Mesh Component Debug Info**
+
+You can extend the viewport text seen in the Skeletal Mesh Editor. Lines were also added, displaying current Cloth Value.
+
+[OBJECT:EmbeddedVideo]
+	[PARAMLITERAL:videoid]
+	JHb5BwJJcHY
+	[/PARAMLITERAL]
+	[PARAMLITERAL:width]
+	78
+	[/PARAMLITERAL]
+	[PARAMLITERAL:height]
+	56.25
+	[/PARAMLITERAL]
+	[PARAMLITERAL:units]
+	%
+	[/PARAMLITERAL]
+	[PARAMLITERAL:type]
+	simple
+	[/PARAMLITERAL]
+[/OBJECT]
+
+**Ability to change Skeletal Mesh section visibility at runtime**
+
+Exposed ShowMaterialSection function to hide/show Skeletal Mesh material sections at runtime to Blueprint. Added ShowAllMaterialSections and IsMaterialSectionShown functions. Removed internal 'disable' and 'highlight' flags from asset: always use per-component flags now.
+
+[OBJECT:EmbeddedVideo]
+	[PARAMLITERAL:videoid]
+	igoOgJPtQFU
+	[/PARAMLITERAL]
+	[PARAMLITERAL:width]
+	78
+	[/PARAMLITERAL]
+	[PARAMLITERAL:height]
+	56.25
+	[/PARAMLITERAL]
+	[PARAMLITERAL:units]
+	%
+	[/PARAMLITERAL]
+	[PARAMLITERAL:type]
+	simple
+	[/PARAMLITERAL]
+[/OBJECT]
+
+**Added "Hide Unconnected Pins" for button BreakStruct nodes**
+
+BreakStruct nodes now have a button to quickly hide unconnected pins
+
+[REGION:asyncgif]
+![image alt text](image_30.gif)(w:800 convert:false)
+[/REGION]
+
+In addition, the Animation Tools improvements include:
+
+* **Fixed motion blur when changing skeletal mesh properties - **The animation system owns the previous bone transform with revision number, so that it can send to the render thread when RenderState is recreated. **Frame Number** still exists because the clothing system still uses it. Revision number progresses when the component ticks. If it doesn't tick, it will keep the same revision number (keeps history).
+
+* **Exposed Sequence Recording Settings in Animation Editors**
+
+* **Ability to name collision shapes** - You can now name individual collision shapes when editing a Physics Asset.
+
+* **Improved bounds following in Animation Editor Viewports** - Bounds following now follows without lag.
+
+* **Added ability to change the Retarget Source inside a Pose Asset** - Pose Assets now save the original Source Pose so that it can recalculate a new pose when Retarget Source changes. Retarget Source is used to bake the pose asset delta transform, so you will need to make sure you have the right Retarget Source if you are using the same asset for different meshes. 
+
+* **Configurable mesh paint vertex size** - Added mesh paint settings to per-user settings, and added preview vertex size, so that it's configurable by the user for use on different-size meshes.
+
+* **Changed animation preview tooltips in Persona to stop loading the animation if it is not already loaded**
+
+* **Preview buttons are now present where applicable in animation editors - **If relevant, Preview Mesh, Preview Animation and Reference Pose are all editable via the toolbar.
+
+* **Added Search bar to Advanced Preview Scene settings panel**
+
+* **Added Context Menu to HLOD level nodes in HLOD outliner treeview** - This allows creating a new cluster from the current level viewport selection.
+
+* **Added a hotkey and command for toggling Post Processing in Preview Scenes**
+
+* **Added a new shortcut for switching foreground and background colors in the Mesh Painter** - By default the shortcut is bound to the ‘X’ Key when the Mesh Painter is active.
+
+* **Material Baking has been moved out of the experimental phase**
+
+* **Added support for RGBA Masks in Mesh Painter** - This allows the use of RGBA channels as masks while using the fill tool inside of the Mesh Painter.
+
+* **Added 'GetSectionFromProceduralMesh' utility function for ProceduralMeshComponent**
+
+* **Notifies on follower animations in a sync group** - Added flag to allow notifies to be triggered on follower animations in a Sync Group.
+
+* **Reduced Logging verbosity when cooking all assets** - Only log 2 lines per asset instead of 4.
+
+* **Exposed Damping and Stiffness as pins on Spring AnimNode** - Spring Anim Node now allows exposing the Spring Damping and Spring Stiffness as pins.
+
+## New: Motion Controller Component Visualization
+
+Motion Controllers in Unreal Engine 4.19 now have a new Visualization category that enables you to quickly (and more easily) add a display model Static Mesh to your Motion Controllers. By default, the system attempts to load a Static Mesh model matching the device driving the Motion Controller.
+
+[OBJECT:EmbeddedVideo]
+	[PARAMLITERAL:videoid]
+	sVp5n4FkZwE
+	[/PARAMLITERAL]
+	[PARAMLITERAL:width]
+	78
+	[/PARAMLITERAL]
+	[PARAMLITERAL:height]
+	56.25
+	[/PARAMLITERAL]
+	[PARAMLITERAL:units]
+	%
+	[/PARAMLITERAL]
+	[PARAMLITERAL:type]
+	simple
+	[/PARAMLITERAL]
+[/OBJECT]
+
+* Motion Controller components now have new visualization fields that offer the following options: *Display Device Model*, *Display Model Source*, *Custom Display Mesh*, and *Display Mesh Material Overrides*.
+
+    * *Custom Display Mesh* is only available when *Display Model Source* is set to "Custom".
+
+    * For more complicated, custom models, manually attach your own Static Mesh model.
+
+    * ‘*Display Model Source*’ can be manually set if the target XR system isn’t loaded.
+
+* By default, when empty, ‘*Display Model Source*’ will attempt to load a model from the backing XR system.
+
+    * Distinct XR systems support this feature by implementing the IXRSystemAssets interface - currently, only Oculus and SteamVR support this functionality.
+
+## New: Vive Pro Support
+
+![image alt text](image_32.jpg)(w:1000 convert:false)
+
+All existing UE4 content that supports SteamVR is compatible with HTC’s newly-announced Vive Pro!  While no modifications are necessary to use the Vive Pro with your existing UE4 project, you will need to check the performance of your UE4 project with the Vive Pro and optimize accordingly as the Vive Pro operates at a higher resolution. 
+
+## New: Platform SDK Upgrades
+
+In every release, we update the engine to support the latest SDK releases from platform partners.
+
+![image alt text](image_33.jpg)(w:1000 convert:false)
+
+* **IDE Version the Build farm compiles against**
+
+    * **Visual Studio:**  Visual Studio 2015 Update 3
+
+    * **Xcode:**  Xcode 9.2
+
+* **Android:  **
+
+    * NDK 12b (New CodeWorks for Android 1r6u1 installer will replace previous CodeWorks for Android 1R5 before release, still on NDK 12b)
+
+* **HTML5:** emscripten toolchain to 1.37.19
+
+* **LInux:** v11_clang-5.0.0-centos7
+
+* **Steam:** 1.39
+
+* **Oculus Runtime:** 1.17
+
+* **SteamVR:** 1.39
+
+* **Switch:** 3.7.0 with NEX: 4.3.1 (optional), 4.5.0 with NEX 4.3.1 (optional); Firmware Version 3.0.1-1.0
+
+* **PS4:**
+
+    * 5.008.071
+
+    * Firmware Version 4.730.001
+
+    * Supported IDE: Visual Studio 2015, Visual Studio 2013
+
+* **Xbox One (XB1, XB1-S, XB!-X):**
+
+    * XDK: June 2017 QFE9 
+
+    * Firmware Version: October 2017 Preview 3 (version 16291.1000.170913-1900)
+
+    * Supported IDE: Visual Studio 2017
+
+* **macOS:** SDK 10.13
+
+* **iOS:** SDK 11
+
+* **tvOS:** SDK 11
+
+## New: Low Input Latency Frame Syncing (Xbox One and PlayStation 4)
+
+Low-latency frame syncing mode modifies the way thread syncing is performed between the game, rendering and RHI threads, and the GPU to greatly reduce and control input latency. The effects of Low-latency frame syncing can be controlled using the r.GTSyncType (game thread sync type) console variable. This console variable drives how the new frame syncing mechanism works and allows the following values:
+
+* 0 - Game thread syncs with rendering thread (old behavior, and default)
+
+* 1 - Game thread syncs to the RHI thread (equivalent to UE4 before parallel rendering)
+
+* 2 - Game thread syncs with the swap chain present +/- an offset in milliseconds (only available on supported platforms, and only works when vsync is enabled)
+
+For additional information, see the [Low Latency Frame Syncing](https://docs.unrealengine.com/en-us/Platforms/LowLatencyFrameSyncing) page.
+
+## New: Encryption/Signing Key Generator
+
+It is now possible to **configure .pak file signing and encryption fully from the Editor** via a single unified settings panel. These features can now be controlled in Project Settings, under the Project section, in the Crypto panel.
+
+![image alt text](image_34.png)(w:577 convert:false)
+
+From this panel, it is possible to generate or clear keys for either signing or encryption, as well as enable or disable the various types of signing and encryption that are available. The following features can be activated or deactivated from this menu:
+
+* **Encrypt Pak .ini Files** - Encrypts all .ini files that exist in your project’s .pak files. This will prevent easy mining or tampering with configuration data for your product, at minimal runtime cost.
+
+* **Encrypt Pak Index** - Encrypts the pak file index, which prevents UnrealPak from opening, viewing, and unpacking your product’s pak files, at minimal runtime cost.
+
+* **Encrypt UAsset Files** - Encrypts the uasset files in the pak file. uasset files contain header information about the assets inside, but not the actual asset data itself. Encrypting this data provides extra security to your data, but does add a small runtime cost and an increase in data entropy which can have a negative impact on patch sizes.
+
+* **Encrypt Assets** - Fully encrypts all Assets within the pak file. Please note that this setting will have a measurable effect on run-time file I/O performance, and increases the amount of entropy in your final packaged data, making the distribution patching system less efficient.
+
+## New: Call Stack Display for Blueprint Debugging
+
+A control that displays the current Blueprint Call Stack has been added to the Developer Tools menu. Double-clicking entries in the Call Stack will focus the corresponding node in the Blueprint Editor.
+
+![image alt text](image_35.png)(w:1000 convert:false)
+
+## New: Single stepping improvements to Blueprint Debugging
+
+The Engine’s Blueprint Debugger now supports **Step Into (F11)**, **Step Over (F10)**, and **Step Out (Alt+Shift+F11)** operations when stopped at a Blueprint Breakpoint, making debugging sessions quicker and easier.
+
+![image alt text](image_36.png)(w:1000 convert:false)
+
+In addition, we have added the ability to break directly on, and step into, over, or out of, macros and collapsed graph nodes:
+
+[OBJECT:EmbeddedVideo]
+	[PARAMLITERAL:videoid]
+	fkBHxds-cKg
+	[/PARAMLITERAL]
+	[PARAMLITERAL:width]
+	78
+	[/PARAMLITERAL]
+	[PARAMLITERAL:height]
+	56.25
+	[/PARAMLITERAL]
+	[PARAMLITERAL:units]
+	%
+	[/PARAMLITERAL]
+	[PARAMLITERAL:type]
+	simple
+	[/PARAMLITERAL]
+[/OBJECT]
+
+## New: Improved Tools For Optimizing Disk Size
+
+Several tools have been added or improved to help analyze disk and download sizes for packaged games. The Asset Audit window has a new "Add Chunks" option which shows all packaged Chunk/Pak files, and the Reference Viewer and Size Map can be used from that view to analyze why Assets are being cooked into a specific chunk.
+
+The Size Map tool has been changed to show Disk Size by default, which is more accurate than the estimated memory size. If a platform is selected in the Asset Audit window, the actual cooked size for that platform will be shown. The Size Map’s user interface has also been improved to support right-click operations and better performance.
+
+![image alt text](image_38.png)(w:1000 convert:false)
+
+## New: Audio Improvements
+
+**Upgrade to Steam Audio Beta 10 (New Audio Engine Only)**
+
+![image alt text](image_39.png)(w:361 convert:false)
+
+Along with the update to Steam Audio (Beta 10 Release), improvements include:
+
+* Android support for spatialization and occlusion
+
+* Editor mode for easy access to Steam Audio functionality
+
+* Buttons to add and remove Phonon Geometry Components to all Static Mesh Actors
+
+* Button to export the scene to .obj for debugging
+
+* Removed Phonon Scene actor in favor of export to file
+
+* Probe generation and baking writes data to disk, to avoid slow loading and saving of map files
+
+* Numerous information display improvements
+
+* Numerous bug fixes
+
+**Oculus Audio Updates (New Audio Engine Only)**
+
+![image alt text](image_40.png)(w:356 convert:false)
+
+Oculus Audio plugin has been updated with support for Android.
+
+In addition to providing HRTF spatialization, we now support Oculus Audio’s room modeling reverb with volumetric parameters as well as Ambisonics Source playback.  
+
+**Added Resonance Audio Plugin (New Audio Engine Only)**
+
+![image alt text](image_41.jpg)(w:340 convert:false)
+
+The Resonance Audio plugin for Unreal is now included as part of the Engine Plugins in UE 4.19. The Resonance Audio Plugin supports binaural spatialization, reverberation, sound directivity, and ambisonics playback on Windows, iOS, Android, Linux, and Mac. 
+
+**First-Order Ambisonics File Support (Playback supported through plugins)**
+
+First-order Ambisonics files can now be imported into Unreal projects and played back using select Spatialization plugins. First-Order Ambisonics is a spherical surround format that offers forward, right, and upward vectors of a sound field allowing designers to create Surround Sound experiences that include vertical sound orientation.  This is useful for VR, where users will experience natural changes in vertical sound fields by tilting their heads or for any other interactive experience where user orientation allows changes in the vertical.  Users can import 4 channel First-Order Ambisonics files by including the FuMa or AmbiX (the two First-Order Ambisonics channel configurations) tag in their file name. Oculus Audio and Google Resonance both have ambisonics playback implementations.
+
+**Improvements to UE4's native VOIP implementation (New Audio Engine Only)**
+
+Our native VOIP implementation has been improved to provide better coherency including improved silence detection on Windows.  Additionally, we now support bringing VOIP sources into the scene allowing for:
+
+* Spatialization
+
+* Distance Attenuation
+
+* Reverb
+
+* Custom Source Effects
+
+All controlled using the UVoipTalker scene component.
+
+**New microphone capture component  (New Audio Engine Only) (Windows Only)**
+
+We have implemented a Microphone Capture Component which allows:
+
+* Feeding audio directly into a game for use with the Audio Mixer
+
+* Driving in-game parameters using real-time amplitude data.
+
+* Processing audio with Source Effects or Submix Effects.
+
+Currently, this is only for local client microphone capture and is not connected to the VoIP system
+
+While it doesn't support recording microphone output directly to assets at the moment, there are opportunities 3rd party plugin developers to take advantage of data from a microphone input.
+
+**Get real-time audio amplitude as Blueprint data (New Audio Engine Only)**
+
+Added ability to get amplitude envelope data, directly following audio components and synth components through BP delegates.
+
+![image alt text](image_42.png)(w:564 convert:false)
+
+**Pre and Post Source Effect Bus Sends Types**** ****(New Audio Engine Only)**
+
+We have added Pre- and Post- Source Effect Bus Sends types allowing designers to send Source audio to Source Busses both before and after Source Effect and Attenuation processing.
+
+**Added support for Multichannel file import**
+
+We now support importing interleaved multichannel wav files directly into the browser (in addition to the legacy method using multiple mono files).
+
+**Sample rate control on synths**
+
+We have added the ability for synthesizers to modify their requested sample rate on initialization allowing synth sources to specify their own sample rate.
+
+**Improved cross-platform audio plugins (New Audio Engine Only)**
+
+We now allow multiple audio plugin source settings so users can have different settings for different plugins which is especially useful for designers working on multi-SKU products where different plugins will be used on different platforms.
+
+We have also added Plugin drop down menus for PS4 and Switch platform settings for added convenience.
+
+**Adding a panner source effect**
+
+This Source Effect will allow designers the ability to drive stereo panning values from Blueprints as a Source Effect. This allows you to change the pan of a mono or stereo 3d or 2d source.
+
+## New: Physics Improvements
+
+Several improvements and optimizations have been made to the physics system to make simulating physical objects easier and efficient!
+
+* **Added "Show Only Selected Constraints" to the physics asset editor**
+
+* **Added Ctrl+T keyboard shortcut to toggle between Body and Constraint Selection**
+
+* **Added "Set Use CCD" function for changing Continuous Collision Detection option at run-time**
+
+* **Saved BodyInstance runtime memory - **Removed FBodyInstance::ResponseToChannels_DEPRECATED property from non-Editor Builds to save memory.
+
+* **Added functions to SpringArmComponent to get relevant status of the collision test**.
+
+* **Added the ability to profile scene query hitches - **Using p.SQHitchDetection console variable
+
+* **Exposed some PhysicsAsset functionality to be more extensible in C++**
+
+* **Improved information in HitResult in case of initial overlap - **ImpactPoint will now contain information from PhysX on deepest overlap point. 
+
+* **Engine now compiles cleanly without PhysX - **Can set WITH_PHYSX define directive to 0 and compile without errors. 
+
+## New: Automatic Code Signing on iOS
+
+Automatic Code Signing (ACS) when using Xcode has now been introduced in UE4.19. ACS plus a Team ID allows a package to be created and signed for an iOS project. This also works when there are no mobile provisions pre-installed on the Mac but the Mac must contain the certificates. A development team will still need to add iOS devices to their Apple account in order to deploy their apps, but they will no longer need to download and install mobile provisions on to a Mac if they use this feature. 
+
+## New: Android 4.19 Improvements
+
+* Added Proxy network support.
+
+* Added a new pop up box that will inform you if the mobile project you are running is using an unsupported Texture format and also let you know what texture formats are supported
+
+* Added NDK and SDK platform validation before you compile and pack your project to let you know sooner if the build has failed or not.
+
+## New: Folder Favorites for the Content Browser
+
+You can now **mark any folder in your Content Browser as a "favorite"** to make it easily accessible! When a folder is marked as a favorite, it will populate under the new "Favorites" section located above the folder tree. Simply right-click to add or remove folders from this section.
+
+The "Favorites" section of the Content Browser can be enabled/disabled using the View Options for Show Favorites. This feature is enabled by default.
+
+[REGION:asyncgif]
+![image alt text](image_43.gif)(w:600 convert:false)
+[/REGION]
+
+## RELEASE NOTES
+
+#### AI
+
+* Bug Fix: Fixed the On Perception Updated delegate to be correctly bindable to AI Perception Component.
+
+##### Behavior Tree
+
+* Bug Fix: Changed the order of the Behavior Tree tick to allow Blueprint-based nodes to update before processing search data.
+
+* Bug Fix: Fixed a bug that caused a Behavior Tree asset containing a Run Dynamic Behavior node to become dirty upon opening.
+
+* Bug Fix: Fixed a Behavior Tree state change during a latent abort. Decorators and services updates are delayed until after the aborting node finishes.
+
+* New: Added the local decorator scope to the Behavior Tree Composite Nodes. The local decorator scope is controlled by the Apply Decorator Scope** **variable. This addition provides better control of flow-observing decorators, allowing them to be active only in a specified branch of the tree.
+
+* New: Added a simple mechanism for logging Behavior Tree** **nodes that reference empty or removed Blackboard keys.
+
+##### Debugging Tools
+
+* New: Extended the EQS Testing Pawn with an option to specify navigation agent properties to support navigation-related EQS tests and features.
+
+* New: Added more Visual Logging functions to Blueprint.
+
+* Bug Fix: Fixed some existing Visual Logging Blueprint functions.
+
+##### Navigation
+
+* Bug Fix: Fixed cases where changes to Nav Mesh Bounds Volumes properties were not causing Nav Mesh to rebuild.
+
+* Bug Fix: Fixed a failure to set the area flags when updating the Nav Area of a given Nav Mesh poly.
+
+* New: Made updates to allow the proper use of Nav Link Component via Blueprints.
+
+#### Animation
+
+* Crash Fix: Fixed a crash when opening Animation Compression settings for newly imported Alembic File.
+
+* Crash Fix: Applied a fix for a rare crash when Notify Queue caches a pointer to the notify that if garbage collected, would have a pointer to invalid memory. We now cache a pointer to the owner of the notify memory so we can track its validity. 
+
+* Crash Fix: Applied a fix for a Live Link module startup crash.
+
+* Crash Fix: Prevented threaded work from running in a single-threaded environment, which avoids a Linux server crash.
+
+* Bug Fix: Fixed debug builds hitting checkSlow when using a post process Anim Instance on a Skeletal Mesh Component without a main AnimInstance.
+
+* Bug Fix: Fixed some race conditions in Skeletal Mesh Resource Sync UV Channel Data and Static Mesh Render Data Sync UV Channel Data.
+
+* Bug Fix: Fixed a bug with adjacency DDC serialization in Skeletal Meshes.
+
+* Bug Fix: Non-rendered anim updates now run on worker threads.
+
+* Bug Fix: Changed the way Get Bone Poses For Time inside of the Animation Modifiers library retrieves the Bone Transforms for a specific time-code.
+
+* Bug Fix: Fixed Anim Sequence Base Get Time At Frame and Anim Sequence Base Get Frame At Time to return correct frame indices.
+
+* Bug Fix: Fixed pasted socket name getting an appended _0 when no other sockets contained that name.
+
+* Bug Fix: Reset Tick Records when doing 'Only Tick Montages When Not Rendered' since the montage will appear to have jumped when regular ticking resumes.
+
+* Bug Fix: URO now ensures you don't skip more frames than desired when switching LODs. 
+
+* Bug Fix: Animation tasks are no longer created for Skeletal Meshes that have no anim instance.
+
+* Bug Fix: Skeletal Meshes that are not being ticked due to URO can no longer dispatch tick tasks.
+
+* Bug Fix: Fix for notifies not getting fired in cases where Always Tick Pose was set on Skeletal Mesh Components.
+
+* Bug Fix: Looping sound cues and particle systems used in notifies no longer leak a component per invocation.
+
+* Bug Fix: Recalc Required Bones now uses correct predicted LOD level instead of defaulting to 0.
+
+* New: Optimized Post Anim Evaluation when URO is skipping a frame with no interpolation.
+
+* New: Changed the notify system to allow notifies on follower animations in a sync group to be triggered.
+
+* New: Added two callbacks on Character Movement Component for the root motion world space conversion to allow users to modify the root motion:
+
+    * Pre conversion (while it is still in component space).
+
+    * Post conversion (when it is in world space).
+
+* New: Added Virtual Subjects to Live Link which are created within the client and contain the bones of multiple real subjects.
+
+* New: Added Initialization and Update Delta Time to Live Link Retargeter API:
+
+    * Pushing skeleton to Live Link can now take source GUID.
+
+    * Message bus source pushes GUID when sending skeleton.
+
+* New: Added 'Play Rate Basis' to Sequence Player node to scale Play Rate without having to do more expensive BP code.
+
+* New: Added support for Dialogue Waves driving Audio Curve Source Components.
+
+* New: Added default Get Axis Vector static function to Axis Option.
+
+* New: Skeletal Mesh Component Init Anim doesn't call Update and Eval anymore. 
+
+* New: Rigid Body Anim Node Updates:
+
+    * Now maintains Bone Velocity transfers through LOD changes. 
+
+    * Added support for transferring angular velocity. 
+
+    * Added option Transfer Bone Velocities to transfer bone velocities to simulation upon start, so ragdolls transition seamlessly from kinematic (animation) to simulation.
+
+    * Added Freeze Incoming Pose On Start option to freeze incoming pose and stop ticking/evaluating rest of the AnimGraph. 
+
+* New: Moved Update Component Pose Any Thread and Evaluate Component Pose Any Thread functions to Anim Node Skeletal Control Base to enable overriding these in child classes.
+
+* New: Added Pre Eval Skel Control Any Thread to Skel Control Base, to allow capture of incoming pose before SkelControl is evaluated.
+
+* New: Added the Empty function to Base Compact Pose and FCS Pose to release allocated arrays.
+
+* New: Added the notify class name to the notify tooltip.
+
+* New: Optimized redundant key removal on curves. Optimized code takes less than 1% of the time to execute of the original.
+
+* New: Added prebuilt binaries for the Maya Live Link plugin to Engine/Extras (Maya 2016, Maya 2017, Maya 2018).
+
+* New: Skeletal Mesh Render Data Get Max Bones Per Section is now an ENGINE_API exported function.
+
+##### Animation Assets
+
+* Crash Fix: Fixed a crash with streaming if a montage is streamed out while an event is queued.
+
+* Crash Fix: A crash no longer occurs due to oversampling animation during compression.
+
+* Crash Fix: No longer crashes when the Skeleton has been changed and joints have been removed from a Pose Asset.
+
+* Bug Fix: Fixed an ensure that was triggering when compressing an animation which has had its Skeleton switched due to Skeletal Mesh Merging.
+
+* Bug Fix: Changed raw animation data evaluation so that Virtual Bone positions are built before interpolation is carried out instead of after. 
+
+    * This fixes a slight mismatch between the Virtual Bone positions in raw animation and compressed animation.
+
+* Bug Fix: Marker Sync position is now maintained across instant transitions.
+
+* Bug Fix: Animations with a scaled root bone now generate correct root motion.
+
+* Bug Fix: Changes Retarget Source hard reference to Soft Object Ptr, so that the Skeleton won't load all the references of Retarget Sources.
+
+* Bug Fix: Fixed Pose Asset retargeting not working for additive pose assets.
+
+* Bug Fix: Fixed issues with Pose Assets of additive blending using shortest path to blend, causing it to blend between rather than accumulating.
+
+* Bug Fix: Fixed Retarget Source for Pose Asset.
+
+    * Changing Retarget Source in Pose Asset now works.
+
+    * Pose Asset now saves original source pose, so that it can recalculate new pose when Retarget Source changes.
+
+    * Retarget Source is used to bake the pose asset delta transform. 
+
+* New: Native Animation Instances can now be set on Skeletal Mesh Components.
+
+* New: Anim Sequence Base is now entirely exported with ENGINE_API.
+
+* New: Added support for curve-only animation (required for facial pose retargeting).
+
+* New: Added inline time/frame settings to notify context menu.
+
+##### Animation Blueprints
+
+* Crash Fix: Loading an Anim Instance with a re-instanced class no longer crashes.
+
+* Crash Fix: A crash no longer occurs when loading reparented Animation Blueprints that do not have an Anim Graph but are no longer Child Blueprints. 
+
+* Crash Fix: Initialized Poses To Evaluate on Blend nodes to fix a crash from gathering stale data after reinitializing.
+
+* Bug Fix: Fixed "Modify Curve" Anim Node breaking certain nodes that were used before it (due to not caching bones correctly).
+
+* Bug Fix: Anim Dynamics correctly calculates simulation transforms and world vector transformations. 
+
+* Bug Fix: Transition nodes no longer reference the wrong state node due to non-unique node GUIDs being introduced by copy and paste.
+
+* Bug Fix: Subinstance nodes inside states no longer fail to correctly create bridge variables on their skeleton class due to not correctly processing subgraphs in the Anim Blueprint compiler.
+
+* Bug Fix: Added a fix for animation reinitializing when replacing animation via Single Node Instance.
+
+* Bug Fix: Added code to correct duplicate node GUIDs during Anim Blueprint compilation.
+
+* New: Copy Pose node now supports copying curves as well as bone transforms. Enable Copy Curves if you want to copy curves (by default, this is not enabled).
+
+* New: Added ability for worker thread Animation Blueprint code to report to Message Log and added additive warning when playing an animation on an additive pose.
+
+* New: Added "Uncheck All Unconnected Pins" for button Break Struct nodes.
+
+* New: Added Native Begin Play function to Anim Instance.
+
+* New: Added ability to control whether or not a post process instance runs on a skeletal mesh component. Exposed both to blueprints and anim editor suite
+
+* New: Added support for logical negation when copying to array properties in the fast path.
+
+* Deprecated: Get Scalar Parameter Default has been deprecated and replaced with Get Scalar Parameter Default Value.
+
+##### Animation Tools
+
+* Crash Fix: Removing bones in a LOD on a Skeleton that has Virtual Bones no longer crashes.
+
+* Crash Fix: Crash no longer occurs when updating clothing paint mode after tab spawners have been destroyed by the hosting application.
+
+* Crash Fix: Attempting to zero out an array with the wrong size in Control Rig no longer crashes.
+
+* Crash Fix: Addressed a crash in Animation Editor menus caused by audio compression. 
+
+* Crash Fix: No longer crashes when displaying a two bone IK gizmo for a node that hadn't had a chance to evaluate or had a zero alpha.
+
+* Bug Fix: Skeleton Tree now correctly no longer allows selection of filtered items. 
+
+* Bug Fix: Bone modifications via transform type-in can now be undone.
+
+* Bug Fix: Fix to make sure Debug Bone Rendering of raw animation data in the Animation Editors actually shows the raw data when an animation uses additive track layers instead of showing the source.
+
+* Bug Fix: Fixed bad blend profile references from copy/paste operations on nodes using external profiles.
+
+* Bug Fix: Fixed Blend Profile properties not being correctly customized on anim nodes - and Animation Sequence references not being correctly filtered on anim nodes.
+
+* Bug Fix: Addressed Single Node Instances not getting ticked properly. Due to them not increasing UpdateCounter, and forcing a tick even if we're doing parallel ticking later.
+
+* Bug Fix: Removed recompress animation from curve track actions. This allows for smoother interaction on animations with a slow recompression time.
+
+* Bug Fix: Text is now red again in Anim Viewports.
+
+* Bug Fix: Fixed Set Root Motion Enabled Anim Data Modifier node (previously didn't set the enabled flag)
+
+* New: Exposed Sequence Recording Settings in Animation Editors.
+
+* New: Changed animation preview tooltips in Persona to stop loading the animation if it is not already loaded.
+
+* New: Tweaked debug display of Anim Sequences and added DeltaTime to AnimInstance debug.
+
+* New: The console command "showdebug animation" now shows current LOD and counters to know if Update/Eval/Cachebones/Init was called. 
+
+    * Also for URO settings, renamed Display Debug Custom to Display Debug Instance.
+
+* New: Added Anim Instance Display Debug Custom, to display custom debug info before AnimGraph display.
+
+* New: Removing a section when creating a clothing data entry will now just disable the section in the mesh, so that the operation is non-destructive to the mesh data.
+
+* New: Moved Live Link Message Bus Source heartbeats to their own thread using a new  Heartbeat Manager singleton. This prevents sources from incorrectly being removed during Slate UI operations.
+
+##### Import/Export
+
+* New: Refactored Facial Animation importing so it can be used from inside other modules.
+
+##### Skeletal Mesh
+
+* Crash Fix: Undoing a change on a Skeletal Mesh that uses Virtual Bones no longer crashes.
+
+* Crash Fix: Setting a Master Pose Component on a Slave Component using the Rigidbody node during actor construction will not result in a crash. 
+
+* Crash Fix: No longer crashes when clearing a Skeletal Mesh on a Skeletal Mesh Component with an active Post Process Anim Instance.
+
+* Bug Fix: Skeletal Mesh Component LOD info no longer cleans up all override resources and only cleans up one. 
+
+* Bug Fix: Set the Used by Morph Target Material option to only the Materials affecting section deform by a Morph Target instead of all Materials used by the Skinned Mesh component.
+
+* Bug Fix: Multi-convex generation for Skeletal Meshes no longer uses incorrect index buffer.
+
+* Bug Fix: Correctly return whether a Skeletal Mesh section is shown or not when it has not been edited.
+
+* New: Exposed Show Material Section function to hide/show Skeletal Mesh material sections at runtime to Blueprint (also added Show All Material Sections and Is Material Section Shown functions). 
+
+* New: Animation System owns previous Bone Transform with revision number that are used for motion blur or temporal AA, so that it can send to the render thread when RenderState is recreated. 
+
+* New: You can set attached and parent to sync their LOD by setting Sync Attach Parent LOD to true.
+
+##### Tools
+
+* Crash Fix: Skeletal Mesh Editor no longer crashes when importing more than 5 LODs.
+
+* Bug Fix: Refactored View menus in Animation Editors.
+
+* Bug Fix: Incorrect reporting of Animation Blueprint Fast Path is no longer caused by enabling the Compilation Manager. 
+
+* Bug Fix: Compression Settings panel no longer allows "Use Adaptive Error" and "Use Adaptive Error 2" to be enabled at the same time.
+
+* Bug Fix: Ensured bounds inside of the Animation Editor(s) are correctly shown for Alembic Skeletal Meshes.
+
+* Bug Fix: Removed unused preferences from Persona Options.
+
+* Bug Fix: Persona Viewport settings are now per-asset editor.
+
+* New: Defaulted the Clothing Create menu to use the Skeletal Mesh Physics Asset.
+
+* New: All animation-related editors now have the ability to open up to 4 separate viewports onto the same scene, each with their own settings.
+
+    * Added the ability to follow (and orbit) a specified bone as well.
+
+    * Added a new "pinnable command list" widget to animation editors.
+
+* New: Preview buttons are now present where applicable inside the Animation Editors.
+
+* New: Added Mesh Paint settings to per-user settings and added preview vertex size so it's configurable by the user for use on differently sized meshes.
+
+* New: Auto-range feature added to cloth view ranges and a way for tools to extend those ranges when necessary.
+
+* New: Improved Skeletal Mesh Component debugging.
+
+    * Extended the viewport text seen in Skeletal Mesh Editor.
+
+    * Added lines for current cloth value.
+
+* New: Clothing visualizations can now be enabled while clothing paint mode is active.
+
+* New: Added a search bar to the Advanced Preview Scene settings panel.
+
+#### Audio
+
+* Crash Fix: Fixed an issue with VOIP crashing on Mac.
+
+* Crash Fix: Fixed Ngs2 sources getting recycled while getting setup.
+
+* Crash Fix: Applied a fix for editing inherited audio component overrides in attenuation detail customization.
+
+* Crash Fix: Fixed a crash in the Audio Device function, Processing Pending Active Sound Stops.
+
+* Bug Fix: Fixed a case where audio is already set to be a preview sound. Sequencer sets the sound to be a preview sound when scrubbing in editor.
+
+* Bug Fix: Fixed a hang during multi-PIE shutdown when using a synth component.
+
+* Bug Fix: Sound mix fade time now fades audio properly.
+
+* Bug Fix: Fixed a bug during shutdown of audio mixer.
+
+* Bug Fix: The sample rate of synth components can now be overridden in C++.
+
+* Bug Fix: Fixed a bug in Android VOIP.
+
+* Bug Fix: Fixed a potential memory leak caused when starting and stopping a SynthComponent on the same frame.
+
+* Bug Fix: Fixed Ogg Vorbis 5.1 channel ordering channel maps in the Audio Mixer.
+
+* Bug Fix: Fixed an issue where audio sources were gently filtered near Nyquist when a source's LPF was set to max frequency in the new Audio Engine.
+
+* Bug Fix: Fixed a memory leak in the audio decompressor.
+
+* Bug Fix: Libvorbis DLL failing to load will now recover gracefully.
+
+* Bug Fix: Removed "Stopping sound already in the process of stopping" log spam.
+
+* Bug Fix: Envelope Attack and Release times are now setting properly on Synth Components.
+
+* Bug Fix: White Noise in the new Audio Engine no longer produces non-uniform noise.
+
+* Bug Fix: Fixed an issue in which the Delay Submix effect would allocate too much memory for delay lines.
+
+* Bug Fix: Recording audio in Recording Sequencer will no longer cause the Editor process to persist after shutdown.
+
+* Bug Fix: Fixed issues in wavetable sampleplayer synth component.
+
+* New: Added the ability for Submixes to define their own channel format. Currently, Stereo, Quad, 5.1, 7.1, and first order ambisonics are supported.
+
+* New: Made improvements to Unreal Engine 4's native VOIP implementation. Spatialization, Distance Attenuation, Reverb, and custom effects for a player's voice can now be controlled using the UVoipTalker Scene Component.
+
+* New: The new Microphone Capture Component enables feeding audio directly into a game for use with the Audio Engine and for driving in-game parameters.
+
+* New: Upgraded to Steam Audio Beta 10.
+
+* New: Added the ability to get envelope following Audio Components and Synth Components through BP delegates.
+
+* New: Added the ability for synthesizers to modify the requested sample rate on initialization.
+
+* New: Users can now have different Audio Plugin source settings for different Plugins.
+
+* New: Added a Panner Source Effect.
+
+* New: Sound Submixes and Sound Classes are now Blueprint types.
+
+* New: Added support for seeking with ADPCM.
+
+* New: Optimized Active Sound by not processing active sounds if they are out of range.
+
+* New: On Mac Media, audio playback uses the operating system media player mixer rather than mixing in Unreal Engine 4. This is currently only enabled on Mac, but may come to iOS in the future.
+
+* New: Oculus Audio Plugin has been updated. It now includes a Reverb Plugin, as well as support for Android and playing back Ambisonics files.
+
+* New: First-order Ambisonics files can now be imported into Unreal projects and played back using select Spatialization Plugins.
+
+* New: Added pre and post Source Effect Bus Send types. At this time, sending audio to post-source effects sends before distance-based effects, such as occlusion and distance attenuation are applied.
+
+* New: Multichannel file import is now supported.
+
+* New: The Allow Anyone To Destroy Me variable is now enabled on audio components.
+
+* New: The Resonance Audio Plugin for Unreal is now included as an Engine Plugin. The Resonance Audio Plugin supports binaural spatialization, reverberation, sound directivity, and ambisonics playback on Windows, iOS, Android, Linux, and Mac.
+
+* New: Audio Components can automatically manage attachment, opting in to cache their attach parent and attach socket. This enables them to attach to their cached parent only while playing audio, then detach after playback is completed. Set to false by default.
+
+* New: Optimizing the way audio stream chunks are serialized to take advantage of alignment optimizations on console.
+
+* New: Added a Blueprint function for enabling and disabling the Low Pass Filter on Audio Components, as well as setting the frequency of the filter.
+
+* New: Added "Owner" parameter to Play Sound At Location and Play Sound 2D functions to support "Limit to Owner" concurrency settings.
+
+#### Automation
+
+* Crash Fix: Fixed a crash in Screen Shot Manager caused by an attempt to access a deleted String in an async lambda expression.
+
+* Bug Fix: Fixed issue with UnrealBuildTool and AutomationTool always being out of date, which forced a rebuild and made it impossible to debug operations dependent on them.
+
+* Bug Fix: Automated test screenshots now have more consistent overrides for gamma, eye adaption, dpi and other post process effects (regardless of capture mode).
+
+* Bug Fix: Fixed automation routines that list files in Perforce to return those whose last action was a rename.
+
+* New: Added automated test for per-platform blacklist defined in the relative platform engine ini file.
+
+* New: Automated test screenshots are now able to test against a hierarchical fallback platform and RHI.
+
+* New: Added button to automated test screenshot browser to add or replace all outstanding test reports (if appropriate).
+
+* New: Automated screenshot tests that fail due to size mismatch will now show a delta image, where mismatched bounds are colored red.
+
+#### Blueprints
+
+* Crash Fix: Fixed inclusion of private plugin header files in nativized Blueprint C++ source.
+
+* Crash Fix: Resolved a potential crash that could occur prior to initiating a global "Find in Blueprints" indexing operation with source control disabled.
+
+* Crash Fix: Fixed a crash when using EDL to load certain data-only blueprints that inherit from other blueprint classes
+
+* Crash Fix: Fixed a crash on load in a nativized build caused by a reference to a Blueprint class containing a User-Defined Enum asset.
+
+* Crash Fix: Fixed a potential crash when collapsing nodes to a function when a potential entry pin had no links.
+
+* Crash Fix: Fixed a crash when using blueprinted components created by a specialized subclass of Blueprint.
+
+* Crash Fix: Added a guard fix against a reported crash when deleting a non-root scene component with no parent attachment.
+
+* Crash Fix: Prevented a reported crash when deleting one or more Actor components.
+
+* Crash Fix: Added a guard fix against a reported crash when setting bIsEditorOnly flag on attached scene components.
+
+* Crash Fix: Fixed an issue where split struct pins would not respect Allowed Classes on things like Slate Font Info, which could cause crashes when passing in the wrong type of asset
+
+* Crash Fix: Fixed occasional crashes when manually compiling simple blueprints.
+
+* Crash Fix: Fixed several crash and data loss issues when undoing changes to user defined structs used in Data Tables
+
+* Bug Fix: In the Blueprint Material And Texture Nodes Plugin, clamping sampled HDR render target values by setting Range Compression Mode in the Read Surface Data Flags to RCM_MinMax now behaves as expected.
+
+* Bug Fix: User Defined Structures now store default values in the same way that native structs do, so they will now be correct in cooked builds and when added to containers. Changing the default value will also correctly propagate to unmodified instances in blueprints or other user structs.
+
+* Bug Fix: Fixed missing Outer mapping at cook time on nested subobjects instanced with the New Object function in emitted C++ code during Blueprint nativization.
+
+* Bug Fix: Blueprint nativization now emits external plugin dependencies to the generated plugin descriptor file in order to mute UBT warnings at build time.
+
+* Bug Fix: Removed the Inverse Lerp node and redirected it to call Normalize To Range. Also Fixed incorrect value being returned by Normalize To Range when Range Min and Range Max are equal.
+
+* Bug Fix: Fixed a rare linker issue where Is Blueprint Finalization Pending would return true when trying to force-load subobjects that weren’t ready to be loaded. This would prevent some placeholder objects from being replaced.
+
+* Bug Fix: We now ensure that "-build" is always passed to BuildCookRun automation workflows in the editor for projects configured with Blueprint nativization enabled so that it doesn't skip that stage.
+
+* Bug Fix: Revised both UAT/UBT-driven and custom code/data workflows to more seamlessly work with Blueprint nativization settings.
+
+* Bug Fix: Ed Graph Schema’s Break Single Pin Link function is now correctly marked as a const function.
+
+* Bug Fix: Local variables can now be created in blueprint function libraries.
+
+* Bug Fix: Fixed several issues with copy/pasting nodes that allow customizing which pins are visible.
+
+* Bug Fix: Light error sprite will no longer appear in blueprint preview window.
+
+* Bug Fix: Copy-pasting a collapsed graph containing a map input will no longer result in a map with an unset value type.
+
+* Bug Fix: Prevented nativization from creating uncompilable code when using a non-referenceable term (e.g. subclass of, weak object pointer) in a select node.
+
+* Bug Fix: We now properly cast between mismatched soft pointer types in nativized Blueprint C++ assignment statements and function calls.
+
+* Bug Fix: Blueprints that are asynchronously loaded inside the editor are correctly registered with editor systems like the list of cast nodes.
+
+* Bug Fix: Text can be specified in blueprint functions as default parameters.
+
+* Bug Fix: Fixed a serialization oversight in Ed Graph Pin that could lead to a compile error during Blueprint nativization.
+
+* Bug Fix: Added more descriptive log output after an attempt to break Blueprint execution on an unimplemented function/interface.
+
+* Bug Fix: Added Component nodes no longer lose their component type when another the component blueprint is compiled
+
+* Bug Fix: Split struct pins now handle both user struct variable renaming and native Core Redirects variable renaming, the same way that Make/Break did before.
+
+* Bug Fix: Implemented Get Redirect Pin Names for async nodes. This fixes a bug where property redirectors on async nodes did not work correctly.
+
+* Bug Fix: Fixed issues where type would not be correctly propagated through multiple redirect nodes.
+
+* Bug Fix: User Defined Structures can now be imported from text using the human-readable field names, which fixes their use in nested arrays inside data table rows.
+
+* Bug Fix: User Defined Structures that are nested inside other user structs can now be safely modified. It is now safe to modify the inner struct's defaults without corrupting the outer struct, and those default changes will be correctly propagated to the outer struct.
+
+* Bug Fix: Blueprint nativization now correctly generates and builds code for "Client" builds
+
+* Bug Fix: Improvements to the Cpp backend to allow VC++ to compile nativized code more effectively
+
+* Bug Fix: Subobjects of assets are now correctly referenced by nativized code Disabled animation fast-path optimization when running a native anim BP, as native code is faster!
+
+* Bug Fix: Fix split pins on macro instances for resolved wildcard pins not resetting when disconnected
+
+* Bug Fix: Child Actor Template label in the details panel will now correctly reflect the selected class.
+
+* Bug Fix: Blueprints that have no parent class can now be reparented when using the compilation manager
+
+* Bug Fix: Harden pin serialization logic, specifically to prevent crash due to bad 'conform' routines when renaming functions
+
+* Bug Fix: Fixed ensure when initializing replication data for blueprint generated classes
+
+* Bug Fix: Fixed missing marching ants on trace wires when stopped at a breakpoint in macro and collapsed source graphs.
+
+* Bug Fix: Fix issue where references from a component to its parent blueprint would be cleared on compilation.
+
+* Bug Fix: Fixed an ensure when exiting PIE after tracing into a Blueprint Function Library source graph while single-stepping.
+
+* Bug Fix: Fixed incorrect access to nested instanced subobjects in nativized Blueprint class constructor code.
+
+* Bug Fix: Prevented debug object selection dropdown from displaying objects with pending-kill outers.
+
+* Bug Fix: Fixed a packaging build failure that could occur with Blueprint nativization enabled and EDL disabled.
+
+* Bug Fix: The proper syntax is now being emitted for set/map fields containing converted assets to C++ class headers when Blueprint nativization is enabled.
+
+* Bug Fix: Fixed a scoping issue for local instanced subobject references in nativized Blueprint C++ code.
+
+* Bug Fix: Fix several issues with the class pin on the Spawn Actor From Class node
+
+* Bug Fix: Added more verbose log output to assist with debugging any crash reports due to an invalid object reference detected in the BP event graph frame during garbage collection.
+
+* New: Added Call Stack controls for viewing Blueprint call stacks when paused at a breakpoint. These are available from the Developer Tools menu.
+
+* New: Improved Blueprint single-step debugging and breakpoints.
+
+    * Added new Step Over (F10) and Step Out (ALT+SHIFT+F11) debugging features.
+
+    * Added the ability to break directly on and step into/over/out of macro and collapsed graph nodes.
+
+* New: Updated numeric graph pins to use a numeric editor.
+
+* New: Validated Get option for Soft Object/Class references in blueprints.
+
+* New: The Game Mode Base function Get Default Pawn Class For Controller can now be called in blueprints.
+
+* New: Added an Editor Preferences setting to disable navigation to native functions from call function nodes.
+
+* New: Deprecated properties can still be accessed in blueprints so long as they have blueprint getters and setters. This creates a clean deprecation path for blueprint-accessed properties.
+
+* New: Added new Blueprint functions to query parameters from Material Instance Constant objects. Previously, these were only available on Material Instance Dynamic objects. The functions included are:
+
+    * Get Scalar Parameter Value
+
+    * Get Texture Parameter Value
+
+    * Get Vector Parameter Value
+
+* New: Added a step to the compilation manager to skip recompilation of classes that are dependent on a given classes function signatures when those signatures have not changed. This greatly improves compile times for blueprints with large numbers of dependencies.
+
+* New: Preserve graph editor widgets when clicking on a hyperlink, this speeds up navigation via the debugger "step" command and Find in Blueprints control - both are now more responsive.
+
+* New: Added Vector Length XY to Blueprint math library.
+
+#### Core
+
+* Crash Fix: Fixed a crash on shutdown due to component visualizers not being cleaned up properly by the Unregister Component Visualizer function.
+
+* Crash Fix: Network platform file heartbeat function can no longer cause a crash by firing during async loading.
+
+* Crash Fix: Fixed a potential garbage collection crash caused by a stale pointer in Anim Instance Proxy.
+
+* Crash Fix: Fixed CVar access crashes with the async loading thread enabled.
+
+* Crash Fix: Fixed a crash in EDL loading when simple construction scripts and inherited component handlers were not being loaded early enough in complicated blueprint hierarchies.
+
+* Crash Fix: Rntering a single quote in an ini field will no longer crash the editor or make a project fail to load.
+
+* Crash Fix: Dumping stats with massive callstacks no longer causes a crash.
+
+* Crash Fix: Fixed some crashes with the new Blueprint Garbage Collection weak reference feature (and possibly Hot Reload) by updating the archive's serialized property when serializing arrays, sets, or maps to point to the inner property.
+
+* Crash Fix: Fixed a crash that could happen when GC Objects were constructed and destroyed during engine shutdown.
+
+* Crash Fix: An editor startup crash when using cook-on-the-side has been fixed.
+
+* Crash Fix: Fixed an issue where packaged games that disabled the TCP/UDP messaging plugins would crash on startup in shipping configuration.
+
+* Crash Fix: Fixed a crash when the engine was reporting package errors.
+
+* Crash Fix: Running DDC commandlet with an invalid package name no longer causes crashes.
+
+* Bug Fix: Multicast Delegate Property’s Identical function now has the same behavior as Delegate Property’s Identical function. This prevents it from giving incorrect results and causing serialization errors.
+
+* Bug Fix: The Asset Registry module is now always preloaded on startup, even if EDL is disabled.
+
+* Bug Fix: Fixed Thread Sanitizer warnings in stats system. Optimized Cycle Counter’s Start function to only read the stat name once (atomically).
+
+* Bug Fix: Fix ambiguities in calculating root directory.
+
+* Bug Fix: When serializing a compressed archive with multithreaded compression enabled, the engine will wait for the oldest async task instead of spinning.
+
+* Bug Fix: When Async time limit is set to a low value, the linker will no longer get stuck in a loop of reloading the preload dependencies.
+
+* Bug Fix: Fixed a small race condition with config data in Low Level Mem Tracker.
+
+* Bug Fix: Fixed some data races in Compression's stats.
+
+* Bug Fix: Shared Ptr, Shared Ref and Weak Ptr assignment operators are now implemented as copy-and-swap, to avoid bad states being visible to any destructors being called.
+
+* Bug Fix: Added atomic reads to Shared Ptr, Shared Ref, and Weak Ptr in Thread Safe mode.
+
+* Bug Fix: Used new Platform Atomics function, AtomicRead, in Thread Safe Counter‘s Get Value function.
+
+* Bug Fix: Fixed thread-unsafe sharing of a static array in the String function, Parse Into Array WS.
+
+* Bug Fix: Fixed a regression in hot reload speed accidentally introduced in 4.17, caused by rerunning UnrealHeaderTool redundantly if a source file has changed in the same folder as a header file which contains at least one UCLASS. This has also been fixed in non-hot-reload compiles.
+
+* Bug Fix: Removed a redundant call to Update Vertices when creating a level's model components, which was causing a race condition with the renderer thread.
+
+* Bug Fix: Suppressed Thread Sanitizer warnings in the Async Read Request interface.
+
+* Bug Fix: Fixed the debugger visualization of Array with an Inline Allocator or Fixed Allocator. Improved the visualization of Sparse Array.
+
+* Bug Fix: Fixed the wrong "P_GET_" macro being used in the generated code for Soft Class Ptr properties.
+
+* Bug Fix: Removed a redundant class access modifier from Queue.
+
+* Bug Fix: Freed the stats thread on shutdown.
+
+* Bug Fix: Is Integral now correctly reports true for int64 on 32-bit platforms.
+
+* Bug Fix: Fixed truncation errors in the parsing of very large integer values in FJsonObject::GetIntegerField().
+
+* Bug Fix: Unreal Header Tool's error reporting now specifies the correct source file when a header doesn't specify a .generated.h as its last #include.
+
+* Bug Fix: Deleted an unnecessary specialization of Make Array View.
+
+* Bug Fix: Removed String's char pointer constructor’s confusing void pointer parameter.
+
+* Bug Fix: Base64’s Encode function is now const-correct.
+
+* Bug Fix: Json Value no longer stringifies whole numbers as floats.
+
+* Bug Fix: Relaxed enum property importing to allow valid integer values to be imported too. This was previously made more strict, which caused a regression in Data Table importing.
+
+* Bug Fix: Fixed memory leaks for pipe writes and added data pipe writes.
+
+* Bug Fix: Fast XML now properly parses attributes with legally unescaped single and double quote characters
+
+* Bug Fix: Zipping in UAT now always uses UTF-8 encoding to prevent Unicode issues.
+
+* Bug Fix: Fixed pure-virtual function call in Multicast Delegate Property.
+
+* Bug Fix: Fixed console output not working when using "-log" in a shipping build.
+
+* Bug Fix: Loading maps via content browser or async load in the editor no longer leads to their references not loading properly.
+
+* Bug Fix: Fixed an issue where deleting a source package would not cause the generated cooked package to get deleted while doing an incremental build.
+
+* Bug Fix: Enum redirects handle typedef enums properly when a value redirect is not also specified. Enum classes continue to work correctly.
+
+* Bug Fix: Fixed an issue where the Streamable Manager could return partially loaded objects when called from Post Load. It will now schedule a callback for when those objects finish loading instead.
+
+* Bug Fix: Comparing NAME_None to an empty string, or a Name constructed from an empty string, now correctly returns true.
+
+* Bug Fix: Fixed cmake include directory generation putting engine include files in the game root.
+
+* Bug Fix: Fixed cases where an include path that is external to both the engine and game root were incorrectly getting the game root prepended.
+
+* Bug Fix: Many fixes were implemented for Thread Sanitizer warnings.
+
+* Bug Fix: Fixed build command line in generated projects missing a space before the compiler version override.
+
+* Bug Fix: If UAT is invoked with platforms in both the "-platform" and “-targetplatform” command-line switches, we will now build using all of them rather than just the ones in “-targetplatform”.
+
+* Bug Fix: Fixed cases where String’s Printf function was used without a formatting string.
+
+* Bug Fix: Updated loading movie code to properly calculate delta time when throttled at 60 frames per second.
+
+* Bug Fix: Fixed an inconsistency with how Soft Object Ptr case is managed between Linker Save and Archive Save Tag Imports, which could cause a cook ensure under some circumstances.
+
+* Bug Fix: Fixed an occasional loading assert in Script Struct’s Post Load function caused by a race condition on the Prepare CPP Struct Ops Completed variable.
+
+* Bug Fix: Changed Visual Studio Code project file generation to use "taskName" rather than "label" for C# build task names.
+
+* Bug Fix: Improved whitespace handling in both Map Property’s and Set Property’s "ImportText_Internal" function. Emptied the map or set if an error occurs during a text import of a property to avoid bad container states.
+
+* Bug Fix: Fixed a potential buffer overrun in Set Property.
+
+* Bug Fix: Map Property’s Import Text function now correctly skips duplicate keys.
+
+* Bug Fix: Unreal Build Tool now enforces Development and Debug configurations in existing C# projects. Fixes command-line builds and prevents solution from being marked dirty on initial load.
+
+* Bug Fix: Cleaned up delegates in UObjectGlobals.h, fixing several incorrect comments and moving some editor delegates into WITH_EDITOR.
+
+* Bug Fix: Fixed an optimization issue with CPU benchmarking. Also added better support for debugging local rocket builds.
+
+* Bug Fix: Fixed the final enumerators being missing from the generated FOREACH_ENUM_ macros.
+
+* Bug Fix: Parallel .pak generation now uses the number of reported CPU cores. Also fixed an issue with concurrent .ini file reading.
+
+* Bug Fix: Dedicated servers do not include RHI modules except NullRHI.
+
+* Bug Fix: Changed String’s Parse Into Array function to use Reset instead of Empty on the passed in array, allowing it to be approximately resized
+
+* Bug Fix: Fixed accounting for Curve Table assets in "obj list".
+
+* Bug Fix: Fixed a bug that caused the engine to fail to identify a running instance of Visual Studio as having the current solution open.
+
+* Bug Fix: Fixed the "quote issue" with foreign project build tasks on PC.
+
+* Bug Fix: Fixed limit on number of async precache tasks while cooking, due to counter not being updated correctly.
+
+* Bug Fix: The Windows version string is now determined correctly.
+
+* Bug Fix: Asset registry memory usage reported by memreport is now correct.
+
+* Bug Fix: Fixed Platform File Pak interface’s Find Files Recursively function not correctly calling the base class implementation.
+
+* Bug Fix: Literal soft object pins are correctly cooked and tracked in the reference viewer. This may cause assets referenced by disabled nodes to be cooked into your project. In this case, the reference viewer will track which blueprints are causing this.
+
+* Bug Fix: The serializing object will now be logged during reachability analysis in the event that an invalid object reference is detected during a garbage collection pass.
+
+* Bug Fix: Hot reloads no longer occur during PIE, which was causing bad states in blueprint types when they are reinstanced.
+
+* New: Encryption/Signing Key Generator.A plugin which allows the automatic creation and configuration of in-engine data cryptography through the project settings, under the "crypto" header.
+
+    * Added options to plugin settings which allow encryption to be used on the various asset files in a pak. Encryption can be used on just the .uasset file, which contains only the package header information, or the .uexp and .ubulk files too, which contain the rest of the package data. Encryption on .uasset files provides a good compromise between data security and runtime performance.
+
+* New: Added Level Streaming CVars to customize streaming for your project. These all default to enabled, but can be disabled to make different trade offs for memory and level loading performance.
+
+    * "s.ForceGCAfterLevelStreamedOut" - Whether to force a GC after levels are streamed out to instantly reclaim the memory, at the cost of a hitch. Defaults to on.
+
+    * "s.ContinuouslyIncrementalGCWhileLevelsPendingPurge" - Whether to repeatedly kick off incremental GC when there are levels still waiting to be purged. Defaults to on.
+
+    * "s.AllowLevelRequestsWhileAsyncLoadingInMatch" - Enables level streaming requests while async loading (of anything) while the match is already in progress and no loading screen is up. Defaults to on.
+
+    * Added several new stats for potentially slow code.
+
+    * Added SCOPE_CYCLE_UOBJECT macro, instead of using Scope Cycle Counter UObject directly.
+
+* New: Cheat cvars are now usable in Test builds by default. Can be overridden by defining ALLOW_CHEAT_CVARS_IN_TEST as 0 in Target.cs files.
+
+* New: Make Weak Object Ptr function, which constructs a Weak Object Ptr from a C++ pointer by deducing its type. Fixed Weak Object Ptr's constructor which implicitly removed const.
+
+* New: Upgraded the Array function, Insert, to allow it to take an array with any allocator type. It can also move elements out of another Array.
+
+* New: Array now has a "_GetRef" variants for many of its functions, which will create an element and return a reference to it rather than its index in the Array. These functions are: “Insert_GetRef”, “InsertZeroed_GetRef”, “InsertDefaulted_GetRef”, “Emplace_GetRef”, “EmplaceAt_GetRef”, “Add_GetRef”, “AddZeroed_GetRef” and “AddDefaulted_GetRef”.
+
+* New: There is a new Multi Map function, Append, to match the Map class’ function of the same name.
+
+* New: Base64 encoding library now supports ANSI/wide characters, decoding streams without padding marks, and decoding into preallocated buffers without writing dummy padding bytes.
+
+* New: Integer support added to Is Aligned.
+
+* New: Static asserts to stop alignment functions (Align(), IsAligned(), etc.) being called with non-intergal, non-pointer types.
+
+* New: Visual Studio Code project file generation now creates a .code-workspace for the workspace. Allows foreign projects to "mount" the UE4 folder so that the engine tasks are available, and all engine source is visible to VSCode for searching purposes
+
+* New: Visual Studio Code - Add launch task for generating project files for the given folder
+
+* New: Ensure that all branches within Generic Platform Misc’s Root Dir function produce an absolute path with no duplicate slashes. Remove relative-to-absolute conversion of the root directory from the Paths function, Make Standard Filename, now that we know Root Dir always returns an absolute path.
+
+* New: Added Net Serialize to Date Time.
+
+* New: Vector 4 now features a -= operator.
+
+* New: Allowed Visit Tuple Elements to take multiple tuples of the same size. These are now iterated in parallel and will invoke the visitor with multiple arguments.
+
+* New: Minimal Name now has an Is None function.
+
+* New: Relaxed memory ordering atomic operations.
+
+* New: Atomic object wrapper has been created. This guarantees atomic access to the underlying object, with additional memory ordering guarantees.
+
+* New: Platform Atomics now has an Atomic Read function. 
+
+* New: Added overloads of atomic functions to handle int8, int16, and int64, in addition to the existing int32.
+
+    * Switched Android and Linux builds over to use the Clang implementation.
+
+* New: Command-line option "-NOEDL" will disable EDL.
+
+* New: Command-line "-diffonly" mode added to cook commandlet. This mode compares packages on disk that have already been cooked with packages currently being cooked to detect deterministic cook issues.
+
+* New: Added As Const helper function, for const-ifying pointers and references.
+
+* New: Checks to prevent Flush Async Loading and Load Object / Load Package from being called from any threads other than the game thread.
+
+* New: Tickable Base Objects can now avoid calling Is Tickable every frame if they will never or always be ticked.
+
+* New: Denormalized values are now allowed when converting float32 to float16.
+
+* New: Added static asserts to Logf and checkf to ensure that they're always called with a literal formatting string.
+
+* New: Cook metadata like DevelopmentAssetRegistry.bin will no longer be packed into a shipping game.
+
+* New: Added a new target rule, "Use Inlining", to UnrealBuildTool to disable inlining for debugging purposes - defaults to true.
+
+* New: Opened .pak files can be listed with the console command "DumpOpenedFiles".
+
+* New: Added a new command line option, "-statnamedevents", for enabling named events.
+
+* New: Json Object Converter now supports sets in its Json Value To UProperty function.
+
+* New: Mutable lambda functions can now be bound to delegates.
+
+* New: To go along with changes to the Size Map tool, the way that ResourceSize is calculated has changed. Exclusive is the same but Inclusive is now EstimatedTotal and represents the total size of the object and all of its sub objects, but not any external assets it references. It now tries to estimate the size on a generic cooked platform so it will handle texture LOD bias better.
+
+* New: The Platform Stack Walk function, Capture Stack Backtrace, now returns the stack size rather than checking for the first null pointer, to prevent truncated callstacks if parts of the stack are zeroed out.
+
+* New: Improved "caching file" message in networkplatformfile so it says "Requesting file..." and is only output when we actually request the file from the server.
+
+* New: Improved Variant's support for move semantics. Switched Variant Types over to an enum class to aid debugger visualization.
+
+* New: When cooking with Generate Chunks and the asset manager enabled, AllChunksInfo.csv is now written out to the CookMetadata directory. This file can be used to look for historical chunk size changes, and it now uses the per-platform cooked size.
+
+* New: Added more debug info for ensure in Linker Save when a name that has not yet been mapped is being serialized.
+
+* New: Tweaked "LogInit" command-line logging to clarify when the command line is incomplete due to some arguments being filtered out.
+
+* New: Promoted Statistical Float (a utility class to track basic statistics of a real number data series with constant storage) into Core from FunctionalTesting
+
+* New: Made optimizations to Paths function, Is Relative.
+
+* New: Improved cook performance by optimizing the EDL assumption verification code.
+
+* New: Removed all locks from EDL Cook Checker to improve Save Package performance.
+
+* New: Optimized adding and removing objects to GC Object Referencer.
+
+* New: Switched our allocator to ANSI when running a Thread Sanitizer build, for better error reporting.
+
+* New: Improved filename to long package name resolution when providing a relative path.
+
+* Removed: Move-unaware allocators are no longer usable with containers.
+
+* Removed: Startup Packages have been removed, along with the associated code.
+
+* Removed: Removed deprecated delegate functions.
+
+#### Editor
+
+* Crash Fix: Fixed a crash when opening the level viewport context menu if the actor-component selection is out of sync.
+
+* Crash Fix: Details panel no longer crashes after compiling blueprints that have edit condition properties.
+
+* Crash Fix: The .fbx scene importer no longer causes a morph target import crash.
+
+* Bug Fix: Added a value threshold for Paint and Fill values while using the Cloth Painter to prevent negative numbers.
+
+* Bug Fix: Fixed Camera Rig Rail Preview Meshes to match Spline Component .
+
+* Bug Fix: Added ".vscode" folder to generated .gitignore file for new projects.
+
+* Bug Fix: Fixed an issue with content browser columns not sorting properly.
+
+* Bug Fix: Made the font editor canvas respect DPI scale.
+
+* Bug Fix: Fixed a crash when dragging and dropping a non-array UI element onto array element..
+
+* Bug Fix: Fixed "Game Gets Mouse Control" setting.
+
+* Bug Fix: Duplicate actors are no longer created for USD primitives that specify a custom actor class.
+
+* Bug Fix: Moving a copy of an Actor to a different persistent level or sublevel now only dirties the target level.
+
+* Bug Fix: Clamped "Assets to Load at Once Before Warning" so it cannot be set below 1.
+
+* Bug Fix: Fixed auto import to import into the correct map folder.
+
+* Bug Fix: Fixed reordering of materials when importing a LOD containing new materials.
+
+* Bug Fix: Fixed a crash using Actor Details with the struct details panel.
+
+* Bug Fix: Fixed .fbx importer scale not always applying.
+
+* Bug Fix: The "Add New C++ Class" dialog now remembers the previously-selected module.
+
+* Bug Fix: Source Control Submit Files now interprets Escape key as if the user clicked "cancel".
+
+* Bug Fix: Fixed a leak in the Git plugin.
+
+* Bug Fix: Fixed Resave Packages working poorly with projects on other Windows drives with names like X, Y, Z.
+
+* Bug Fix: Prevented hit proxies from being rendered for cursor query when the cursor is not visible.
+
+* Bug Fix: Source control submenu menu customization
+
+* Bug Fix: Fixed window size and position adjustment not accounting for primary monitor not being a high DPI monitor when a secondary monitor is.
+
+* Bug Fix: Expanded array elements no longer collapse after reordering with drag-and-drop.
+
+* Bug Fix: Fixed a bug where Undo History would be cleared up when starting a multiplayer game in editor with 2 or more players.
+
+* Bug Fix: Fixed a crash when additional level viewport sprites are added after the level editor module is loaded.
+
+* Bug Fix: Fixed pending netgame travels in Play in Editor games erroneously stating there was a problem browsing to the pending map.
+
+* Bug Fix: Opening a color picker no longer causes values to change each time.
+
+* Bug Fix: Source control module now properly allows you to select a workspace when a dialog is up (e.g. the save asset dialog).
+
+* Bug Fix: Fixed a bug that caused Visual Studio to open repeatedly.
+
+* Bug Fix: Prevented importing more than MAX_SKELETAL_MESH_LODS when importing a skeletal mesh.
+
+* Bug Fix: Array reordering in a Blueprint now propagates to placed instances.
+
+* Bug Fix: Fixed a crash that could happen while hot-reloading content in the editor.
+
+* Bug Fix: Fixed levels not being lockable or unlockable if they are not checked out.
+
+* Bug Fix: Fixed HLOD auto-clustering incorrectly including transient and template actors.
+
+* Bug Fix: The Enter key now works as expected for details panel asset pickers.
+
+* Bug Fix: Removed circular referencing to parent window when creating an asset from an object.
+
+* Bug Fix: Perforce Plugin: Removed all calls to methods that would store passwords on the users local machine.
+
+* New: Added "favorite folders" to the content browser. You can add or remove folders from the Favorites list in the folder's right-click context menu, and hide or show the Favorites list in the Content Browser options.
+
+* New: Added support for the replay (demo) system, both record and playback, in PIE sessions. Demos from PIE sessions are compatible with those made in standalone instances, and vice versa.
+
+* New: Added the ability to have conditionally overridden HLOD settings from project-global defaults.
+
+* New: Created an Actor Component to decode Linear timecodes from an audio source. Outputs timecode, or timecode event.
+
+* New: Added an editor setting for the shared DDC location to use.
+
+* New: Improved cook process asset error handling including content browser links and better failure notification
+
+* New: Added support for canvas-based DPI scaling.
+
+* New: If a shared DDC is not being used with Perforce, a notification with a link on how to setup a shared DDC will be presented.
+
+* New: Unattended commandlets can now rename and save packages.
+
+* New: Added camera speed scalar setting and Toolbar UI to increase range on camera speed presets.
+
+* New: Added a setting to enable/disable high DPI support in editor.
+
+* New: Replaced the skeletal mesh import option "keep overlapping vertex" with thresholds, enabling the user to control the vertices’ welding thresholds.
+
+* New: The Persistent Level in the Levels tab can now be locked.
+
+* New: Added the Auto Checkout function to Asset Rename Manager for commandlet usage.
+
+* New: The realtime level editor viewport is disabled if running under remote desktop.
+
+* New: Added the project setting "bValidateUnloadedSoftActorReferences" to the Blueprints section. This setting is true by default. Setting it to false it will disable the code that loads assets to check for references when deleting/renaming an actor.
+
+* New: Added documentation tooltips showing type information to header columns of DataTables.
+
+* New: Added Rename Assets With Dialog to Asset Tools. All previous implementations now use the "With Dialog" version. Delete Object and Duplicate Asset can be completed without dialog.
+
+* New: Added a variable called GIsRunningUnattendedScript. If true, we are running an editor script that should not prompt any modal dialog. Can be set for commandlets with "-RUNNINGUNATTENDEDSCRIPT" command-line option.
+
+* New: Many improvements to the Asset Audit, Size Map, and Reference Viewer tools to support their use in analyzing cooked content
+
+* New: Added a function. Get All Content Browser Command Extenders to Content Browser Module that supports registering commands/keybinds to extend the content browser via plugins.
+
+* New: Added "Don't show this again" checkbox to warning about deleting actors referenced by HLOD Actors.
+
+* New: Added new options to the "ResavePackages" commandlet when using "-BuildHLOD":
+
+    * "-ForceEnableHLOD" turns on Hierarchical LOD in the level automatically.
+
+    * "-ForceSingleCluster" turns on the Force Single Cluster for Level" option automatically.
+
+    * "-ForceHLODSetupAsset=<asset path>" clobbers the level's customized HLOD asset with the specified one..
+
+    * "-SkipToMap" will resume generating HLODs at the specified map name (going in the order of map files found while searching folders).
+
+* New: In the glTF importer plugin: Import Static Mesh, Material, and Texture assets from the Khronos glTF 2.0 format. Text (.gltf) and Crash Fix: Fixed frequent crash when opening Config Editor more than once per session.
+
+* New: Added support for binary (.glb) files are supported.
+
+* New: Improved morph target import time.
+
+* Removed: Removed .fbx export file version 2010 compatibility. The 2018 .fbx SDK does not export to 2010 and earlier file formats.
+
+##### Content Browser
+
+* Crash Fix: Fixed a crash when renaming an asset without an enabled source control via the Asset Tool.
+
+* Crash Fix: Fixed a crash that could happen when content hot-reloading a Level with streaming sub-Levels. 
+
+* Bug Fix: Fixed the asset search suggestions list closing if you clicked on its scrollbar. 
+
+* Bug Fix: Selection no longer resets when changing certain view settings.
+
+* Bug Fix: Fixed the row structure tag not appearing for Data Table assets.
+
+* Bug Fix: Fixed an issue so that we verify a folder is newly created before removing from the tree list.
+
+* Bug Fix: Fixed an issue with dragging and dropping folders.
+
+* New: Added a Clear option for the Content Browser folder search box when using the escape key.
+
+* New: Improved responsiveness of Open Asset dialog and when searching or filtering assets.
+
+* New: Added the ability for the Reference Viewer to show and copy references lists for nodes with multiple objects, or multiple selected nodes.
+
+##### Foliage
+
+* Crash Fix: Fixed a crash for Instanced Static meshes and Foliage when used with a Static Mesh that had an invalid lightmap coordinate index set.
+
+* Crash Fix: Fixed a crash that could happen when replacing placed instances with another asset type.
+
+* Bug Fix: We now properly update random streams for Static Mesh instancing only when adding new instances or from saved data.
+
+* Bug Fix: Fixed a Static Mesh Instancing issue where the proper index would not be removed after multiple add and removes.
+
+* Bug Fix: When duplicating Static Mesh Instances through the Details panel, the source transform is now correctly used for the result.
+
+* Bug Fix: Fixed an issue with a special case where building a Blueprint that used a Foliage type would remove all visible Foliage already painted.
+
+* Bug Fix: When moving Foliage from one level to another, if there was an error during the process, it would not abort the process gracefully. 
+
+* Bug Fix: Fixed a case where having very low FPS would cause the foliage to disappear for a few frames while moving the camera. 
+
+* Bug Fix: Fixed an issue that prevented foliage addition during PIE or Simulate mode.
+
+* Bug Fix: Fixed a special case where adding and removing multiple times would cause the regeneration of random streams for all changed instances.
+
+##### Landscape
+
+* Crash Fix: Fixed a crash when updating Landscape grass.
+
+* Crash Fix: Fixed a crash where you would use the Spline Component with no point specified.
+
+* Crash Fix: Fixed a case where it was possible that doing undo on Landscape would cause an assert or crash.
+
+* Bug Fix: Fixed an issue that could cause Landscape shaders to be recompiled very often during the cook process.
+
+* New: Added the possibility to update the Landscape Material at runtime using Material Instance Dynamic Materials.
+
+##### Material Editor
+
+* Crash Fix: Fixed a crash when filtering Material Layers drop-down with no set value.
+
+* Crash Fix: Fixed a crash caused by the Layered Material property widget of the Material Instance Editor.
+
+* Crash Fix: Fixed an assert on expanding vector param in Layered Material.
+
+* Bug Fix: Fixed an issue to not show empty Global group if it only contained the Material Layer parameter.
+
+* Bug Fix: Fixed an issue with Details panels to now "own" a color picker so that a different Detail panel refreshing doesn’t close it. This fixed the refreshing state of the graph after changing the texture or color parameter values. 
+
+* Bug Fix: Fixed an issue to allow visible parameters parameter retrieval to correctly recurse through internally called functions. Previous checks were intended to prevent function previews from leaving their graph through unhooked inputs but unintentionally blocked all function inputs. 
+
+* Bug Fix: Fixed an issue with Material overrides to display the actually used values even when greyed out.
+
+* Bug Fix: Fixed a reset to default issue on Material Layer Parameter node and layers.
+
+* Bug Fix: Changing a parameter in the Material Parameter Panel now dirties the Material correctly. Changing a layer parameter in the Material Instance Editor now refreshes the details panel as well.
+
+* Bug Fix: Fixed an issue with thumbnails for Material Function Instances not being able to delete new function instances.
+
+* Bug Fix: Empty Material groups are now being pruned correctly.
+
+* Bug Fix: Toggling "Show Background" now updates the background correctly.
+
+* Bug Fix: Material window now uses the background color set in preview scene.
+
+* Bug Fix: Fixed an issue to properly display the Material node preview label is now centered on HiDPI screens.
+
+* New: Turned on Material Layering by default.
+
+    * Adding parent field and asset field to Material Layers and Material Layer Blends. Setting the parent restricts the assets available, and setting the layer or blend asset fills in the parent as well.
+
+    * Added Content Browser filters for Material Layers, Material Layer Instances, Material Blends, and Material Blend Instances.
+
+* New: Added a Parameter Defaults tab to the Material Editor so users can see and change the default value of all the parameter nodes currently in the Material Graph, also added Parameter Defaults tab support for Material Functions.
+
+* New: Save to Child and Save To Sibling buttons now work for material layer parameters and are shown on the material layer parameter panel.
+
+* New: Creating a new parameter immediately allows naming.
+
+* New: Added axis input to the material editor's viewport client to allow rotating the preview viewport.
+
+* New: Material Blends now have plane mesh previews in their icons.
+
+* New: Added Material Layers Parameters Preview (all editing disabled) panel to the Material Editor.
+
+* New: Layers in Material Layer Functions can now be renamed, except Layer 0, which is always named "Background". 
+
+* New: Added access to a Material Function Instance's parent in the Material Instance Editor.
+
+* New: Disabled being able to create a layer or blend in the asset drop-down, sections of the stack that have been disabled now inactivate that part of the UI.
+
+* New: The create a Layer or Blend in asset drop-down menu has been disabled. Sections of the stack that have been disabled are no inactive for that part of the UI.
+
+* New: Create Function Instance now indicates if you are making a Layer or a Blend.
+
+* New: Updated the optional example Layers and Blends checkbox in the Experimental Settings so that it makes basic Layers and Blends when you create new assets.
+
+* New: The Parent field in Layer and Blend Instances is now editable.
+
+* New: Any applicable Layer and Blend compile warnings now also appear in the Layer or Blend itself. We enforce only two layers in a blend and prevent users from deleting Layer or Blend output nodes.
+
+* Removed: Parent drop-down has been removed from Layers and Clends. A filter button has been added instead.
+
+##### Media Framework
+
+* Crash Fix: Fixed a crash with Media assets in Media Sound Component’s "Update Player."
+
+* Crash Fix: Fixed a crash when opening AVF Media after opening multiple media players.
+
+* Crash Fix: Fixed a crash due to a race condition by adding a missing scope lock to WMF Media.
+
+* Crash Fix: Fixed a crash for media textures that could happen if the media player is generated from GC clustered Blueprint.
+
+* Bug Fix: Fixed external textures not registering correctly in all cases for the Engine.
+
+* Bug Fix: Fixed Media Asset issue with Media Player’s Previous and Next not playing automatically if the player was playing previously. 
+
+* Bug Fix: Fixed an issue with image media player never finishing initialization if loading failed.
+
+* Bug Fix: Fixed an ensure when cooking projects that use Image Media Source.
+
+* Bug Fix: Fixed media sessions that were no generating Playback End Reached event when session forced to stop for WMF Media.
+
+* Bug Fix: Fixed an issue with video capture on Windows 7 where the feed would display black or error out when selecting a different track and format.
+
+* Bug Fix: Fixed issues with tracks being reported in reverse order.
+
+* Bug Fix: Fixed an issue where external textures referenced by a Material before being associate with a Media Player never having their uniform expressions reached. 
+
+* New: Added simplified track switching to enable seeking.
+
+* New: Added a notification if no video track is available or selected to the Media Player Editor.
+
+* New: Added the ability to disable video output when no video track is selected for Image Media.
+
+* New: Added verbose logging to the Media Player BP functions for Media Assets.
+
+* Removed: Looping option has been removed from playlists for Media Assets.
+
+##### Sequencer
+
+* Crash Fix: Fixed a copy/paste crash that could happen when only processing Movie Scene Copyable Binding and objects that can be spawned by the movie scene spawn were registered. 
+
+* Crash Fix: Fixed a crash when an existing hotspot is null.
+
+* Crash Fix: Fixed a crash when pasting a camera cut track that happened when the outer track isn’t set to the movie scene. 
+
+* Crash Fix: Prevented a crash by ensuring that the null object in property actuate when there’s a track without a binding object. This occured when pasting a property track to the root. 
+
+* Crash Fix: Fixed a crash when dragging a Level Sequence into the tree area.
+
+* Crash Fix: Fixed a crash caused by erroneous persistent references to FSequencer.
+
+* Crash Fix: Fixed a crash when calling Set Playback Position in an event when adding Set Playback Position as a latent action.
+
+* Crash Fix: Prevented a reload crash that could happen if a Sequence was being recorded in a currently opened Sequencer. On Recording Started and On Recording Finished are now set to multicast delegates.
+
+* Crash Fix: Added a nullptr check to fix a crash in Actor Sequence when it doesn’t exist.
+
+* Crash Fix: Fixed a crash opening a new Sequence while another is still active.
+
+* Crash Fix: Fixed a crash when an Actor Factory is not found.
+
+* Crash Fix: Added null checks to prevent a crash when saving the default state of a spawnable.
+
+* Bug Fix: Fixed an issue to unregister missing track editors in Movie Scene Tools.
+
+* Bug Fix: Fixed an issue with quaternions > rotation > quaternions so that unwinding is only processed if the last transform is valid.  
+
+* Bug Fix: Fixed a rotation rotator to quaternion to rotator conversion which prevents you from typing in a rotation of 0,0,320 into the Key Editor. Unwind rotations from the previous transform to the current transform so that the nearest rotation is set rather than the rotator to quaternion to rotator conversion. 
+
+* Bug Fix: Sequencer: Set min/max values for generic key area so that they don't default to 0,10.
+
+* Bug Fix: Fixed an issue returning unhandled only if not dragged. There was an issue where dragging in the track area would sometimes leave the handled state with the time slider controller to not allow you to pop up a menu with the movement tool. 
+
+* Bug Fix: Fixed a problem to now defer Details panel updates on scrubbing and playing.
+
+* Bug Fix: Fixed an issue by adding a null check and warning for invalid Get Parameter Collection Instance.
+
+* Bug Fix: Fixed an issue to set row index when creating a new take
+
+* Bug Fix: Fixed issues with  pre-edit property chain broadcast so that the property path will include possible struct/array node.
+
+* Bug Fix: Fixed all light colors being set if any color property is animated.
+
+* Bug Fix: Fixed spawnables not playing back when Editor Preview Actor was set to false for Sequencer spawnables so that Begin Play doesn't get skipped.
+
+* Bug Fix: Fixed an issue to reset drop node on drag leave.
+
+* Bug Fix: Fixed a delay at shot boundaries causing Sequences to not play back and render out. The shot ID needs to be tracked to determine whether a new shot is encountered.
+
+* Bug Fix: Fixed a delay before warm up causing Sequences not to render.
+
+* Bug Fix: Fixed an issue to enable the camera cut track when popping back to the master only if there’s a camera cut track in the Master Sequence. It addresses an issue where if you don’t have a camera cut track in the Master, the camera gets locked to a camera cut in a sub-scene and you can’t toggle out of it. 
+
+* Bug Fix: Pivot locations are now updated for a selection when closing Sequencer.
+
+* Bug Fix: Fixed an overlap issue by looking at the key behind it. It addresses if there’s 3 keyframes on consecutive frames, so if you zoom out, you should see two bordered keys when the overlap threshold is passed. 
+
+* Bug Fix: Fixed UMG animations not working with Blueprint Nativization.
+
+* Bug Fix: Fixed an issue in Sequencer to remove null tracks on object bindings. Tracks can become null if they're from a plugin and the plugin is disabled.
+
+* Bug Fix: Moved some commands out of the generic Sequencer command bindings so that they don't take over the viewport. For example, end for "Snap to Floor" should still function in the viewport.
+
+* Bug Fix: Fixed an issue to restore pre-animated state when changing active channels.
+
+* Bug Fix: Fixed an issue to suspend broadcast of selection delegates when performing multiple operations.
+
+* Bug Fix: Display names for Shot are now switched to FString so that it's not localized.
+
+* Bug Fix: Fixed an issue with Sequence Recorder to close the target animation asset editor if it exists before recording into it. 
+
+* Bug Fix: Fixed an issue with Sequence Recorder to record to the specified target animation for the target Actor only. Newly tracked components will have newly created animations so that they don’t record to the same target animation assets.
+
+* Bug Fix: Sequencer now uses drag and drop actor factory if specified. This ensures that the correct Actor Factory is used in creating the object template for the Sequencer spawner. It fixes some spawnables not getting created properly (for example, an Empty Actor).
+
+* Bug Fix: We prevent throttling on the Curve Editor so that the Editor world tick can apply. It fixes the viewport not updating sometimes when adjusting curves in the Curve Editor. 
+
+* Bug Fix: Fixed an issue where parameters were not used. 
+
+* Bug Fix: Fixed WMF Media for IYUV encoded AVIs not playing correctly. 
+
+* Bug Fix: Fixed WMF Media that caused H.265 frames to be dropped due to false negative buffer size checks.
+
+* Bug Fix: Fixed an issue in Sequencer to always save the default spawnable states regardless of the focused Sequence. Addresses issue if you step back to the Master Sequence (and the spawnable if it still exists) and then scrubs outside the region where the spawnable exists. It gets destroyed but saved default spawnable state doesn’t get called because it’s no longer the focused Sequence.
+
+* Bug Fix: Fixed Sequencer viewport invalidation so that it happens on Sequence evaluation.
+
+* Bug Fix: Fixed an issue when scrolling in track area.
+
+* Bug Fix: Fixed an issue with Sequence Recorder causing duplicate Actor triggers before playing so that the Sequence can be recorded and played back at the same time. 
+
+* Bug Fix: Fixed an issue with Sequence Recorder when recording spawnables not getting the correct position for being spawned at.
+
+* Bug Fix: Fixed an issue with Sequence Recorder to record spawn Actors immediately so that they won’t be missed if they’re deleted before Tick.
+
+* Bug Fix: Fixed a problem where Sequencer was not updating after toggling Bind Sequencer to PIE/Simulate while PIE is active.
+
+* Bug Fix: Addressed an issue in Sequencer to find available non-overlapping row index when adding sub-sections. 
+
+* Bug Fix: Fixed an issue with Sequence Recorder to address Level Sequences no triggering when recording. Level Sequences would not get recorded if the World Settings Actor was not recorded.
+
+* Bug Fix: Fixed an issue in Sequence Recorder causing Editor-only components to be recorded.
+
+* Bug Fix: Fixed missing display text for audio and video tracks when switching media. 
+
+* Bug Fix: Fixed a slot animation issue where it was not being restored for montages that are recreated during evaluation. 
+
+* Bug Fix: Fixed the import camera in Sequencer so that when new Cameras are created, values from the FBX are going to the newly created Cameras. Also, added Reduce Keys and Reduce Keys Tolerance to import FBX options. 
+
+* Bug Fix: Fixed an issue in Sequencer when exporting unmatched float properties to custom attributes. 
+
+* Bug Fix: Fixed an issue in Sequencer to override the animation asset in the player state if it doesn’t match the animation asset being evaluated. 
+
+* Bug Fix: Fixed an issue in Sequencer so that we don’t update current time to be within the view range when stepping into a Sequence. 
+
+* Bug Fix: Fixed an issue to skip binding IDs if the Sequence can’t be found. 
+
+* New: Added some options to the Curve Editor for Modify Owner Change to the Curve Owner Interface. We call mark as changed when modifying keys or tangents. 
+
+* New: We now add a transaction for easing length.
+
+* New: Added Copy, cut, paste, and duplicate options to object bindings.
+
+* New: Added the ability to drag sections up in Sequencer.
+
+* New: Sequencer can now compile on the fly.
+
+    * It is now able to compile partially or completely out-of-date evaluation templates from the source data as needed. It affords much more efficient compilation when working within Sequencer.
+
+    * The concept of "instance data" for sub-Sequences, available through Movie Scene Player interface or persistent data stores. It replaces the compilation of specific templates for control rig templates. 
+
+    * Moved sub-tracks and sections to the Movie Scene module. 
+
+    * Removed the concept of shared tracks. Any previous uses should port over to shared execution tokens instead.
+
+* New: Now automatically create a camera cut track if a camera is dropped and there’s no existing camera cut or there’s no existing camera cut sections. 
+
+* New: The sub-Sequence name is now displayed on binding ID pickers rather than sub-section name.
+
+* New: Added an active column for Sequence Recorder.
+
+* New: Added an option to export an object’s transform to a Camera Anim asset.
+
+* New: We now allow blending for vector tracks.
+
+* New: Added an option to bake transforms to keys. 
+
+* New: Added configuration for default completion mode for Movie Scene Sequences. The default Level Sequences is Restore State. All others, such as UMG are set to Keep State. 
+
+* New: We now update auto scroll when moving Keys or Sections.
+
+* New: Added the ability to specify a global transform origin for Level Sequences" 
+
+    * It leverages the new concept of specifying an instance data object to the evaluation that allows systems to inject information into a Sequence at runtime, enabling more dynamic control of the track. 
+
+    * Level Sequence Actors use this by default to supply a dynamic "transform origin" to all component transform tracks, which all absolute transform sections will be added. 
+
+* New: The constraint guid has been switched to constraint binding ID. 
+
+* New: Playback speed has been added to the settings menu. 
+
+* New: Added recording button to Sequence Recorder that adds for any selected Actors. 
+
+* New: Added option to only show selected nodes only. 
+
+* New: Added option to disable camera cuts on the Movie Scene Player.
+
+* New: Added icons in Sequencer for Delete and Rename icons. 
+
+* New: Added default expansion states to allow track editors to specify them per track type. Material track is currently the only track that defaults to expanded.
+
+* Deprecated: We now show binding ID picker customization on all Details panels. This allows creation of new camera cut sections from existing bindings. Add New Camera Cut has been deprecated.
+
+* Removed: In Sequencer, the hotkeys Shift-C and Shift-D for toggling the Cinematic viewport have been removed. This was causing some confusion when users were accidently hitting them. 
+
+##### Static Mesh Editor
+
+* Crash Fix: Fixed an Editor crash when a distributed Simplygon job fails.
+
+* New: Moved the LOD sections dropdown to the LOD picker category. Custom LOD checkboxes are now hidden until "Custom" is checked. The same change applies to the Skeletal Mesh Editor.
+
+* New: Added LOD menu to viewport toolbar.
+
+##### World Browser
+
+* Bug Fix: Fixed issues with re-importing of tiled Landscape data in the Word Composition window.
+
+#### Gameplay Framework
+
+* Crash Fix: Fixed a crash in the Get Asset Registry Tags function in Static Mesh that happened when the Find Object function was called during saving.
+
+* Crash Fix: Fixed a rare crash when child actor components were loaded as part of an async sublevel load.
+
+* Crash Fix: Servers traveling to a world that shares a name with an Asset will log a warning instead of crashing. The full path to the map file can be used to travel successfully in these cases.
+
+* Crash Fix: Fixed an initialization crash when using GameplayTags from global singleton objects.
+
+* Crash Fix: Removed an extra dereference that was causing a crash in the Gameplay Cue Manager.
+
+* Bug Fix: Added tracking of Idle Seconds to the Performance Data Consumer interface’s Frame Data.
+
+* Bug Fix: Fixed the Load Asset blueprint node so it can be correctly called from a loop with different inputs.
+
+* Bug Fix: Player Camera Manager’s Update View Target Internal function will no longer apply values from Blueprint Update Camera if it returns false.
+
+* Bug Fix: Fixed an issue where module redirects could cause the Asset Manager to incorrectly allow Primary Assets with the wrong base class.
+
+* Bug Fix: Fixed several bugs with async loading streaming levels via Streamable Manager, levels loaded this way now correctly handle offsets and navigation.
+
+* Bug Fix: Disabled the duplicate Primary Asset ID ensure for editor-only types like Maps, as it could activate unnecessarily.
+
+* Bug Fix: Exporting a Package without a type now adds warnings to the output logs.
+
+* Bug Fix: Components and Tick functions will no longer be unregistered and re-registered when testing a rename for an Actor.
+
+* Bug Fix: Modified Gameplay Tag’s Import Text Item function to account for redirects when establishing Tag Name.
+
+* Bug Fix: Fixed character proxies doing up to two floor checks when only rotation is changed (there shouldn't be any). Added some optional verbose logging to the Find Floor and Compute Floor Dist functions.
+
+* Bug Fix: String allocations are no longer performed on every Server Move function call that the server runs.
+
+* Bug Fix: A bug with Post Significance not firing when unregistering an object unless it was the last object of its tag has been fixed.
+
+* Bug Fix: The engine will now only create a Material Instance Dynamic as a child of the calling object if the construction script is running. Otherwise, it will be created in the transient package.
+
+* Bug Fix: Fixed clients bypassing network move combining and RPC throttling when there is no acceleration and velocity is non-zero (for instance, when falling or decelerating). Changed velocity condition on whether moves may be combined to a delta from or to zero velocity.
+
+* Bug Fix: Fixed "infinite jump" exploits from rogue clients. The server no longer incorrectly resets the Jump Key Hold Time variable while not changing the Was Jumping variable.
+
+* Bug Fix: Log Gameplay Tags is now exposed.
+
+* Bug Fix: Disabled HLOD volume overlap events, and set to not load on client or server.
+
+* Bug Fix: Optimized client movement to avoid updating skeletal mesh bones on the client when teleporting back to combine a network move to the server. This can be avoided since the following Perform Movement call will move the mesh and cause the update there.
+
+* Bug Fix: Added stat for Force Position Update (when server ticks client that has timed out) so that it doesn’t show up as Perform Movement.
+
+* Bug Fix: Redirectors or otherwise invalid .umap files are no longer considered sub levels for world composition.
+
+* Bug Fix: All Actor subclasses now correctly enforce not ticking when tick type is Viewports Only, even if they override TickActor. The Should Tick If Viewports Only function can override this by returning true.
+
+* Bug Fix: Added several transform related functions to the HIDE_ACTOR_TRANSFORM_FUNCTIONS macro and fixed the signatures on some that were already there.
+
+* New: Added a new plugin for performing data validation. Data validation can be run on individual assets, multiple assets or whole directories. It can also be run as a commandlet and included in your build process.
+
+* New: Added a function called Return To Main Menu to the Game Instance class. Games may call function to return to their main menu reliably, even in cases where a Player Controller isn't available to call Client Return To Main Menu With Text Reason.
+
+* New: Data Tables now support getting the value of each row for a specific column/property, excluding the title of the column.
+
+* New: Improved the editor UI for Primary Asset ID properties. It now uses the same one as Pins and better handles objects that aren't already loaded.
+
+* New: Added functions to Spring Arm Component to get relevant status of the collision test.
+
+* New: Added a To String to the Lex class for physics enums Collision Trace Flag, Physics Type, and Body Collision Response.
+
+* New: Added "p.NetForceClientServerMoveLossPercent" cvar to simulate loss of client-to-server movement RPCs (without hosing the rest of networking).
+
+* New: The Player Input function Massage Axis Input can now be overridden in subclasses.
+
+* New: A new variable, Should Manager Determine Type And Name, has been added to Asset Manager settings. It is disabled by default. If enabled, the Asset Manager will be able to load assets that do not have a native class that implements Get Primary Asset ID, in both the editor and a cooked game. This will behave similarly to Should Editor Guess Type And Name, but in exchange for some runtime performance it will work properly in a cooked build.
+
+* New: Primary Data Asset can now be directly blueprinted to enable blueprint-only Primary Assets, even if Should Manager Determine Type And Name is disabled. If you do this, the Primary Asset Type will be the blueprint below Primary Data Asset, which you can then inherit from to create your Primary Assets.
+
+* New: If overrides of Actor Component’s Begin Play function are not routed up the hierarchy code, the engine will now use Ensure to warn that initialization steps were missed.
+
+* New: Alpha Blend struct is now marked as BlueprintType.
+
+* New: Gameplay Abilities now pass along their Source Object when creating the effect context for a new effect.
+
+    * Consolidated all of the events on Active Gameplay Effect into a single struct.
+
+* New: Added an active effect event for when the inhibition of an effect changes, meaning that the effect is no longer relevant, but had not actually been removed.
+
+* New: Added functionality to cycle between targets with Page Up and Page Down for 'showdebug' commands. List of targets is contextual (For example 'showdebug animation' will consider all visible actors with an AnimGraph). Current debug Target is highlighted in a green bounding box.
+
+* New: Performance optimization: Added Network Skip Proxy Prediction On Net Update to the Character Movement Component. This skips immediate forward prediction for proxies on the frame they receive a network update, avoiding all forward prediction sweeps and floor checks on those updates. Intermediate frames will interpolate with prediction. This can also be disabled globally by setting the CVar "p.NetEnableSkipProxyPredictionOnNetUpdate" to 0.
+
+* New: Reduced overhead of evaluating streaming level volumes when all Player Controllers have Is Using Streaming Volumes set to false.
+
+* New: Significance Manager’s Update function now takes viewpoints in as an Array View instead of a const Array reference.
+
+* New: When keys with unknown character code are received by the input system, a new Key struct will now be dynamically created to represent it rather than it being rejected as an invalid key.
+
+* New: Significance Managers’ Post Significance Update functions can now specify to be executed sequentially on the game thread instead of only as concurrently in the parallel for.
+
+* New: Conditional Init Axis Properties is now protected instead of private, so it can be seen by child classes.
+
+* New: The Smooth Mouse function is virtual, so games can implement their own smoothing code.
+
+* New: Added Rand Point In Circle function. It returns a random point, uniformly distributed, within the specified radius.
+
+* New: Updated ensures to give more information when there is an invalid Ground Movement Mode.
+
+* New: Added a new trait called Post Script Construct. This enables users to perform initialization after blueprints have created a temporary object.
+
+* New: Hooked up gameplay tag containers to use Post Script Construct to populate parent tags.
+
+#### Localization
+
+* Crash Fix: Fixed a crash if a break iterator outlives ICU during shutdown.
+
+* Crash Fix: Fixed a crash if a regex pattern or matcher outlive ICU during shutdown.
+
+* Bug Fix: Translation Editor now shows stale translations as "Untranslated".
+
+* Bug Fix: Fixed Roboto Light for certain Russian characters.
+
+* Bug Fix: Gather Text From Source no longer causes the commandlet to fail.
+
+* Bug Fix: Fixed some case-sensitivity issues with Text format argument names/pins. This change restores their original behavior.
+
+* Bug Fix: Fixed game preview language not updating from the native text when switching between preview languages.
+
+* Bug Fix: Fixed game preview language not updating in realtime while PIE was enabled.
+
+* Bug Fix: Fixed struct instances sometimes gathering text that was identical to a parent instance.
+
+* Bug Fix: Fixed culture being incorrect when added via the Localization Dashboard.
+
+* New: Support for culture correct parsing of decimal numbers has been added.
+
+* New: Implemented Always Sign number formatting option.
+
+* New: Added an improved Korean font to the editor.
+
+* New: Added an improved Japanese font to the editor.
+
+* New: Improved the ability to lockdown available game languages. In addition to the previous "DisabledCultures" array, you can now add an array of "EnabledCultures".
+
+* New: Added a warning for duplicate package localization IDs when gathering asset localization.
+
+* New: Added a way to display and reset the package localization ID via the Content Browser.
+
+#### Networking
+
+* Crash Fix: Fixed a rare crash that could occur on servers where replicated, placed-in-map actors could be destroyed, and clients could send references to these actors in server function parameters.
+
+* Crash Fix: Fixed a crash caused by not clearing the send buffer after closing the connection when trying to send data before the handshake is completed.
+
+* Crash Fix: Fixed a client-side assert when trying to send RPCs while the connection is pending closure due to queued bunches.
+
+* Crash Fix: Fixed several errors and asserts that could occur when using network encryption.
+
+* Bug Fix: Removed an ensure that could fire inappropriately in some cases during playback of a replay.
+
+* Bug Fix: Made a change to TCP Listener to prevent it from polling while waiting for connections.
+
+* Bug Fix: Fixed an issue where the single process Play-In-Editor dedicated server could tick incorrect objects.
+
+* Bug Fix: Fixed a replication issue with network startup actors that are renamed on the server.
+
+* Crash Fix: Fixed a crash that could occur when replicating objects during a Play-In-Editor game.
+
+* Bug Fix: Steam voice now continues to work after seamless travel.
+
+* Bug Fix: Changed the adaptive network update frequency option to be disabled by default. This was causing annoying and difficult-to-debug replication artifacts if network update rates weren't managed carefully.
+
+* Bug Fix: Fixed a socket send bug that could cause large files to fail to transfer from the cook on the fly server.
+
+* Bug Fix: UDP Messaging serialization notification is now thread-safe.
+
+* Bug Fix: Fixed a replication checksum issue with Slate brushes.
+
+* Bug Fix: Added stats to Net Driver’s Tick Flush function and analytic stat-gathering blocks within that function.
+
+* Bug Fix: When a client informs the server that it has loaded a streaming level, we now force a net update on all dormant actors that have, at any point, replicated data to relevant clients and ensure that the connection's destroyed startup/dormant actors list is properly populated.
+
+* Bug Fix: Fixed an issue that could cause character movement to look jittery in replay playback in some cases.
+
+* Bug Fix: Fixed a bug that could cause replicated actors that are destroyed while dormant to never be destroyed on clients.
+
+* Bug Fix: Fixed a server crash that could occur if a client sent a reference to an object who's outer object had already been garbage collected.
+
+* Bug Fix: GetGameLoginOptions to URL for split screen players was missing, it has now been added for consistency with first player login flow
+
+* Bug Fix: Fixed a rare server crash when actors in the process of being destroyed could erroneously be re-added to the network object list.
+
+* Bug Fix: Actors are now removed from all net drivers' object lists instead of just the one matching the net driver name stored on the actor.
+
+* Bug Fix: Fixed websocket not providing information when the peer closes the connection.
+
+* Bug Fix: Fixed a rare server crash that could occur if network encryption was in use.
+
+* Bug Fix: Fixed a rare crash due to a null game mode that could occur when joining a server.
+
+* New: Added a local file replay streamer that supports playback, recording, compression, and replay events.
+
+* New: Added a way to view a dedicated server's network stats on a client in realtime (useful for seeing server stats for cloud hosted servers).
+
+    * Changed some server stats to be the total for the server instead of the stats for Connection Zero.
+
+* New: Added a console command ("net dumpdormancy") to view the dormancy status of Actors. Improved the names displayed when editing Net Dormancy properties.
+
+* New: Switched property type comparisons from Strings to Names for speed.
+
+* New: Optimized uint replication by removing unnecessary memory management.
+
+* New: Optimized updating level visibility state between client and servers. When possible, we now only send a single RPC to communicate which levels are visible or should be visible rather than separate RPCs for each level.
+
+* New: Added a new console variable net.Max Connections To Tick Per Server Frame, which can be set to limit the number of client connections a server processes every frame. When many clients are connected, this can help maintain the desired framerate on the server.
+
+* New: Added the ability to disable automatic ping replication to all clients in Player State. This can add up to a lot of CPU and bandwidth overhead with many clients connected, since player states are always network relevant. See the Should Update Replicated Ping variable in the Player State class.
+
+* New: Added new global network delegates callback Network Cheat Detected for responding when cheating is detected on a server.
+
+* New: Analytics have been added for oodle compression percentages.
+
+* New: We now use "-multihomehttp" instead of "-multihome" for routing HTTP sockets.
+
+* New: Support has been added to Oodle for iOS, Android, and Switch.
+
+* New: Network optimizations to reduce memory management churn have been implemented, saving about `10%` of net tick time.
+
+* New: Improved handling for "pending connection lost" notification, which is triggered only if a connection is lost and there is no "owning actor" to deal with the connection loss.
+
+    * Added Received Join and Cleaned Up states to connection to make sure that the pending connection lost delegate only fires at the appropriate time.
+
+* New: Games now receive a notification when the Player State integer unique ID is available.
+
+* New: The Player class can handle beacon connections calling Exec functions.
+
+* New: Added the "demo.Loop" console command which will restart playback from the beginning of a replay once the end has been reached.
+
+* New: Added some additional cycle counter stats in networking code to provide more insight when viewing stat profiles.
+
+* New: CurlHttp now logs send/response basics at normal Log level, providing a better intermediate log level between Warning and Verbose.
+
+* New: Added the command-line option "-ExitAfterReplay" to close the client when a replay reaches the end of playback.
+
+* New: Added verbose http logging to CurlHttp when a request is canceled.
+
+* New: The "net.UseEncryptionToken" console variable has been renamed to “net.AllowEncryption” and made more robust.
+
+#### Online
+
+* Crash Fix: Fixed a crash in web browser module destruction related to the freeing of memory via default destructors after Cef Shutdown has been called.
+
+* Bug Fix: Dedicated servers will now respect the voice online subsystem override option. Bug Fix: Fixed a thread safety issue in the user cloud interface for Steam.
+
+* Bug Fix: Fixed LoginFlow module not being included in installed builds of the engine.
+
+* Bug Fix: Is Server Delegate For OSS can now use the play world during PIE even if no context was passed in.
+
+* Bug Fix: Failure to load TCP socket library (possibly due to lacking the protocol file) on Windows is now handled by defaulting to IPPROTO_TCP.
+
+* Bug Fix: The Net Connection variable, Player ID is now set up properly on clients and servers for both beacons and game net drivers.
+
+* Bug Fix: Improved the way HTTP requests handle timeouts.
+
+* New: Included a Rich Presence interface to the Steam Online Subsystem.
+
+* New: Created a test harness for testing an online subsystem's rich presence interface.
+
+* New: Added embedded browser cookie deletion for Google sign-in when Logout is called.
+
+* New: Added the ability to retry HTTP operations by VERB. This enables cloud-save PUT operations to retry automatically.
+
+* New: Removed the "Pragma: no-cache" header from libcurl requests.
+
+* New: You can now enable libcurl debug callbacks in shipping builds to get more information on libcurl errors.
+
+* New: Added the ability to update cacert.pem from Curl's website. ([https://curl.haxx.se/docs/caextract.html](https://curl.haxx.se/docs/caextract.html))
+
+#### Physics
+
+* Crash Fix: Fixed a crash that would delete max distance mask from a clothing data entry.
+
+* Crash Fix: Fixed a crash when adding empty materials to destructible meshes in the destructible mesh editor.
+
+* Crash Fix: Fixed a crash on destroying Physical Animation Component if the scene had gone away.
+
+* Crash Fix: No longer crashes when switching between clothing assets while painting with a gradient tool that has cached start/end indices.
+
+* Crash Fix: Opening a physics asset that uses a preview mesh with no skeleton no longer causes a crash.
+
+* Crash Fix: Fixed a crash that could happen due to outdated rendering data while editing destructible meshes.
+
+* Crash Fix: Pasting properties in the Physics Asset Editor no longer crashes.
+
+* Crash Fix: Fixed a crash with destructible components cleaning up physics bodies.
+
+* Bug Fix: Removed the deprecated Response To Channels property of Body Instance from non-Editor builds to save memory.
+
+* Bug Fix: Moved constraint and physical anim profile indicators to the top of the details panel.
+
+* Bug Fix: We now update cloth appropriately in the physics asset editor, so physics can interact in real time with cloth.
+
+* Bug Fix: Fixed hang when selecting all bodies or constraints in a physics asset.
+
+* Bug Fix: Copy/paste of physics properties when the skeleton tree has focus now works as expected.
+
+* Bug Fix: Fixed clothing sphere collision detection issue when clothing bounds are flat or almost flat.
+
+* Bug Fix: Welded bodies with negative scale now correctly use the mirrored convex hull.
+
+* Bug Fix: Fixed generation of adjacency buffers for older clothing data.
+
+* Bug Fix: Fixed 'NODEBUG' option for PVD CONNECT console command (PhysX Visual Debugger). It now (correctly) only sends profile data.
+
+* Bug Fix: Fixed calling Update Kinematic Bones in Skeletal Mesh Component if there are no physics bodies
+
+* New: Added "Anim Drive" to clothing assets and runtime interaction for simulations.
+
+* New: Added gravity override support for clothing simulations.
+
+* New: Exposed some PhysicsAsset functionality to be more extensible.
+
+* New: Improved information in Hit Result in case of initial overlap. The Impact Point member variable will now contain information from PhysX on the deepest overlap point.
+
+* New: Added the ability to profile scene query hitches. (See p.SQHitchDetection in code.)
+
+* New: Added 'Set Use CCD' function for changing Continuous Collision Detection option at run-time.
+
+##### Clothing
+
+* Crash Fix: Fixed a crash in clothing actor creation when the clothing simulation mesh has no simulated vertices.
+
+* Crash Fix: Fixed a crash when setting a Poseable Mesh Component as the Master Pose Component for a Skeletal Mesh Component that was using clothing.
+
+* Bug Fix: Fixed non-working tethers in clothing.
+
+* Bug Fix: Clothing configuration now applies to active simulation after editing.
+
+* Bug Fix: Fixed and re-enabled accurate wind mode.
+
+* Bug Fix: Added a clamp to solver frequency to make it impossible to hang the editor with very high clothing solver frequencies.
+
+* New: Changed "bCompileNvCloth" compile option to use the same pattern as “bCompileAPEX” (on by default, disabled on some platforms). This allows game projects to disable it.
+
+##### Collision Detection
+
+* New: Improved "Auto Convex" tool for generating simplified convex collision for meshes.
+
+    * Now uses "Max Hulls" option instead of “Accuracy”.
+
+    * Computation is now done asynchronously.
+
+    * Now using V-HACD version 2.3.
+
+##### Physics Asset Editor
+
+* Crash Fix: Fixed a crash undoing primitive regeneration when stopping an in-progress simulation.
+
+* Crash Fix: Deleting primitives in Physics Asset Editor no longer causes a crash during garbage collection.
+
+* Bug Fix: Constraint setup shortcuts and body-body collision are now undoable.
+
+* Bug Fix: Right-click now displays the context menu for constraints correctly.
+
+* Bug Fix: Fixed dangling lines, poor search focus, and graph not refreshing when making new constraints.
+
+* Bug Fix: Added warning to "Use Async Scene" property when shown in the physics asset editor if the project doesn't currently use an async scene.
+
+* Bug Fix: Tweaked physics asset editor viewport tooltip.
+
+* Bug Fix: Renamed "Create Joints" to "Create Constraints" in Physics Asset Editor.
+
+* Bug Fix: Fix that now allows edits to physics sim options to be undone/redone.
+
+* Bug Fix: Constraints are now displayed only once in the physics asset editor tree view.
+
+* Bug Fix: Fixed not being able to assign/unassign constraint and physical animation profiles from the context menu.
+
+* Bug Fix: Prevented undo/redo breaking when moving both a constraint and a body at the same time.
+
+* New: Added "Show Only Selected Constraints" to the physics asset editor.
+
+* New: Included a keyboard shortcut (Ctrl+T) to toggle between body and constraint selection.
+
+* New: Added the ability to rename shapes in the Physics Asset Editor.
+
+* New: Added back constraint shortcut to PhAT toolbar.
+
+* New: Added error-reporting to the physics profile name widgets.
+
+#### Platforms
+
+* Crash Fix: Updated Intel SPC Tex Comp DLLs to address crashes with some processors on Windows during texture compression.
+
+* Bug Fix: Fixed an issue with vsnprintf on platforms that use the custom implementation found in Standard Platform String to ignore length modifier for certain types (floating point and pointer).
+
+* Bug Fix: Delete existing Source Config File before allocating a new one to prevent them leaking.
+
+* Bug Fix: Fixed compile issues in builds without Low Level Memory enabled.
+
+* Bug Fix: Fixed Update Texture 3D to support  source pitch correctly.
+
+* New: Added missing executable icons for Lightmass and other programs.
+
+* New: Added a descriptive error message when using an Invalid iOS Mobile Provision.
+
+##### All Mobile
+
+* Bug Fix: Fixed an issue for functional code checks being removed for shipping builds in media framework.
+
+* Bug Fix: Fixed an issue that prevented properly flushable logs from assertions and errors on renderthread for Android and iOS.
+
+* New: Disabled Niagara vertex factories for mobile targets.
+
+* New: Disabled manual vertex fetch on mobile.
+
+##### Android
+
+* Crash Fix: Fixed crash if quickly clicking on "Configure Now" button in Android Project Settings multiple times.
+
+* Crash Fix: Fixed a crash in OpenGL Dynamic RHI with Create OpenGL Texture when launching on Mali Galaxy S III.
+
+* Crash Fix: Fixed a crash on application exit of SPRINGBOARD. The process-exit watchdog transgression found several issues if the game is forced closed when the startup movie is playing and the "Wait for movies to complete" is enabled/ 
+
+* Crash Fix: Addressed some crashes related to Experimental Android Virtual Keyboard. 
+
+    * Backspace on an empty line would crash.
+
+    * Virtual keyboard text field isn’t visible after rotating from landscape to portrait. 
+
+    * Chinese and Korean virtual keyboards don’t allow native characters.
+
+    * Number pad reverts back to normal keyboard
+
+    * Gboard and Swift swipe entry are not supported.
+
+    * Cursor and UMG do not match. 
+
+    * Text field doesn’t appear if there was too much text.
+
+* Crash Fix: Samsung Galaxy S8 with Android 7.0 is not hiding suggestions and disabling predictive input. There are cases with this that can cause a crash. Global solution: On text change, downgrade to simple suggestions all the easy correction spans that are not a spell check span (remove android.text.style.SuggestionSpan.FLAG_EASY_CORRECT flags).
+
+* Bug Fix: Fixed an issue to properly detect need for rebuilds from the final NDK and SDK levels only in UEBuildSettings.txt.
+
+* Bug Fix: Fixed Android icon paths.
+
+* Bug Fix: Fixed polling code in Tcp Message Transport that was wasting CPU cycles when using Launch on an Android device.
+
+* Bug Fix: Fixed an issue related to Text Watcher On Text Change being called multiple times when typing or deleting the content of a native Android Edit Text control. 
+
+* Bug Fix: Fixed an issue where Android startup movies not always playing.
+
+* Bug Fix: Fixed an issue with values for pinch input producing very different results for the same area on Android devices:
+
+    * When the pinch goes beyond the viewport boundaries, the touch that goes off-screen is "released" and the zoom effect is over. This is solved by remembering the last pinch event values
+
+    * The initial distance for the pinch/rotate are "hacked" by touching the screen and moving the finger to another position before using the second finger. This is solved by using the correct values when the pinch event starts.
+
+* Bug Fix: Fixed an issue with Android Web Browser 3D. The external texture was moved from Content Browser to the Web Browser plugin.
+
+* Bug Fix: Fixed access to Frame Update Info in MediaPlayer14.java and CameraPlayer14.java with Proguard.
+
+* Bug Fix: Fixed an issue with flag vertex and fragment shaders belonging to materials with external textures.
+
+* Bug Fix: Protected against asking for duration before prepare completed in movie player.
+
+* Bug Fix: Fixed support for GL_OES_standard_derivatives on Android.
+
+* Bug Fix: Fixed an issue so that the relative path to ResonanceAudio_API.xml is accurate.
+
+* Bug Fix: Fixed an issue with external texture shaders for some devices.
+
+* Bug Fix: Added missing SecureHash.h include needed for non-unity builds. 
+
+* Bug Fix: Fixed some issues on Lenovo 939 when rendering a simple scene. The shader compiler complained about code lines before the #extension directectives. Added a placeholder ("// end extensions") in the shader code. 
+
+* Bug Fix: Fixed DXT apps that failed on first launch when 'For Distribution' enabled with the error message: "Unsupported Texture Format".
+
+* Bug Fix: Corrected the code which rounded up the coordinates to the nearest even number.
+
+* New: Updated Android SDK license agreement for Gradle.
+
+* New: We now use UPL hash to invalidate build if any UPL files are changed.
+
+* New: Added a wait for the audio thread to pause audio when going to background. 
+
+* New: Added a new packaging mode for Android:
+
+    * ETC1a, which will compress files with alpha channel as ETC2 and files without alpha channel. 
+
+    * At runtime, detection will determine when ETC2 is not supported and will decompress those files. 
+
+    * This format will work on all Android devices while using one single compressed texture format, but devices without ETC2 support will have an impact on performance.
+
+* New: Added support for texture streaming to most Android devices with ES 3.1 support.
+
+* New: Added verification of support for cooked texture format(s) on device at runtime. (optional with Validate texture formats checkbox in Android project settings and skipped for cook on the fly)
+
+* New: Added new Android Project Settings to override the NDK and SDK may now be optionally used per-project. Leave them blank to use the global setting in Android SDK Project Settings. 
+
+* New: Added LOD Proxy support.
+
+##### HTML5
+
+* Bug Fix: Chrome renders transparent on OSX even though canvas has been set with alpha to false. 
+
+    * Note that if backbuffer is needed with alpha values, this will need to be rewritten.
+
+* Bug Fix: Fixed an issue for "Runtime Error: integer result unrepresentable" compiler error removed the following linker options: 
+
+    * -s OUTLINTING_LIMIT flags -s BINARYEN_TRAP_MODE=’clamp’
+
+* Bug Fix: Disabled Support Screen Percentage to render the full screen otherwise. This was causing HTML5 screens to render only a portion of the screen and in black. 
+
+* Bug Fix: Fixed an issue with 2D text rendering. 
+
+* Bug Fix: Fixed an issue so that logs are redirected to the console window instead of "low window".
+
+* New: Added HTML5 Launch Helper to check to see if the port is already in use and gracefully exit with a message if already in use. . 
+
+* New: Set controller axis and button max value in code instead of relying on emscripten Get Gamepad Status. 
+
+* New: Updated HTML5 readme file with a lot of new and clearer information on files generated, what they are, and how to use the compressed versus uncompressed versions. 
+
+* New: Resurrected vehicle games and removed "IsContentOnly" check during HTML file generation to allow code projects to work again. This corrected optimization suffix for PhysX 3 Vehicles. 
+
+##### iOS
+
+* Bug Fix: Fixed an issue while remote compiling Metal shaders to now check the connection before starting the compile process. 
+
+* Bug Fix: While Remote Shader Compiling for Metal, the editor no longer spawns and closes a lot of minimized windows.
+
+* Bug Fix: Fixed an issue to modify the toolchain to insert the correct SDKRoot. 
+
+    * XCode address sanitizer feature did not work on iOS. The dylib loader depended on the default SDKROOT parameter. 
+
+* Bug Fix: Fixed local WiFi multiplayer discovery on iOS.
+
+* Bug Fix: Fixed large Metal texture memory leak when opening card packs.
+
+* Bug Fix: Fixed a issue causing iOS remote building to fail when the remote mac user name contained a space character.
+
+* Bug Fix: When building on Mac if no valid provisioning profile and sign in certificate is set up, but the XCode has a correct development wildcard, it will use that wildcard to sign the build.
+
+* Bug Fix: tvOS now correctly package or launch-on.
+
+* New: Enabled Avf Media plugin for tvOS.
+
+* New: Disabled the validation layer in iOS App Store builds for shipping builds.
+
+##### Linux
+
+* Crash Fix: Fixed a crash on a headless system.
+
+* Crash Fix: Fixed a crash when printing `%c` on Linux, Mac, HTML5 etc.
+
+* Crash Fix: Made RHI initialization not crash on Mesa drivers.
+
+* Crash Fix: Pop-up a message box when unable to launch ShaderCompileWorker instead of crashing.
+
+* Crash Fix: Disabled some unstable gdb visualizers that can crash the debugger.
+
+* Crash Fix: Fixed an edge-case initialization crash (will print a user-friendly message and exit).
+
+* Crash Fix: Fixed the Editor crashing on start if compiled with a recent (5.0+) clang.
+
+* Crash Fix: Fix an audio clean up crash when exiting PIE. 
+
+* Bug Fix: Fixed the Setup.sh handling of special characters in PNG path.
+
+* Bug Fix: Fixed UBT to invoke newer (supported) clang versions.
+
+* Bug Fix: Fixed libelf breaking build if some third party libs are used.
+
+* Bug Fix: Remove unneeded dependency on CEF for base engine.
+
+* Bug Fix: Fixed mismatching "noperspective" qualifier.
+
+* Bug Fix: Fixed Linux thread/process priorities (requires system settings to allow raising them).
+
+* Bug Fix: Fix priority calculation when RLIMIT_NICE is not 0.
+
+* Bug Fix: Fixed the Output Device Std Output function to use printf() on Unix platforms. Mixing wprintf/printf is undefined behavior.
+
+* Bug Fix: Fixed the unnecessary test for popcnt presence.
+
+* Bug Fix: Fixed compilation with glibc 2.26+ (e.g. Ubuntu 17.10).
+
+* Bug Fix: Fixed Slate Dialogs to show image icon for *.tga.
+
+* Bug Fix: Fixed submenus of context menu not working.
+
+* Bug Fix: Fixed RHI settings being wrong when removed via project properties UI.
+
+* Bug Fix: Fixed the inability to click on some UI elements at the edges of the screen due to tooltip windows obscuring it.
+
+* Bug Fix: Made UBT use response files for the Linux compiler (avoids MAX_PATH issues when cross-compiling).
+
+* Bug Fix: Fixed symbol interposition problems with both AMD and NVidia drivers.
+
+* Bug Fix: Fixed ARB callback errors when hitting Build Lighting.
+
+* Bug Fix: Fix shadow compiler warnings around our includes.
+
+* Bug Fix: Do not assert on an unimplemented function, just log the function.
+
+* Bug Fix: Set sound to unfocused volume multiplier if not focused.
+
+* Bug Fix: Open project in the file explorer when finished creating it if the Null Source Code Access is the default source code editor
+
+* Bug Fix: Fixed inability to create C++ projects on Linux despite being able to compile the editor.
+
+* New: Updated Linux VHACD.
+
+* New: Added LTO support.
+
+* New: Enabled DPI on/off switch.
+
+* New: Changed BinnedAllocFromOS to place descriptor after allocated chunk instead of in front of it to save some memory.
+
+* New: Re-enabled index buffer optimization. It was disabled due to problems found in the live release.
+
+* New: Added build-id to all binaries to improve lldb startup times.
+
+* New: Vulkan RHI will default to SM5
+
+* New: Updated OpenAL to 1.18.1.
+
+* New: Use elementary OS versions to set the Ubuntu versions for dependencies.
+
+* New: Use Xft.dpi when calculating DPI on Linux.
+
+* New: Rebuild libcurl 7.57.
+
+* New: Added detection and messaging when building lighting is impossible due to local network interface problems.
+
+* New: Add better error handling when using new Toolchain for Linux on Windows.
+
+* New: Add Alembic (*.abc) support
+
+* New: Added fixes to make porting to gcc easier.
+
+* New: Included VSCode in the list of default project generators.
+
+* New: Implemented voice capture using SDL.
+
+* New: Recompiled the bundled SDL library with Wayland support.
+
+* New: UBT will use compiler response files when compiling for Linux.
+
+* New: Enabled pooling of memory allocations (only 64KB allocations) on servers.
+
+* New: Added LTO support to the Linux cross-builds (requires v11 and newer toolchain).
+
+* New: Enabled packaging for Android Vulkan.
+
+* New: Added a message about running out of inotify handles a warning instead of error so it doesn't break cooking.
+
+* New: Added better error messaging when running out of inotify watches.: 
+
+* Removed: Removed the support for /-prefixed command line switches.
+
+##### Mac
+
+* Crash Fix: Fixed a rare crash caused by a deferred call to Close Window function executing during Slate's PrepassWindowAndChildren call.
+
+* Crash Fix: Fix a crash when using VOIP on Mac.
+
+* Crash Fix: Fixed a crash in Avf Media Player related to reusing an Avf Media Player object.
+
+* Crash Fix: Fixed a crash on destroying certain menus in an animation persona window.
+
+* Bug Fix: Fixed a problem with Xcode project generator not handling quoted preprocessor definitions correctly.
+
+* Bug Fix: Fixed a problem with IME not working in context menu text edit fields.
+
+* Bug Fix: Disable Zoom button if project requests a resizable window without it.
+
+* Bug Fix: Fixed header search paths in XCode project for paths containing space character.
+
+* Bug Fix: Fixed an issue with creating an NSString for the Low Level Output Debug String. The string can just be printed out as a UTF-8 string.
+
+* Bug Fix: Fixed a C++ std in generated Xcode project to match rest of engine (C++14). 
+
+* Bug Fix: Fixed an issue with some context menus opened right after closing other context menu disappearing right after being shown.
+
+* Bug Fix: Don't activate Mac Crash Reporter Client on launch when it's running in unattended mode.
+
+* Bug Fix: Fixed a small number of persistent memory leaks on the Mac build that slowly consume more and more memory as you use the Editor. Interacting with menu's was particularly egregious as each NSMenu would leak after you move away.
+
+* Bug Fix: Fixed a problem with Mac editor not exiting with error code returned from Guarded Main function.
+
+* Bug Fix: Better support for GPUs with 4GB or more of video memory.
+
+* New: Updated Xcode project generator to sort projects in the navigator by name (within folders) and also sort the list of schemes so that their order matches the order of projects in the navigator.
+
+* New: Changed XCode project to automatically build ShaderCompileWorker in addition to the selected target for non-program targets.
+
+* New: Added initial support for building and running Niagara.
+
+* New: Added support for changing monitor's display mode on Mac in fullscreen mode. This greatly improves performance on Retina screens when playing in resolutions lower than native.
+
+##### Windows
+
+* Bug Fix: Switched Windows to use wprintf() consistently.
+
+* Bug Fix: Fixed an issue with Windows file save dialog not appending the correct extension when there were multiple options available.
+
+* Bug Fix: Fixed an issue with custom window positions in "windowed mode".
+
+#### Programming
+
+* New: Minimized number of fractional digits you can have in a resultant string while using FString::SanitizeFloat, defaulting to 1. Setting this to 0 prevents adding '.0' to integers.
+
+* New: Added ‘.idea/misc.xml’ file generation to speed up CLion indexing.
+
+* New: Added option to generate minimal set of targets for CMake files.
+
+* New: Added shader and config files to CMake file generation for searching in IDEs.
+
+* New: Added icon size override to FObjectEntryBox.
+
+* New: Added CMAKE_CXX_COMPILER and CMAKE_C_COMPILER settings to generated CMake files.
+
+* New: Removed deprecated version of RHISetStreamSource.
+
+* New: Optimized hlslcc to use hashes instead of lists with hash bucket sizes being prime.
+
+* New: When generating CMake files on Mac, added support for iOS and TVOS targets.
+
+* New: Merged CLion source code access Plugin support.
+
+* New: Reduced calls to malloc by using larger memory blocks for hlslcc Type allocation pools.
+
+* New: Optimized hlslcc memory allocation policy on Mac, resulting in `34%` faster and `83%` cheaper memory moves.
+
+* New: Added basic support for hotfixing live assets from .ini files.
+
+#### Rendering
+
+* Crash Fix: Fixed a Metal shader standards iOS crash when users would try to load a project that tries to use an incompatible shader version.
+
+* Crash Fix: Disable the METAL_SM5_NOTESS shader platform again to workaround the Nvidia pipeline state compiler crash by changing the buffer address space from "constant" to "device".
+
+* Crash Fix: Fixed a crash on D3D12 when locking buffers for read access.
+
+* Crash Fix: Resolved an intermittent crash in the MetalRHI by retaining the bound Render Targets until we bind something else. 
+
+* Crash Fix: Developed a workaround for a crash on fullscreen startup in packaged builds of Showdown for Mac.
+
+* CrashFix: Fixed a crash in Metal on one kind of draw call when tessellation is enabled.
+
+* Crash Fix: Fixed crash when building Texture Streaming in the resave package commantlet.
+
+* Crash Fix: Attempt to reduce the number of crashes on Nvidia Macs by retaining resources used in each command-buffer.
+
+* Crash Fix: Fixed a crash in the MetalRHI when using dithered LODs caused by one of the fragment shaders not declaring any explicit outputs.
+
+* Crash Fix: Fixed a crash in projects that used Distance Field Shadowing when Static Meshes were updated while not yet registered with the level.
+
+* Crash Fix: Doubled the size of the occlusion query buffer in D3D12, to 256KB of memory, to prevent crashes in larger levels.
+
+* Crash Fix: Fix for Intel HD 4000 crashing on Mac OS 10.12.6.
+
+* CrashFix: Fixed a crash when calling the "FreezeRendering" console command on non Editor builds.
+
+* Crash Fix: Fixed crash when scaling a Blueprint with a Planar Reflection component.
+
+* Crash Fix: Fixed issue where FOpenGLDynamicRHI RHISetRenderTargetsAndClear function would crash when the render target did not have a color binding.
+
+* Bug Fix: Fixed Alembic Geom Cache no longer rendering correctly.
+
+* Bug Fix: Fix the Intel vector-array-dereference workaround so that it doesn't cause the AMD compiler to explode instead.
+
+* Bug Fix: Fixed the way iOS version specific Metal features are enabled to avoid an overflow bug that meant some weren't being enabled.
+
+* Bug Fix: Removed support for Metal Shader Model 4 on Mac. 
+
+* Bug Fix: Fix compilation on HLSLCC with integral values are not automatically converted into comparisons with zero.
+
+* Bug Fix: Fixed the normalization of Morph Tangents. 
+
+* Bug Fix: Tweaked Metal GPU identification so that it works with eGPU boxes and prototype hardware.
+
+* Bug Fix: Implement Byte Address Buffer RWByte Address Buffer in HLSLCC in a similar manner to Structured Buffer RWStructured Buffer so that the backends don't need too much modification. 
+
+* Bug Fix: Enabled the ability to use GPU morph targets on Mac.
+
+* Bug Fix: Fixed RHI thread and parallel execution in Mac -game mode.
+
+* Bug Fix: DX11 Texture formats are now supported on Mac Metal.
+
+* Bug Fix: Fixed long standing bug in OpenGL Vertex Declaration cache with hash collisions in the cache sometimes returning the wrong declaration.
+
+* Bug Fix: Fixed Material events not showing up in RenderDoc.
+
+* Bug Fix: Fixed Anisotropic Filtering not working on Vulkan.
+
+* Bug Fix: Memory is no longer reallocated on every Vulkan draw call.
+
+* Bug Fix: Stopped memory leak from MTL Textures missing autorelease pool blocks in Metal Texture functions.
+
+* Bug Fix: Fixed Metal compilation for PCSS shadows.
+
+* Bug Fix: Fixed comparison of Scene Color Format. 
+
+* Bug Fix: Fixed incorrect handling of some float to uint shader loads in Metal.
+
+* Bug Fix: Reduce the number of Metal PCH shader compiler errors by using atomic file operations.
+
+* Bug Fix: Changed the Metal sampler filter translation to better match Direct3d.
+
+* Bug Fix: Shared Texture samplers are now refreshed when changing Max Anisotropy level.
+
+* Bug Fix: Fixed a bug where Static Meshes were being added to the scene twice when in Simple Forward rendering was enabled.
+
+* Bug Fix: Fixed the native shader libraries again & undo.
+
+* Bug Fix: Removed the single linear array for accumulated shader code.
+
+* Bug Fix: Uncompress the shader data for the native library system. 
+
+* Big Fix: Fixed an errant change to the Metal compiler that was trying to force the fully compiled library into the single Metal library.
+
+* Bug Fix: Fixed up the internal validation layer in MetalRHI which can now be enabled from the OS command-line.
+
+* Bug Fix: Fixed condition controlling which features are enabled when iOS >= 10.3.
+
+* Bug Fix: Errors from failed attempts to compile global shaders are now reported correctly. This information enables us to identify the causes of these failures on non-Windows platforms.
+
+* Bug Fix: Fixed Texture resolution returning 0 when running Blueprint Construction scripts when cooking.
+
+* Bug Fix: We can now handle swizzle's in the HLSLCC FMA identification pass so that we reduce the number of instructions and the platform compiler can't break the instructions up.
+
+* Bug Fix: Developed a work around two different bugs with Metal shader compilation behavior that caused black flickering when a material used World Position Offset. 
+
+* Bug Fix: Metal Shader Format will no longer fallback to text shaders when you ask it to compile to bytecode but the bytecode compiler is not available (either locally or remotely).
+
+* Bug Fix: Metal can now render patches as well as triangles when tessellation is enabled.
+
+* Bug Fix: Add the ability to emit more details when a pixel shader is found to have no outputs at all.
+
+* Bug Fix: Fixed image plates not showing up in high resolution screenshots.
+
+* Bug Fix: VIS_Max meta is now set to hidden. 
+
+* Bug Fix: Fixed Distance Field Ambient Occlusion not working on Metal platforms.
+
+* Bug Fix: Changed Metal pipeline bindings warning to Verbose.
+
+* Bug Fix: Prevented an unnecessary red<->blue swap in FVulkan Dynamic RHI RHI Read Surface Data function when reading back a surface in VK_FORMAT_R16G16B16A16_SFLOAT format.
+
+* Bug Fix: Fixed non-deterministic locations of procedurally placed Landscape grass on different platforms and compilers.
+
+* Bug Fix: Fixed huge client hitch when applying changes to rendering scalability settings while in game.
+
+* Bug Fix: Stopped MetalRHI leaking vertex-descriptors and debug-group strings.
+
+* Bug Fix: applied a fix to runtime error with textures being loaded from bulk data when texture streaming is disabled.
+
+* Bug Fix: Removed faulty ensure in the texture streaming code.
+
+* Bug Fix: Fix on D3D12 for CPUs with more than 16 cores.
+
+* Bug Fix: Added support for Vulkan devices that don't support cached memory type.
+
+* Bug Fix: Fixes for UScene Capture Component Hidden Actors always staying gray out in World Editor.
+
+* Bug Fix: Fixed Recompute Tangents and Skin Cache on OpenGL 4.
+
+* Bug Fix: Fixed OpenGL sometimes overwriting resources if too many were bound on OpenGL 4.
+
+* Bug Fix: Fixed Metal support for Manual Vertex Fetch by making it a property of the shader platform.
+
+* Bug Fix: Fixed tessellation shader capability initialization in OpenGL renderer.
+
+* Bug Fix: Fixed instancing in Vulkan.
+
+* Bug Fix: Fixed Skeletal Mesh LODs inside editor.
+
+* Bug Fix: Fixed loading shaders from the cooked DDC.
+
+* Bug Fix: Planar Reflections now properly handle world origin rebasing.
+
+* Bug Fix: Fixed initialization of instancing random vertex stream.
+
+* Bug Fix: Using the -opengl command line argument will now default OpenGL to 4.3.
+
+* New: Added a recompute Tangents options for cloth assets.
+
+* New: Updated Nvidia Ansel to v1.4
+
+* New: Added a new console-variable to enable and disable Manual-Vertex-Fetch for Metal. MVF works on macOS, though testing did expose an error with Tessellation on Nvidia (true for MVF enabled & disabled).
+
+* New: Refactored the way we compile Buffer and RW Buffer types for Metal so that we can support the type-conversion semantics of HLSL/D3D.
+
+* New: Manual Vfetch for Static Meshes. Skeletal Meshes also use the new layout.
+
+* New: Added code to D3D11RHI (enabled with "-d3ddebug") to validate Draw Indexed Primitives isn't trying to draw off the end of the instanced vertex streams.
+
+* New: Added Get Section From Procedural Mesh utility function for the Procedural Mesh Component.
+
+* New: Added support for 2D IES asymmetrical light profiles that can be up to 256x256 in size. 
+
+* New: Added support for linear HDR EXR output from Sequencer.
+
+* New: More Metal debugging commandline options "-metalfastmath" & "-metalnofastmath" to force fast-math on or off for all shaders, must be using runtime-compiled shaders (i.e. -metalshaderdebug or r.Shaders.Optimise=0) to take effect.
+
+* New: Added support for Nvidia Aftermath v1.3.
+
+* New: New  cvar r.ForceStripAdjacencyDataDuringCooking that will strip adjacency data from all Skeletal and Static meshes during cooking.
+
+* New: ViewToClipNoAA added to uniform shader parameters. This allows "After Tonemapper" post process blendables that can be positioned in view space when using temporal AA.
+
+* New: Added Render Target Pool Stats. Use "stat render target pool" to view.
+
+* New: Added support for Speedtree 8.
+
+* New: Update to AGS v5.1.1 support.
+
+* New: Changed IES texture brightness to be the max candela value and set the texture multiplier to 1. 
+
+* New: Made major changes on how Vulkan manages its descriptor pools.
+
+* New: Support for DX11 quad buffer stereo rendering.
+
+* New: Added unified present threshold console variables between all supported platforms. 
+
+* New: Added RHI Command List Enqueue Lambda method which can be used to enqueue arbitrary tasks on the RHI thread from the render thread without having to write lots of boilerplate FRHI Command functor classes.
+
+##### FX
+
+* Bug Fix: Fixed texture streaming with particle Subuv.
+
+* New: Implemented a fast particle pool memory which default is 2MB and automatically cleans up oldest used pool slots. 
+
+##### Lighting
+
+* Bug Fix: Fixed Lightmass correctness issues so that, Point, Spot, Directional and Skylights more closely match reference renderer. 
+
+* Bug Fix: Removed legacy fudge factor on point / spot light photon energy * Spotlights no longer emit based on indirect photon paths. 
+
+* Bug Fix: Added the ability to enable LPVs for Mac Metal.
+
+* Bug Fix: Integrated the LPV_STORE_INDEX_IN_HEAD_BUFFER related changes from Lionhead to make Light Propagation Volumes potentially viable on non-Microsoft platforms.
+
+* Bug Fix: Created Stat Map Build Data to track the memory size of lighting and reflection capture build outputs.
+
+* Bug Fix: Added support for Reflection Captures to support Lighting Scenarios without the need for recapturing. 
+
+* Bug Fix: Reflection Captures are now part of the Map Build so modifying a Reflection Capture in Editor will display a preview, but game can only display built captures (black for unbuilt with screen message) * 
+
+* Bug Fix: Moved Reflection Capture build data to the Build Data package so building lighting / Reflection Captures no longer dirties U Levels.
+
+* Bug Fix: Added the ability for Skylights which capture the scene to work correctly with Lighting Scenarios.
+
+* Bug Fix: Lighting Scenarios must now be loaded for each time they are made visible no switching back and forth while keeping both loaded.
+
+* Bug Fix: Exposed the Emissive Boost, since Lightmass supports Emissive areas on Static Meshes.
+
+* Bug Fix: Added the ability for Translucency lighting modes to work with Volumetric Lightmaps to give the  lighting much more accuracy. 
+
+* Bug Fix: Added A Volumetric Lightmap Density Volume which gives local control over Volumetric Lightmap density. 
+
+* Bug Fix: Volumetric Lightmap no longer refines around static translucent geometry.
+
+* Bug Fix: Reworked brick culling by error mechanism which now compares error to interpolated parent lighting instead of the brick average.
+
+* Bug Fix: Added Lightmap Type to Primitive Component.
+
+* Bug Fix: Added new property, Force Volumetric which allows forcing Static Geometry to use Volumetric Lightmaps. 
+
+* Bug Fix: Added Force Surface which replaces b Light As If Static Improvements to Volumetric Lightmap quality needed for static geometry.
+
+* Bug Fix: Fixed Stationary light shadowing so it is now dilated when inside geometry.
+
+* Bug Fix: Added refinement around geometry which uses an expanded cell bounds when the geometry is going to use Volumetric Lightmaps, since cross-resolution stitching causes leaking.
+
+* Bug FIx: Added a New World Settings Lightmass property, Volumetric LightmapSphericalHarmonic Smoothing, which  controls the global amount of smoothing applied.
+
+* Bug Fix: Hidden levels now keep their last build data when using lighting scenarios.Fixed light build bug where hidden levels affected volumetric lightmaps.
+
+* Bug Fix: Prevent an assert from firing when LPV Intensity is saved to a value of 0.0.
+
+* Bug Fix: Fixed instanced static mesh lightmaps when using lighting scenarios.
+
+* New: Added a intensity units property, for Point and Spot lights, that can be set to Candelas, Lumens or Unitless (as before). 
+
+* New: Added r.Capsule Direct Shadows and r.Capsule Indirect Shadows for more specific scalability control over capsule shadow features. 
+
+* New: Added a Lighting Feature show flags for RTDF shadows and Capsule Shadows.
+
+* New: Capsules Shadows are now disabled on the lowest shadow scalability quality.
+
+* New: Added Forward shadowing of Directional light for Translucency * Static shadowing and added CSM supported with minimal filtering (1 PCF) * Deferred renderer. This affects Translucency using 'Surface Forward Shading' lighting mode while the Forward renderer affects all translucency.
+
+* New: Volumetric fog is now always interpolated in the Pixel shader, since per-vertex interpolation gives consistently poor results. 
+
+* New: Added a checkpoint for Screen Shadow Mask Texture, allowing Vis Screen Shadow Mask Texture.
+
+##### Materials
+
+* Bug Fix: Return error for external texture if not used in pixel shader
+
+* Bug Fix: Fixed vertex interpolator node when used with sprite and gpu particle materials.
+
+* Bug Fix: Fixed an issue where media textures that were referenced by a material and then later associated with a media player by a blueprint would fail to render media
+
+* Bug Fix: Fixed a bug where renaming a parameter in a material using functions would sometimes reset the value in child instances.
+
+* Bug Fix: Fixed texture streaming issue when the camera FOV is reduced by gameplay (like when using a weapon scopes)
+
+* New: Added a Channel Mask material parameter, changeable at runtime and returns the selected channel of it's input.
+
+* New: Material Layers experimental enable flag has moved to Rendering > Experimental > Support Material Layers. 
+
+##### Mobile Rendering
+
+* Bug Fix: Cancelled the applied exposure scale with planar reflections for non-hdr mobile
+
+* Bug Fix: Added half float -> float conversion for static vertex buffers when device does not support half float vertex attributes.
+
+* Bug Fix: Fixed ES2 issue reading from a render target in float format for BGRA8888 when the glReadPixels() call fails
+
+* Bug Fix: Provided an implementation of round() under GLES2 when the shader compiler does not provide it
+
+* New: Added GLES3.1 compatibility mode for mobile preview to the windows Open GLRHI.
+
+* New: Changed the way Manual Vertex Fetch is enabled on Metal platforms so that it is enabled when targeting Metal v1.2 and higher (macOS 10.12+/iOS 10+). 
+
+* New: Pass project defined maximum CSM count to mobile CSM shaders. 
+
+* New: Added CSM distance fading for mobile.
+
+* New: Add mobile CSM shader culling for movable directional lights.
+
+##### Optimizations
+
+* Bug Fix: Vulkan descriptor pools are now allocated per PSO instead of globally to reduce peak mem consumption and fragmentation (enabled on Windows only).
+
+* Bug Fix: Reduced Vulkan saved pipeline cache file size by filtering duplicated shader SPIRV.
+
+* Bug Fix: Fixed Actors with HLOD parent primitives not having synchronized visibility with their parent.
+
+* Bug Fix: Improved consistency of HLOD transitions by using the distance to the closest point on the AABB instead of distance to the origin of their bounding sphere.
+
+* Bug Fix: Applications using OpenGL now properly discard copy of vertex and index buffers in memory once they are uploaded to GPU.
+
+* New: Added new command "stat Streaming Overview" giving high level metrics of Texture usage.
+
+* New: Added optional support for compressing saved Vulkan PSO cache with r.Vulkan.Pipeline Cache Compression.
+
+* New: Added Shader compiling memory optimizations to reduce memory and time it takes to re-compile Editor shaders.
+
+* New: Removed read / write locks from PipelineStateCache, Estimated half time on the main thread for PipelineStatecache operations. Introduces duplicate which are cleaned up at end of frame. Slightly higher memory overhead.
+
+* Optimized the texture streaming update loop when using a many small levels in big worlds.
+
+* New: Added new CVar "r.DistanceFields.ForceMaxAtlasSize" (defaults to zero, which leaves the current behavior of auto-sizing up to configured maximum, as needed.) This can be used to avoid costly hitches when resizing the distance field in projects with large numbers of actors.
+
+* New: Added new "r.HLOD.DistanceOverride" CVar, which can be used to force all HLODs to transition at a specific distance, regardless of the individual cluster's settings. Defaults to zero (off).
+
+* New: Added new options to HLOD settings to strip adjacency and distance fields from generated meshes.
+
+* New:  Added new CVar for adjusting view distance scale, which is not affected by scalability settings
+
+* Depreciation: Removed deprecated Lightmap UV Bias, Shadow Map UV Bias from Instanced Static Mesh Component.
+
+##### Post Processing
+
+* Bug Fix: Fix the black ground scattering on Metal by ensuring that input into abs() is non-negative for atmosphere rendering.
+
+* New: Added color grading control Blue Correction to correct for artifacts with "electric" blues due to the ACEScg color space.
+
+* New: Added color grading control Expand Gamut which expands bright saturated colors outside the sRGB gamut to fake wide gamut rendering.
+
+* New: Implemented a SSR input Post Process Material location to compose linear color space backplate after TAA for next frame's SSR.
+
+#### Tools
+
+* Crash Fix: Fixed a crash when attempting to move windows in VR Mode.
+
+* Bug Fix: Fixed an issue where the Selection Laser is detached from the controllers in VR Mode.
+
+* Bug Fix: ASTC texture versioning is now corrected.
+
+* Bug Fix: Fixed an issue reporting some files did not copy to the device while those files were not actually needed.
+
+* Bug Fix: Fixed an issue where SSH Key was failing to generate when the remote Mac user contained a space character.
+
+* Bug Fix: Fixed memory corruption that could occur when using the Open Level dialog.
+
+* New: Added Post Process preview in the Material Editor viewport.
+
+* New: Reduced memory requirements of many graph related classes by storing data as FName instead of FString.
+
+* New: Added support for Debug Symbols (.udebugsymbols) files generated for iOS builds to be used for memory debugging with the Memory Profiler 2 tool.
+
+##### AutomationTool
+
+* Bug Fix: Fixed parallel executor getting the working directory from task element, not tool node.
+
+* Bug Fix: Fixed an issue to no longer produce access violations when creating an installed build of the Engine  when running the Automation Tool directly, rather than using the Automation Tool Launcher.
+
+* Bug Fix: Fixed one CPU core spinning during compilation using Parallel Executor. This was due to Completed Event never being reset.
+
+* Bug Fix: Fixed the ability to distribute app-local dependencies with games packaged using the launcher UE4 distribution. 
+
+* Bug Fix: App-local dependencies now work for C++ and Blueprint projects.
+
+* Bug Fix: Fixed a race condition where stdout/stderr write handles could be inherited by multiple processes, resulting in them not being closed (and the process exit being detected) until all processes that inherited them had been closed. This improves performance of Parallel Executor.
+
+* New: Added support for external projects with new parameter to "List Third Party Software". The parameter -ProjectPath specifies the path to the UProject file so that the build rules of the project are parsed and the folder is part of the search for Third Party Software files.
+
+* New: Lowered logging level for cached file notification during Cook from Warning to Verbose in order to reduce clutter in the logs.
+
+* Removed: Code that deletes cooked data on a failed cook has been removed because deleting the cooked data prevented post-mortem analysis. The Engine now writes out packages transactionally by writing to a temporary file and moving them into place.
+
+* Removed: Code that executes pre-built steps from UE4Build have been removed. These build steps should be executed before generating the action graph. So, they have to be done inside of Unreal Build Tool. Running them here causes them to execute twice.
+
+##### BuildGraph
+
+* New: Added an example script for building the Editor and Game targets for a project, and submitting the binaries to source control.
+
+* New: Improved error reporting when compiling a C# project with an invalid reference.
+
+* New: Added an <MsBuild> task for compiling MSBuild projects.
+
+##### Network Profiler
+
+* New: The Network Profile is now 64 bit and the solution has been converted to VS2015 and are now about `6%` smaller.
+
+* New: Added a checkbox above the connections list that will select or deselect all connections in the Network Profiler.
+
+##### UnrealBuildTool
+
+* Bug Fix: Fixed an issue causing Unreal Build Tool on Mac to not use makefiles which made the time it needed to determine what needs to be compiled and linked longer than on other platforms.
+
+* Bug Fix: Fixed some errors encountered when cleaning a target now cause an error code to be returned.
+
+* Bug Fix: Fixed the command line arguments -monolithic and -modular to now be parsed correctly.
+
+* Bug Fix: Fixed the include paths ending with a backslash that were escaping the closing quote when passed on the command line.
+
+* Bug Fix: Fixed post-build steps for plugins not being executed.
+
+* Bug Fix: Output subfolders are no longer being appended twice when generating project files if set by a target. This was causing incorrect command lines when launching from Visual Studio.
+
+* Bug Fix: Fixed batch file paths not being quoted correctly when running through XGE.
+
+* Bug Fix: Fixed issues where Includes with backslashes were causing an unhandled exception on a background thread in UnrealBuildTool.
+
+* Bug Fix: Fixed Version and Module files not being included in the output manifest. This fixes missing version information when using precompiled binaries in UGS.
+
+* New: UnrealBuildTool now passes the "/d2cgsummary" flag to the Visual C++ compiler when running with the -Timing argument.
+
+* New: Added include and library-related fields to module JSON output. 
+
+* New: VSWHERE is now used to find the location of MsBuild.exe if available.
+
+* New: Improved error handling while generating intellisense project files. It now includes the name of the target being compiled and allows project file generation to continue without it. 
+
+* New: The Module Rules "Definitions" property has been renamed to "Public Definitions", to make its semantics clearer. A complementary "Private Definitions" property has also been added.
+
+* New: UnrealBuildTool now validates that targets don't modify properties that affect other targets. For example, modifying global definitions for an Editor target, where other Editor targets may use the same engine DLLs).
+
+* New: Added the build version to the Target Rules and Read Only Target Rules class, allowing it to be read from the target and module rules files. 
+
+    * Changed Launch.build.cs to set private defines for the build changelist in non-formal builds. This allows the version to be set correctly with executables copies to console. 
+
+    * Modular builds on desktop platforms still read the Version file from disk since it’s quicker than compiling a value in. They need the build ID to prevent mismatched binaries. 
+
+* New: Version files are only generated for non-promoted builds, and are included in the target receipt.
+
+* New: Added a proper error message if the user attempts to clean a target through hot-reload, rather than just failing to delete DLLs because they are locked.
+
+* New: Added a logging option to format output messages in a form that can be parsed by MSBuild. This prevents errors showing as "EXEC: Error:", and displays them correctly in the error list window.
+
+* New: Unreal Build Tool will now use the preferred source code accessor selected in the editor to compile with, if an explicit compiler is not specified in the Project Settings.
+
+* Deprecated: Unified command line arguments now use the -Name=Value syntax. The following are no longer supported:
+
+    *  -Module <Name> 
+
+    * -ModuleWithSuffix <Name> <Suffix>
+
+    * -Plugin <Path>
+
+    * -Receipt <Path>
+
+* Removed: UDN documentation source files are no longer included in the generated project files by default.
+
+##### UnrealGameSync
+
+* Bug Fix: Fixed an issue so that the project filename is now included in the Editor build command. The project may not be in one of the .uprojectdirs paths.
+
+* Bug Fix: Fixed an issue where the "More Info..." button was not using P4 server override.
+
+* New: UProject and UPlugin files are now treated as code changes. This change was made since they can invalidate the Editor target. 
+
+* New: Added an in-place error panel to show when a project fails to open. This enables the user to retry and have their tabs saved instead of creating a modal dialog.
+
+##### UnrealVS
+
+* New: Added Log GenerateProjectFiles to file when using Unreal VS.
+
+#### UI
+
+* Crash Fix: Fixed a line-ending inconsistency when retrieving rich-text that could cause Input Method crashes. 
+
+* Bug Fix: Fixed some clipping issues that could occur with right-aligned text. 
+
+* Bug Fix: No longer include trailing whitespace in the justification calculation for soft-wrapped lines. 
+
+* Bug Fix: Fixed incorrect text selection when selecting via double-click beyond the bounds of the line.
+
+* Bug Fix: Special case zero-width space in the text shaper to avoid fonts rendering the fallback glyph.
+
+* New: Added support for displaying and editing numbers in a culture-correct way. 
+
+    * Numeric input boxes in Slate will now display and accept numbers using the culture correct decimal separators. 
+
+    * This is enabled by default, and can be disabled by setting Should Use Localized Numeric Input to False in the XEditorSettings.ini (for the editor), or XGameUserSettings.ini (for a game).
+
+* New: Added support for a catch-all fallback font within composite fonts. 
+
+    * This allows you to provide broad "font of last resort" behavior on a per-composite font basis, in a way that can also work with different font styles.
+
+* New: Added support for per-culture sub-fonts within a composite font.
+
+    * This allows you to do things like create a Japanese specific Han sub-font to override the Han characters used in a CJK font. 
+
+* New: Safe Area is no longer centered or symmetrical. The UI will use the Safe Area reported by the OS. 
+
+* New: Added Editable Text functionality.
+
+##### Slate
+
+* Crash Fix: Fixed a crash when trying to use the widget reflector in a client cooked build.
+
+* Bug Fix: Fix for accessing Slate window on render thread.
+
+* Bug Fix: Updated Scroll Descendant Into View to account for situations where there are widget(s) between the target widget and the scroll panel, causing the widget geometry position to be incorrect.
+
+* Bug Fix: Fix slate material box brushes not being keyed off image size.
+
+* Bug Fix: Added conditional to only unregister the active timer if it is still valid after execution.
+
+* New: Slate now supports a Custom Boundary Navigation type which allows a custom handler if the boundary is hit. 
+
+    * This provides the ability to have normal navigation while within the boundary and the custom function only on the boundary. 
+
+    * This differs from Custom, which is a full override of the navigation behavior.
+
+##### UMG
+
+* Crash Fix: Fixed a crash on compile due to partial GC of a widget.
+
+* Crash Fix: Implemented a fix for a crash when playing in editor after changing, but not compiling, a Widget Blueprint.
+
+* Crash Fix: Calling Find All Ancestor Named Slot Host Widgets For Content with a null widget no longer crashes. 
+
+* Crash Fix: When the widget passed to Scroll Widget Into View is null no longer crashes.
+
+* Bug Fix: Fixed an issue where instanced properties on User Widgets might not be properly serialized.
+
+* Bug Fix: Fixed paused state so that it doesn't return widgets to their original state.
+
+* Bug Fix: UMG objects returned from Create Widget and Create Drag Drop Operation are no longer marked as always being hard referenced by the blueprint. 
+
+    * This means that they will no longer leak if they are not in use. 
+
+    * If the widget or drag operation are on screen they will correctly stay in memory, but if you need them not to be freed for other reasons, you can assign them to a blueprint variable. 
+
+* Bug Fix: Fixed Additive UI materials not being affected by the widget opacity.
+
+* Bug Fix: Common Activatable Panel does a better job of cleaning up when it is reset.
+
+* Bug Fix: Common Tree View Set Selection correctly updates list navigation and behaves similarly to Common List View Set Selected Item.
+
+* Bug Fix: Fixed Common Buttons sending OnClicked events while not interactable.
+
+* Bug Fix: UMG Designer now refreshes the widget preview even if the dragged widgets are dropped in an invalid location.
+
+* Bug Fix: Implemented a fix for large numbers of unused variables being added to the widget blueprint when the widget hierarchy was changed.
+
+* New: Exposed the virtual keyboard dismiss action to UMG.
+
+* New: Added Auto Wrap Blueprint-callable function to Text Block.
+
+* New: Now able to set escape keys and no key specified text on Input Key Selector.
+
+#### VR
+
+* Crash Fix: Fixed a crash that could occur when increasing and then decreasing screen percentage in the editor's VR preview.
+
+* Crash Fix: Fixed a bug on Gear VR that could cause the app to crash or hang when invoking the system quit menu.
+
+* Crash Fix: Fixed a crash that could occur on Oculus when creating a layer surface likely from a bad graphics state.
+
+* Crash Fix: Fixed a crash that could occur on Gear VR when using more than four stereo layers.
+
+* Bug Fix: Prevented the Front Buffer surface on Gear VR from needlessly updating. This improves performance.
+
+* Bug Fix: Fixed a perf regression introduced in 4.17 on Oculus platforms, where operating on the HMD was slowed by a superfluous module check.
+
+* Bug Fix: Motion Controller Component cleanup should now be in Begin Destroy and not Destructor. 
+
+* Bug Fix: Moved View Extension cleanup to Begin Destroy from the Destructor.
+
+* Bug Fix: Fixed Google Cardboard rendering upside down on some iOS devices.
+
+* Bug Fix: Fixed an issue with the Google VR controller model laser pointer portion wasn't rendering.
+
+* Bug Fix: Prevented the editor's camera from locking after playing in Google VR instant preview.  
+
+* Bug Fix: Set the default value of  CVar "vr.HiddenAreaMask" back to one, meaning that it is enabled by default.
+
+* Bug Fix: Fixed a bug causing world-locked stereo layer positions to be wrong on SteamVR when there was a camera component representing the HMD.
+
+* Bug Fix: Oculus Go touchpad event is now properly handled.
+
+* New: VR Screen Percentage Refactor to support new upsample and Dynamic Resolution rendering features. 
+
+* New: Added new Visualization settings to Motion Controller components. 
+
+* New: You can now input "bSupportsDash" under the [Oculus.Settings] heading in your engine config files to enable Oculus Dash support.
+
+* New: Added a render event name for the Oculus Rift's eye padding draw 
+
+* New: New code path for deciding whether to add a post processing pass for HMD distortion has been simplified to happen only in a single place.
+
+* New: Allow VR plugins to have multiple viewports and rotated eye orientations.
+
+* New: Added a new Controller model for the Oculus Go controller.
+
+* New: Added several Oculus specific functions that can be accessed through Blueprint. 
+
+* New: Enabled the debug canvas layer by default for the console command window on Gear VR and Google VR.
+
+#### XR
+
+* Bug Fix: Fixed a bug with Oculus, that caused the tracking origin to be reset on startup.
+
+* Bug Fix: Fixed a visual bug that could produce a strobe effect with lighting when using the VR preview in editor.
+
+* Bug Fix: Fixed a soft assert that would trigger when exiting the editor's VR preview on Oculus.
+
+* Bug Fix: Disable Mobile Multi-View when running with Google cardboard.
+
+* New: Added a new Blueprint function that supplies a tracking-to-world space transform.
+
+* New: Added an option to Motion Controller components to display an optional model (see the Display Device Model option in the component's details). This can be set up to generically use the device model provided by the backing XR system (if the system has one).
+
+* New: Added new Blueprint functions to generically enumerate tracked devices, and another to request a model component from the backing XR system (NOTE: at the moment, not all systems support returning a model).
+
+* New: Added a "MotionDelayBuffer" service to the engine that hooks into Motion Controller components, and lets programmers render the components at a delayed rate (to simulate latency).
+
+* New: The XR Tracking System and Head Mounted Display interfaces have been updated to make handling of pose updates and late update more explicit.
+
+* New: Added a new function to convert poses from XR tracking space to Unreal world space.
+
+* New: Changed Motion Controller "Hand" property to a more generic “MotionSource” name field, and added a function to query for generic float parameters associated with the Motion Controller.
+
+* New: Added new Blueprint node for getting the current VR pixel density.
+
+* New: Refactored the Rift dynamic resolution implementation to use the new general Dynamic Resolution State interface.
+
+* Deprecated: XR Tracking System interface function Refresh Poses has been deprecated.
+
+## UPGRADE NOTES
+
+#### Animation
+
+##### Animation Blueprints
+
+* Bug Fix: We changed the order of how to get the material default value. Previous implementation was to get parent's default value first before itself, which was incorrect, so we fixed it by getting own default value if exists. However, if you were relying on the previous behavior, you're going to see the issue. Make sure you set correct default value for the material instance if it was incorrect.
+
+* Bug Fix: If this code executes (it will add a log warning if it does) it could cause problems with deterministic cooking / patch creation as new guids will be generated during cook. Nothing will break but patch sizes could be bigger than they need to be. Re-saving the problematic assets will correct the issue.
+
+##### Skeletal Mesh
+
+* The function name of SkeletalMeshComponent::PostBlendPhysics() has changed to FinalizeAnimationUpdate() for clarification.
+
+#### Blueprints
+
+* New: Step Into has been remapped to the F11 key. The Step Over command is now mapped to F10 instead.
+
+* User Defined Structure default values now correctly propagate to blueprints and other user structs the same way that native structs do, which may cause issues if your blueprints were depending on this not happening. If this is a problem you may need to modify your outer blueprint or struct and set some values back to empty or 0.
+
+#### Core
+
+* VisitTupleElements has had its parameters swapped - the visitor is now the first parameter, followed by the tuple (now multiple tuples). Any existing calls will need to swap arguments.
+
+* FPlatformAtomics::AtomicRead64() has been deprecated - the int64* overload of FPlatformAtomics::AtomicRead() should be used instead.
+
+* The bWasReloaded flag was removed from FModuleManager::LoadModule(), FModuleManager::LoadModuleChecked() and FModuleManager::LoadModuleWithFailureReason() - any calls to it should also remove this flag.
+
+* TWeakObjectPtr's constructor was silently casting away const - this has now been fixed and any use of TWeakObjectPtr will now warn if it is not const-correct.
+
+* UE4 containers now expect move-aware allocators, which means they need to support the MoveToEmpty(SrcAllocator) member function to move all elements across from SrcAllocator, leaving it empty. Any custom allocators should be updated to add this support and specialize the TAllocatorTraits<Allocator>::SupportsMove trait to be true.
+
+* FArchive::Logf() and checkf() macros will now fail to compile if they are not passed a literal formatting string - these should be fixed to use a formatting string like TEXT("%s").
+
+* If you have any build scripts that were looking for DevelopmentAssetRegistry.bin or other metadata files you will need to change them to look inside the new Metadata directory.
+
+* Code using PLATFORM_COMPILER_HAS_DEFAULTED_FUNCTIONS or PLATFORM_COMPILER_HAS_AUTO_RETURN_TYPES should be replaced with the assumption they are always 1.
+
+* ENamedThreads::RenderThread and ENamedThreads::RenderThread_Local are no longer publicly accessible - they should be replaced with ENamedThreads::GetRenderThread() and ENamedThreads::GetRenderThread_Local() respectively.
+
+* TVariant::GetType() now returns an EVariantTypes enum class instead of an int32. Any use of this return value as an int32 will need to be changed to EVariantTypes.
+
+* Bug Fix: PPF_DeltaComparison has been removed - please remove any use of it from your projects.
+
+* Bug Fix: Calls to Align(), AlignDown(), AlignArbitrary() and IsAligned() will no longer compile if passed a non-integer/pointer type. This often fires when trying to align an enum value - it should be explicitly cast to an integer type of the right size.
+
+* Bug Fix: Calling FReferenceCollector::AddReferencedObjects() with a container will no longer compile if the container has pointers to forward declared types. You should ensure the types are included before calling this function.
+
+#### Editor
+
+* The debug canvas for stats is always dpi scaled in editor and pie. Any existing tools that directly use viewport positions and size as positions for canvas elements will need to divide out dpi scale using 1/CanvasObject->GetDPIScale().
+
+* Bug Fix: Note: This only applies to perforce users NOT using ticket based perforce servers.The perforce source control login flow will no longer make any calls to the perforce api which results in a password being stored on the local users machine by perforce. We have determined that the security level of this feature is not high enough to continue using it. Users can now expect to have to type in their perforce password each time they open the editor. To work around this users can call p4 passwd directly from a command window to save a password on their machine if they so choose.
+
+##### Sequencer
+
+* 'Shared' templates are no longer supported in sequencer. Any previous usage should be easily converted to a shared execution token instead (FMovieSceneExecutionTokens::AddShared()). Such tokens are shared globally, rather than per sub-sequence which makes them more robust for shared responsibilities.
+
+#### Gameplay Framework
+
+* Bug Fix: AActor subclasses that intended to be tick when the tick type is Viewports Only and did not set ShouldTickIfViewportsOnly to return true, but instead overrode TickActor and did not filter on Viewports Only will now need to correctly implement the ShouldTickIfViewportsOnly virtual to achieve this behavior.
+
+#### Networking
+
+* UNetConnection::ClientWorldPackageName is now a private member and should accessed through the available Get/Set methods.
+
+* Bug Fix: The net.UseAdaptiveNetUpdateFrequency console variable is now disabled by default, since it was causing replication issues in many projects. If you want to keep using this feature, you may enable it in your project's configuration files.
+
+#### Rendering
+
+* We changed the Vertexbuffer layout from AoS to SoA because after this change we can bind the Input Attributes via SRV and load them in the shader directly (also from compute). Before this change it was very difficult and involved manual interpretation of the Raw Data in the Shader. We recommend everyone to follow this route and change their Vertexbuffers as we probably will deprecate mixed AoS Vertexbuffers in the Future.
+
+* Tube shaped point lights will be automatically rotated to compensate for this change. Dynamically generated point lights with a source length will need to have a 90 degrees pitch added to their orientation.
+
+* Inside Sequencer->Render->CompositionGraphOptionsSet Output Format to "Custom Render Passes"Set Include Render Passes to "Final Image"Set Capture Frames in HDR "On"Set Capture Gamut to "Linear"
+
+#### UI
+
+##### UMG
+
+* Bug Fix: If a game was depending on this functionality to keep things loaded, change your blueprint to set the return value into a variable, that will keep it hard referenced normally
+
+#### XR
+
+* Deprecated the MotionController component's EControllerHand 'Hand' property and replaced it with a 'MotionSource' name. To support the same functionality, use names corresponding to the old enum values ("Left", "Right", etc.).
+
+* vr.oculus.PixelDensity.adaptive has been removed. Dynamic resolution needs to be enabled using the general engine system now.
+
+## PROGRAMMING UPGRADE NOTES
+
+#### AI
+
+* Bug Fix: Fixed a bug in UAISense_Sight::UnregisterSource that resulted in sources being removed not getting removed from the ObservedTargets collection.
+
+##### Debugging Tools
+
+* Redone how visual log is using categories. Now it's using FName rather than FLogCategoryBase to indicated log category. All UE_VLOG macros have been updated. This change made it possible to actually vislog things using blueprint functions.
+
+##### Navigation
+
+* Bug Fix: Fixed navmesh not loading properly in PIE when owning world has been duplicated-for-play
+
+#### Animation
+
+* Bug Fix: Now when meshes are set to EMeshComponentUpdateFlag::AlwaysTickPose, we optionally kick of a task to perform parallel update only (no evaluation). Previously we needed to run this on the game thread as if we skipped evaluation we did not run a worker task.
+
+#### Audio
+
+* New: The new Microphone Capture Component enables feeding audio directly into a game for use with the Audio Engine and for driving in-game parameters. Limitations:
+
+    * Only implemented on PC.
+
+    * Only for local client microphone capture, not connected to VoIP system.
+
+    * Doesn't support recording microphone output to assets.
+
+#### Automation
+
+* Added automated test per-platform blacklist, defined in the project DefaultEngine.ini or platform equivalent e.g. MacEngine.ini:	[AutomationTestBlacklist]	;Use format +BlacklistTest=(Map=/Game/Tests/MapName, Test=TestName)	+BlacklistTest=(Map=/Game/Tests/Physics/Stability, Test=StackCubes)
+
+* Automated test screenshots are now able to test against a hierarchical fallback platform and RHI. This is defined in DefaultEngine.ini, for example:	[AutomationTestFallbackHierarchy]	;Use format +FallbackPlatform=(Child=/Platform/RHI, Parent=/Platform/RHI)	+FallbackPlatform=(Child=/Windows/OpenGL_SM4, Parent=/Windows/D3D11_SM5)	+FallbackPlatform=(Child=/Mac/SM5, Parent=/Windows/D3D11_SM5)
+
+#### Core
+
+* Any usage of FStartupPackages in project code should be removed
+
+* "Native" (now called “FNativeFuncPtr”) is now a function pointer that takes a UObject* context, rather than a UObject member function pointer. Use “P_THIS” if you were previously using the “this” pointer in your native function.
+
+* FTickableBaseObject::GetTickableType can be overriden and returns an enum value of either Conditional (default), Always, or Never.
+
+    * Tickable Never objects will not get added to the tickable array or ever evaluated.
+
+    * Tickable Always objects do not call IsTickable and assume it will return true.
+
+    * Tickable Conditional objects work as in the past with IsTickable called each frame to make the determination whether to call Tick or not.
+
+* FTickableObjectBase::IsTickable() is no longer a pure virtual.
+
+* To match the change in functionality EResourceSizeMode::Inclusive was changed to EstimatedTotal. If any of your code was using Inclusive, you need to decide whether you want to change to EstimatedTotal or Exclusive.
+
+* The ResourceSize asset registry tag has been removed as it was inaccurate. To get an accurate estimate, you must load the object and call GetResourceSize().
+
+* Bug Fix: Cleaned up delegate definitions in UObjectGlobals.h and moved several to be WITH_EDITOR. If you were using those definitions in a game you will need to wrap your use in WITH_EDITOR.
+
+#### Editor
+
+* Being an early feature, some Editor UI buttons do not work during playback of demos.
+
+* The debug canvas for stats is always dpi scaled in editor and pie. Any existing tools that directly use viewport positions and size as positions for canvas elements will need to divide out DPI scale using 1/CanvasObject->GetDPIScale().
+
+* RenameAssetsWithDialog() is often recommended in places where RenameAssets() was being used.
+
+* The SizeMap and ReferenceViewer have been merged into the AssetManagerEditor plugin so they can share the data needed to show cooked sizes.
+
+#### Gameplay Framework
+
+* C++ Header cleanups for Actor, Pawn, Character, Controller, and PlayerController. Marked several fields as private that were accidentally exposed in 4.17.
+
+* Projects using this plugin are strongly encouraged to add a project specific subclass of UDataValidationManager. Any objects that require validation should override IsDataValid.
+
+* Overriden versions of USignificanceManager::Update will need to change signature to receive a TArrayView instead of const TArray&.
+
+#### Networking
+
+* Bug Fix: Programmers using HasPendingConnection() on a listening FScocket should consider using the new WaitForPendingConnection() function instead, which prevents wasting CPU cycles due to polling. This change supported a change to the behavior of the FTcpListener constructor's InSleepTime parameter which now specifies the maximum time to wait for a connection in its polling loop. It now defaults to 1 second, but using a value of 0 will restore the previous polling behavior.
+
+#### Online
+
+* This Online Presence Test requires that the testing user to have friends and be connected to the platform.Usage is "online test presence [OPTIONAL ID]". The ID passed should be an arbitrary non-friend. If no ID is passed, the arbitrary lookup test is skipped. This test does wait for presence updates for about two minutes, after this time if no updates are received, the harness will mark the test as a failure.
+
+#### Rendering
+
+* The STRONG_TYPE #define now seen in UE4 .usf/HLSL is actually hlslcc's "invariant" keyword applied as a type-qualifier to a Buffer<>/RWBuffer<> type - only valid when using Metal which exports this through ILanguageSpec and #define'd out for everyone else. When this is applied the MetalBackend will emit a "raw" buffer whose type must match the shader exactly, which avoids the overheads of Linear-Textures & Function-Constants.
+
+* Vertexbuffer Objects now bind their SRVs directly instead of filling out a StreamDescriptor. Care must be taken when the Vertexfactory is initialized before the VertexBuffer as the bound SRV will become invalid if the underlying vertexbuffer is re-created. We do not support late binding of the data anymore.
+
+##### Mobile Rendering
+
+* Removed cvar "r.AllReceiveDynamicCSM" and primitive flag “FPrimitiveSceneProxy.bReceiveCombinedCSMAndStaticShadowsFromStationaryLights”. Added “FPrimitiveSceneProxy.bReceiveMobileCSMShadows”, which now applies to both stationary and movable lights and defaults to on.
+
+#### Tools
+
+* UEdGraphPin::PinCategory, UEdGraphPin::PinSubcategory, and UEdGraphPin::PinName are now stored as FName instead of FString.UEdGraphPin::CreatePin has several simplified overrides so you can only specify Subcategory or SubcategoryObject or neither.UEdGraphPin::CreatePin also takes a parameter bundle for reference, const, container type, index, and value terminal type rather than a long list of default parameters.Material Expressions now store input and output names as FName instead of FString.FNiagaraParameterHandle now stores the parameter handle, namespace, and name as FName instead of FString.Most existing pin related functions using string have been deprecated.
+
+#### UI
+
+##### Slate
+
+* Renamed FTextBlockLayout to FSlateTextBlockLayout to reflect that it's a Slate specific type
+
+#### XR
+
+* If you maintain a VR or AR plugin, you will have to remove your RefreshPoses() override and move any late update functionality to OnBeginRendering_RenderThread/GameThread(). Alternatively, if your plugin does not support late update, you can remove RefreshPoses() completely and override DoesSupportLateUpdate(), returning false.
+
+    * Added two methods to the IXRTrackingSystem API - GetTrackingToWorldTransform() and UpdateTrackingToWorldTransform(). GetTrackingToWorldTransform() is expected to return an transform that will convert poses from GetCurrentPose() into world space. UpdateTrackingToWorldTransform() is meant for users to override the transform since the anchor between the two spaces is controlled by the project.
+
+    * MotionControllerComponents' 'Hand' property is no longer available. Instead, switch to using the new 'MotionSource' field with text names matching the old enum values ("Left", "Right", etc.). For ease of use, "Left" and "Right" are predefined in FXRMotionControllerBase (as LeftHandSourceId & RightHandSourceId).
+
+    * Replaced EControllerHand parameters in the IMotionController API with generic "source" name parameters. To handle backwards compatibility, use name values corresponding to the old enum ("Left", "Right", etc.).
+
+    * Added two new functions to the IMotionController API - EnumerateSources() & GetCustomParameterValue(). Use EnumerateSources() to return a list of the supported devices (replaces the static EControllerHand enum). GetCustomParameterValue() can be used to expose float values associated with the motion device (such as FOV with cameras, etc.).
+
+#### VR
+
+* The GetHMDDistortionEnabled method on IHeadMountedDisplay how takes a single argument of type EShaderPath. If you maintain a custom HMD plugin, you will need to update your method signature.
+
+## DEPRECATIONS
+
+#### Automation
+
+* UFunctionalTestingManager::RunAllFunctionalTests has dropped an unused parameter. The original version is now deprecated.
+
+#### Core
+
+* Removed the PLATFORM_HAS_64BIT_ATOMICS macro, as it is always 1.
+
+* Deprecated the VARARG macros and replaced them with variadic templates.
+
+#### Editor and Tools
+
+* Deprecated Material Baking functionality inside of the MaterialUtilities module. Users should use the functionality provided by the Material Baking module instead.
+
+#### Gameplay Framework
+
+* Add APlayerController::ClientReturnToMainMenuWithTextReason to support an FText reason parameter so that the message can be localized. Deprecated APlayerController::ClientReturnToMainMenu.
+
+#### Rendering
+
+* The Editor is now only supported on Feature Level 5 platforms. Feature Level 4 platforms (D3D10, OpenGL3) are no longer supported.
+
+##### Materials
+
+* Material Layer experimental support feature, Shared Input Collection and associated material graph nodes have been removed. See the forum feedback thread for more insight.
+
+#### Tools
+
+##### UnrealBuildTool
+
+* The bOutputToEngineBinaries flag has been removed from the target rules. The output directory for targets is always determined by the location of that target.
+
+#### XR
+
+* The IHeadMountedDisplay::GetHMDDeviceType() method has been removed as it was not extensible enough. Code that needs know which XR plugin is currently active should use IXRTrackingSystem::GetSystemName() instead.
