Copying //UE4/Dev-Documentation to //UE4/Samples-Main (Source: //UE4/Dev-Documentation @ 4321946)

============================
  MAJOR FEATURES & CHANGES
============================

Change 3683933 by [REDACTED_PERSON]

	Updated metadata

Change 4094604 by [REDACTED_PERSON]

	Fixing a bug in Match3 with incorrect attempts at unbinding a delegate. Not as much of a code problem as a bad learning example.

Change 4099126 by [REDACTED_PERSON]

	#loc UE4DocKOR

Change 4105127 by [REDACTED_PERSON]

	removed wrong [include]

Change 4105128 by [REDACTED_PERSON]

	#loc UE4DocKor

Change 4169203 by [REDACTED_PERSON]

	#loc UE4DocKOR

Change 4201118 by [REDACTED_PERSON]

	Initial check-in of 4.20 Release Notes - Major features only.

Change 4201313 by [REDACTED_PERSON]

	misc fixes

Change 4201319 by [REDACTED_PERSON]

	Adding 4.20 minor release notes.

Change 4201528 by [REDACTED_PERSON]

	misc fixes

Change 4216361 by [REDACTED_PERSON]

	#loc UE4DocKOR moved or deleted

Change 4216416 by [REDACTED_PERSON]

	#loc UE4DocKOR

Change 4216563 by [REDACTED_PERSON]

	#loc UE4DocKOR

Change 4217772 by [REDACTED_PERSON]

	Adding some missing minor release notes in the Slate and UMG sections.

Change 4218292 by [REDACTED_PERSON]

	Commenting out Shotgun documentation link and adding a "coming soon" blurb

Change 4219358 by [REDACTED_PERSON]

	#loc UE4DocCHN moved or deleted

Change 4219564 by [REDACTED_PERSON]

	#loc UE4DocKOR

Change 4220008 by [REDACTED_PERSON]

	Reordered top features; updated language for Niagara

Change 4220814 by [REDACTED_PERSON]

	Updating Niagara landing page and adding new image.

Change 4220945 by [REDACTED_PERSON]

	Updating page with note pointing to digital humans page and added related page link for digital humans

Change 4221079 by [REDACTED_PERSON]

	Submitting page and images for Working With Tracks In Sequencer.

Change 4221223 by [REDACTED_PERSON]

	Checking in updates to Sequence Recorder page and new images.

Change 4222174 by [REDACTED_PERSON]

	#loc UE4DocCHN check in 420 Release Note againest INT verison 4218292  (need update againest for latest release note)

Change 4222188 by [REDACTED_PERSON]

	#loc UE4DocKOR moved or deleted

Change 4222189 by [REDACTED_PERSON]

	#loc UE4DocCHN 420 Release Note Update against INT 4220008

Change 4222217 by [REDACTED_PERSON]

	#loc UE4DocCHN moved or deleted

Change 4222310 by [REDACTED_PERSON]

	#loc UE4DocCHN 420 Release Note Typo Fixing against INT 4220008

Change 4222335 by [REDACTED_PERSON]

	#loc UE4DocCHN TwinblastBust Update against INT 4220945

Change 4222348 by [REDACTED_PERSON]

	#loc UE4DocCHN SequenceRecorder Update against INT 3866572

Change 4222355 by [REDACTED_PERSON]

	#loc UE4DocCHN SequenceRecorder Update against INT 3244012

Change 4222371 by [REDACTED_PERSON]

	#loc UE4DocCHN

Change 4222713 by [REDACTED_PERSON]

	#loc UE4DocKOR

Change 4223218 by [REDACTED_PERSON]

	Adding Controlling Anim Instances with Sequencer page and images

Change 4223330 by [REDACTED_PERSON]

	Minor update to formatting.

Change 4223439 by [REDACTED_PERSON]

	Submitting new TimeRefactorNotes page and images

Change 4224395 by [REDACTED_PERSON]

	#loc UE4DocKOR moved or deleted

Change 4224481 by [REDACTED_PERSON]

	#loc UE4DocCHN 419 Release Note Update against INT 3941748

Change 4224497 by [REDACTED_PERSON]

	#loc UE4DocCHN moved or deleted

Change 4228591 by [REDACTED_PERSON]

	#loc UE4DocCHN Localization Edit Against 4221079

Change 4228596 by [REDACTED_PERSON]

	#loc UE4DocCHN

Change 4228598 by [REDACTED_PERSON]

	#loc UE4DocCHN Localization Bug Fixing

Change 4228600 by [REDACTED_PERSON]

	#loc UE4DocCHN Localization Edit

Change 4233224 by [REDACTED_PERSON]

	#loc UE4DocCHN Localization Bug Fixing

Change 4233437 by [REDACTED_PERSON]

	#loc UE4DocKOR

Change 4233440 by [REDACTED_PERSON]

	#loc UE4DocKOR

Change 4240164 by [REDACTED_PERSON]

	#loc UE4DocKOR

Change 4242799 by [REDACTED_PERSON]

	#loc UE4DocCHN Localization Editor

Change 4243822 by [REDACTED_PERSON]

	Manually moveing these files over from the 4.20 release branch.

Change 4257108 by [REDACTED_PERSON]

	Switch ActionRPG to use IWYU headers

Change 4257489 by [REDACTED_PERSON]

	#loc UE4DocKor

Change 4258770 by [REDACTED_PERSON]

	Updating filters page and images with 4.20 changes

Change 4258956 by [REDACTED_PERSON]

	Updating Color Picker doc and images for 4.20

Change 4261325 by [REDACTED_PERSON]

	#loc UE4DocKOR

Change 4261534 by [REDACTED_PERSON]

	Updating mobilepreviewer doc for 4.20 changes.

Change 4263566 by [REDACTED_PERSON]

	Replacing images to address doc-md tool erros. #UE4Doc #fyi Jeff.Wilson, Mitchell.Wilson

Change 4264102 by [REDACTED_PERSON]

	#loc UE4DocCHN Localization Edit

Change 4264104 by [REDACTED_PERSON]

	#loc UE4DocCHN Localization Edit

Change 4264120 by [REDACTED_PERSON]

	#loc UE4DocCHN Localization Bug Fixing

Change 4264799 by paulo.souza

	Fix ActionRPG main menu logo "action" and "sample" labels were misplaced #fyi sam.deiter #rb none

Change 4267571 by [REDACTED_PERSON]

	#loc UE4DocKOR

Change 4268090 by [REDACTED_PERSON]

	#loc UE4DocKOR

Change 4272340 by [REDACTED_PERSON]

	Adding all MR pages for 4.20

Change 4273092 by [REDACTED_PERSON]

	Updating Mesh Details page for 4.20

Change 4276731 by [REDACTED_PERSON]

	Checking in nDisplay doc.

Change 4276770 by [REDACTED_PERSON]

	Updated formatting

Change 4276798 by [REDACTED_PERSON]

	Minor formatting update

Change 4276998 by [REDACTED_PERSON]

	Checking in nDisplay changes for 4.20

Change 4278693 by [REDACTED_PERSON]

	#loc UE4DocKOR moved or deleted

Change 4279407 by [REDACTED_PERSON]

	Adding ScriptingAndAutomation folder and content

Change 4279418 by [REDACTED_PERSON]

	Updating crumbs

Change 4279440 by [REDACTED_PERSON]

	Adding INT to file.

Change 4280273 by paulo.souza

	ARPG - Player AI can now use Potions and Skills when possible/necessary

Change 4280342 by [REDACTED_PERSON]

	Adding blueprint/python switchable content.

Change 4282437 by [REDACTED_PERSON]

	#loc UE4DocCHN moved or deleted

Change 4282603 by [REDACTED_PERSON]

	Outsourcing 20180814

Change 4282681 by [REDACTED_PERSON]

	#loc UE4DocCHN English Update

Change 4282721 by [REDACTED_PERSON]

	#loc UE4DocCHN Localization Editor 20180807

Change 4282938 by [REDACTED_PERSON]

	#loc UE4DocKOR

Change 4283510 by paulo.souza

	ARPG - Initial support for consumable items (health and mana potions)

Change 4283679 by [REDACTED_PERSON]

	Submitting updates to some ScriptingAndAutomation pages and adding HowTos sections.

Change 4283941 by [REDACTED_PERSON]

	Submitting new page and content

Change 4284256 by [REDACTED_PERSON]

	Minor updates to fix topic lists

Change 4284421 by [REDACTED_PERSON]

	Checking in python page and content

Change 4284468 by [REDACTED_PERSON]

	Minor updates.

Change 4284888 by [REDACTED_PERSON]

	Adding new page and content

Change 4285161 by [REDACTED_PERSON]

	Adding new page and content

Change 4285387 by [REDACTED_PERSON]

	minor styling updates

Change 4287470 by [REDACTED_PERSON]

	Adding new page and content

Change 4291477 by [REDACTED_PERSON]

	#loc UE4DocKOR moved or deleted

Change 4291478 by [REDACTED_PERSON]

	#loc UE4DocKOR

Change 4291995 by [REDACTED_PERSON]

	#loc UE4DocKOR

Change 4292193 by [REDACTED_PERSON]

	Updating page and content

Change 4292827 by [REDACTED_PERSON]

	adding digital humans page and content

Change 4293098 by [REDACTED_PERSON]

	Updating CascadeToNiagara and adding some content

Change 4295473 by [REDACTED_PERSON]

	#loc UE4DocKOR moved or deleted

Change 4295612 by [REDACTED_PERSON]

	#loc UE4DocKOR

Change 4295997 by [REDACTED_PERSON]

	#loc UE4DocKOR

Change 4296011 by [REDACTED_PERSON]

	#loc UE4DocKOR

Change 4296954 by [REDACTED_PERSON]

	Fixing LODs on one mesh to resolve flicker. Rebuilt lighting

Change 4297342 by [REDACTED_PERSON]

	Updating copyright text in Match3 credits

Change 4298354 by [REDACTED_PERSON]

	Updating Photorealistic Character Bust eye material to be consistent with Content Examples setup.

Change 4299738 by [REDACTED_PERSON]

	#loc UE4DocCHN moved or deleted

Change 4299936 by [REDACTED_PERSON]

	#loc UE4DocKOR

Change 4300185 by [REDACTED_PERSON]

	Updating page to 4.20

Change 4300194 by [REDACTED_PERSON]

	Updating page and one image

Change 4300247 by [REDACTED_PERSON]

	Fixed publishing errors

Change 4300533 by [REDACTED_PERSON]

	Updates for publishing

Change 4300884 by [REDACTED_PERSON]

	Updating page and content

Change 4301216 by [REDACTED_PERSON]

	Updating page and content

Change 4301268 by [REDACTED_PERSON]

	Updating page and content

Change 4301301 by [REDACTED_PERSON]

	Updating page and content

Change 4302195 by [REDACTED_PERSON]

	Updating page and adding content

Change 4303253 by [REDACTED_PERSON]

	#loc UE4DocCHN Localization Edit 20180414

Change 4303325 by [REDACTED_PERSON]

	#loc UE4DocKOR

Change 4303706 by [REDACTED_PERSON]

	#loc UE4DocKOR

Change 4305335 by [REDACTED_PERSON]

	Fixing errors on pages in prep for publish

Change 4305416 by [REDACTED_PERSON]

	Removing Graph Panel section to resolve errors since it was already commented out.

Change 4305436 by [REDACTED_PERSON]

	Removing graph panel from commented section to resolve errors

Change 4305504 by [REDACTED_PERSON]

	Fixing up an image to resolve errors

Change 4305576 by [REDACTED_PERSON]

	Updated graph panel link to new location.

Change 4305778 by [REDACTED_PERSON]

	Fixing broken comment

Change 4306900 by [REDACTED_PERSON]

	#loc UE4DocKOR

Change 4307010 by [REDACTED_PERSON]

	Updating this page to resolve issues on another

Change 4307038 by [REDACTED_PERSON]

	Removing missing image references, but leaving the descriptions if they ever need to be added back. images are commented out.

Change 4307067 by [REDACTED_PERSON]

	Removing outdated note to resolve errors.

Change 4307073 by [REDACTED_PERSON]

	Removing old link that was commented out to resolve errors

Change 4307105 by [REDACTED_PERSON]

	Removing old image that is no longer used in the new CMS

Change 4307127 by [REDACTED_PERSON]

	Fixing up comment and removing section that is not used to resolve errors

Change 4307254 by [REDACTED_PERSON]

	Removing some old URL references that were already commented out to resolve errors.

Change 4307277 by [REDACTED_PERSON]

	Fixing up incorrect link to resolve errors

Change 4307289 by [REDACTED_PERSON]

	Fixing up incorrect URL to resovle errors

Change 4307291 by [REDACTED_PERSON]

	Updating URL for a page that was moved.

Change 4307670 by [REDACTED_PERSON]

	Fixing up multiple issues to resolve errors

Change 4308120 by [REDACTED_PERSON]

	Fixing up comments to resolve errors

Change 4308165 by [REDACTED_PERSON]

	Removing old commented section to resolve errors

Change 4308189 by [REDACTED_PERSON]

	Removing commented out related pages section to resolve errors.

Change 4308942 by [REDACTED_PERSON]

	Removing topic image link to resolve error

Change 4309000 by [REDACTED_PERSON]

	Fixing up some typos in images to resolve errors

Change 4309042 by [REDACTED_PERSON]

	Fixing up old URL to resolve errors

Change 4309060 by [REDACTED_PERSON]

	Updating URLs to resolve crashing

Change 4309073 by [REDACTED_PERSON]

	Updating URLs to resolve crash

Change 4309085 by [REDACTED_PERSON]

	Removing old URL located in a comment to resolve error

Change 4309106 by [REDACTED_PERSON]

	Updating related pages to resolve errors

Change 4309130 by [REDACTED_PERSON]

	Fixing topic image to resolve error

Change 4309167 by [REDACTED_PERSON]

	Remove include to resolve errors

Change 4309172 by [REDACTED_PERSON]

	Removing include to resolve errors

Change 4309183 by [REDACTED_PERSON]

	Fixing up table to resolve error

Change 4309189 by [REDACTED_PERSON]

	Removing old image to resolve errors

Change 4310574 by [REDACTED_PERSON]

	#loc UE4DocKOR

Change 4310769 by [REDACTED_PERSON]

	#loc UE4DocKOR

Change 4311446 by [REDACTED_PERSON]

	Adding missing image file

Change 4312754 by [REDACTED_PERSON]

	Added new DoF pages and updates

Change 4313881 by [REDACTED_PERSON]

	#loc UE4DocKOR moved or deleted

Change 4313999 by [REDACTED_PERSON]

	#loc UE4DocCHN moved or deleted

Change 4314037 by [REDACTED_PERSON]

	corrected some links, tags, and spaces

Change 4314071 by [REDACTED_PERSON]

	#loc UE4DocKOR

Change 4314302 by [REDACTED_PERSON]

	#loc UE4DocKOR

Change 4314504 by [REDACTED_PERSON]

	Adding missing tutorial image and updating tutorial.

Change 4317861 by [REDACTED_PERSON]

	#loc ?????? -> ??????

Change 4320148 by [REDACTED_PERSON]

	Removing comments on topics

Change 4320159 by [REDACTED_PERSON]

	Fixed async gifs

Change 4320499 by [REDACTED_PERSON]

	Checking in topic images

Change 4321330 by [REDACTED_PERSON]

	#loc UE4DocKOR

Change 4321337 by [REDACTED_PERSON]

	Outsourcing 20180821

Change 4321362 by [REDACTED_PERSON]

	#loc UE4DocCHN Localization Bug Fixing 2080828

diff --git a/Engine/Documentation/Source/Engine/Animation/NodeReference/Blend/BlendNodes.INT.udn b/Engine/Documentation/Source/Engine/Animation/NodeReference/Blend/BlendNodes.INT.udn
--- a/Engine/Documentation/Source/Engine/Animation/NodeReference/Blend/BlendNodes.INT.udn
+++ b/Engine/Documentation/Source/Engine/Animation/NodeReference/Blend/BlendNodes.INT.udn
@@ -114,7 +114,7 @@ The **Layered blend per bone** node performs a masked blend that only affects a
 
 The **Slot** node plays the resultant animation from a given animation slot. 
 
-[INCLUDE:Engine/Animation/AnimMontage#slots]
+
 
 For more details on using Slots with the Animation Montage system, please see [](Engine/Animation/AnimMontage).
 
diff --git a/Engine/Documentation/Source/Engine/Animation/Persona/BasePoseManager/BasePoseManager.INT.udn b/Engine/Documentation/Source/Engine/Animation/Persona/BasePoseManager/BasePoseManager.INT.udn
--- a/Engine/Documentation/Source/Engine/Animation/Persona/BasePoseManager/BasePoseManager.INT.udn
+++ b/Engine/Documentation/Source/Engine/Animation/Persona/BasePoseManager/BasePoseManager.INT.udn
@@ -19,113 +19,119 @@ Related: Engine/Animation/Sequences
 
 [TOC(start:2 end:2)]
 
-The **Retarget Manager** within the **Skeleton Editor** can be used to manage your retarget sources, set up rigs and define retargeted base poses for use with [Animation Retargeting](Engine\Animation\AnimHowTo\Retargeting).
+The **Retarget Manager** within the **Skeleton Editor** can be used to manage your retarget sources, set up rigs and define retargeted base poses for use with [Animation Retargeting](Engine/Animation/AnimHowTo/Retargeting).
 
 ## Manage Retarget Source
 
-Since retargeting uses a [Skeleton](Engine/Animation/Skeleton) asset and the Skeleton has its original proportions defined by the Skeletal Mesh for which it was initially created, it follows that one-directional retargeting will work smoothly in most cases. 
-However, there will often come times when you will have special animations that are built just for the variant versions. 
-For example, suppose you have multiple characters that share the same Skeleton asset (a base character, a short character, and a tall character) and you have a special animation that was created just for the tall version of the character.
+Since retargeting uses a [Skeleton](Engine/Animation/Skeleton) asset and the Skeleton has its original proportions defined by the Skeletal Mesh for which it was initially created, it follows that one-directional retargeting will work smoothly in most cases. However, there will often come times when you will have special animations that are built just for the variant versions. For example, suppose you have multiple characters that share the same Skeleton asset (a base character, a short character, and a tall character) and you have a special animation that was created just for the tall version of the character.
 
-If you imported this new tall-only animation in, you would still need to use the same Skeleton asset as before, which was created from the base version of the character. 
-This will cause the new animation's proportions to be distorted. 
-The solution is to use the **Manage Retarget Source** option within the **Retarget Manager** which allows you to associate an Animation Sequence with the actual Skeletal Mesh for which it was designed. 
-In this way, any problems of retargeting for special animations will be fixed.
+If you imported this new tall-only animation in, you would still need to use the same Skeleton asset as before, which was created from the base version of the character. This will cause the new animation's proportions to be distorted. The solution is to use the **Manage Retarget Source** option within the **Retarget Manager** which allows you to associate an Animation Sequence with the actual Skeletal Mesh for which it was designed. In this way, any problems of retargeting for special animations will be fixed.
 
 ### Adding Retarget Sources
 
-1. Inside the **Skeleton Editor**, from the **Toolbar** click the **Retarget Source Manager** button. 
+1.  Inside the **Skeleton Editor**, from the **Toolbar** click the **Retarget Source Manager** button.
+    
+    ![](RetargetManager.png)
+    
+2.  Click the **Add New Retarget Source** button.
+    
+    ![](AddSource.png)
+    
+3.  Choose the **Skeletal Mesh** for which the special animation was created.
+    
+    ![](SelectSkeletalMesh.png)
+    
+    You will now see that Skeletal Mesh listed in the **Retarget Manager**.
+    
+    ![](Created.png)
+    
+4.  Open the special case **Animation Sequence** intended for your specific **Skeletal Mesh**.
+    
+5.  In the **Asset Details** panel, locate the **Animation** category and locate the **Retarget Source** property and choose your **Skeletal Mesh** from the drop-down menu.
+    
+    ![](DropDownSelection.png)
+    
+    By choosing this mesh, you are specifying that you want to use the proportions of that mesh to drive this animation.
+    
 
-	![](RetargetManager.png)
-
-1. Click the **Add New Retarget Source** button. 
-
-	![](AddSource.png)
-
-1. Choose the **Skeletal Mesh** for which the special animation was created.
-
-	![](SelectSkeletalMesh.png)
-
-	 You will now see that Skeletal Mesh listed in the **Retarget Manager**.
-
-	![](Created.png)
-
-1. Open the special case **Animation Sequence** intended for your specific **Skeletal Mesh**. 
-
-1. In the **Asset Details** panel, locate the **Animation** category and locate the **Retarget Source** property and choose your **Skeletal Mesh** from the drop-down menu.
-
-	![](DropDownSelection.png)
+## Set up Rig
 
-	 By choosing this mesh, you are specifying that you want to use the proportions of that mesh to drive this animation.
+The middle section of the Retarget Manager allows you assign a **Rig** to the Skeleton which can be used to pass animation data to a different Skeleton that uses the same Rig.
 
+![Set up Rig-1.png](Set20up20Rig-1.png "Set up Rig-1.png")
 
-## Set up Rig
+This process is needed in order to perform any animation retargeting for characters that use different Skeleton assets. 
 
-The middle section of the Retarget Manager allows you assign a **Rig** to the Skeleton which can be used to pass animation data to a different Skeleton that uses the same Rig. 
+You can select the Rig to use from the **Select Rig** drop-down option where a **Humanoid** option is available which you will want to select for most characters.
 
-![](SetUpRig.png)
+![](HumanoidSetup-1.png "HumanoidSetup-1.png")
 
-This process is needed in order to perform any animation retargeting for characters that use different Skeleton assets. 
-You can select the Rig to use from the **Select Rig** drop-down option where a **Humanoid** option is available which you will want to select for most characters. 
+Once you assign the Humanoid Rig, you will need to assign Bones from the Skeleton that correspond to the same (or similar) location as the Node on the Rig. You can use the drop-down menu for Node and manually assign the corresponding Bone from your Skeleton, or you can use the **Auto Mapping** feature located at the top of the menu. This will look through your Skeleton and attempt to find the best matching Bone for each Node on the Rig.
 
-![](HumanoidRig.png)
+The **Clear Mapping** button will wipe all the currently assigned Bones from their corresponding Node assignment. You can also use the **Save **and **Load **buttons to save the current mapping assignment or load a previously saved mapping assignment from a **Node Mapping Container **asset (shown below). 
 
-Once you assign the Humanoid Rig, you will need to assign Bones from the Skeleton that correspond to the same (or similar) location as the Node on the Rig. 
-You can use the drop-down menu for Node and manually assign the corresponding Bone from your Skeleton, or you can use the **Auto Mapping** feature located at the top of the menu. 
-This will look through your Skeleton and attempt to find the best matching Bone for each Node on the Rig. 
+[OBJECT:EmbeddedVideo]
+    [PARAMLITERAL:videoid]
+    qna6b-kAN48  <!-- Insert the youtube video link here -->
+    [/PARAMLITERAL]
+    [PARAMLITERAL:width]
+    78
+    [/PARAMLITERAL]
+    [PARAMLITERAL:height]
+    56.25
+    [/PARAMLITERAL]
+    [PARAMLITERAL:units]
+    %
+    [/PARAMLITERAL]
+    [PARAMLITERAL:type]
+    simple
+    [/PARAMLITERAL]
+[/OBJECT]
 
-The **Clear Mapping** button will wipe all the currently assigned Bones from their corresponding Node assignment while the **Show Advanced** button will allow you to assign additional Node/Bone correlations for things like finger digits, IK Bones or Twist Bones. 
-Once you set up the Rig for your source Skeleton (or the Skeleton asset that drives the animation that you want to retarget to another character), you will need to go into the Skeleton of the target Skeleton and assign the same Rig and define the Bones in the new Skeleton that best match up with the Nodes on the Rig. 
+The **Show Advanced** button will allow you to assign additional Node/Bone correlations for things like finger digits, IK Bones or Twist Bones. Once you set up the Rig for your source Skeleton (or the Skeleton asset that drives the animation that you want to retarget to another character), you will need to go into the Skeleton of the target Skeleton and assign the same Rig and define the Bones in the new Skeleton that best match up with the Nodes on the Rig.
 
 [REGION:note]
-Please see [](Engine\Animation\AnimHowTo\Retargeting) for a step-by-step guide on setting up a Rig to retarget animations between characters that use different Skeletons.
+Please see [Using Retargeted Animations](Engine/Animation/AnimHowTo/Retargeting) for a step-by-step guide on setting up a Rig to retarget animations between characters that use different Skeletons.
 [/REGION]
 
-
-
-
 ## Manage Retarget Base Pose
 
-There may be instances where you want to retarget animations to a Skeleton that is not quite using the same base pose as your Source Skeleton. 
+There may be instances where you want to retarget animations to a Skeleton that is not quite using the same base pose as your Source Skeleton.
 
-Take for example the Skeleton in the image below on the right when attempting to retarget an animation:  
+Take for example the Skeleton in the image below on the right when attempting to retarget an animation:
 
 ![](ManageBasePose1.png)(w:840)
 
-Our Source Skeleton (left) is in a A-Pose while our Target Skeleton is using an T-Pose. If we were to retarget animations as is now, we may encounter issues: 
+Our Source Skeleton (left) is in a A-Pose while our Target Skeleton is using an T-Pose. If we were to retarget animations as is now, we may encounter issues:
 
 ![](TargetAnimation.png)(w:540)![](BadBasePose.png)(w:550)
 
-Above we have our target animation on the left of a character holding a shotgun, however, when we retarget the animation to our new character on the right (because they are using different base poses) the arm positioning is not correct. 
-We can fix this by Retargeting the Base Pose inside the **Retarget Manager** which allows us to define a Retargeted Base Pose to use for animation retargeting purposes rather than using the characters normal Base Pose. 
+Above we have our target animation on the left of a character holding a shotgun, however, when we retarget the animation to our new character on the right (because they are using different base poses) the arm positioning is not correct. We can fix this by Retargeting the Base Pose inside the **Retarget Manager** which allows us to define a Retargeted Base Pose to use for animation retargeting purposes rather than using the characters normal Base Pose.
 
-We can select bones in our character and rotate them (in this case the left and right shoulders) so that our character is in an A-Pose, then we click **Modify Pose**: 
+We can select bones in our character and rotate them (in this case the left and right shoulders) so that our character is in an A-Pose, then we click **Modify Pose**:
 
 ![](RetargetedBasePose.png)
 
-In context menu, select **Use Current Pose**: 
+In context menu, select **Use Current Pose**:
 
 ![](UseCurrentPose.png)
 
-This will set the pose you define as the Retargeted Base Pose to use when performing any animation retargeting. 
+This will set the pose you define as the Retargeted Base Pose to use when performing any animation retargeting.
 
-Now when we go to retarget an animation, we will see the updated Retargeted Base Pose: 
+Now when we go to retarget an animation, we will see the updated Retargeted Base Pose:
 
 ![](NewBasePose.png)(w:840)
 
-When we retarget animations for Skeletons that have more similar Base Poses, we get better-looking results: 
+When we retarget animations for Skeletons that have more similar Base Poses, we get better-looking results:
 
 ![](HoldingShotGun.png)
 
 ### Importing Retarget Base Pose from Pose Asset
 
-From the **Modify** context menu, you can also choose to import a pose from an animation [Pose Asset](Engine/Animation/AnimPose) to use as the Retargeted Base Pose. 
+From the **Modify** context menu, you can also choose to import a pose from an animation [Pose Asset](Engine/Animation/AnimPose) to use as the Retargeted Base Pose.
 
 ![](ImportPose.png)
 
-In the image above, after selecting which Pose Asset to use (1) the available poses will be displayed in the selection drop-down menu (2). 
-After selecting a pose from the Pose Asset to use, clicking the **Import** button (3) will update the mesh in the viewport to use the pose selected as the Retargeted Base Pose. 
-Below (left) our default pose is shown and (right) our Retargeted Base Pose selected from a pose within our Pose Asset is shown. 
+In the image above, after selecting which Pose Asset to use (1) the available poses will be displayed in the selection drop-down menu (2). After selecting a pose from the Pose Asset to use, clicking the **Import** button (3) will update the mesh in the viewport to use the pose selected as the Retargeted Base Pose. Below (left) our default pose is shown and (right) our Retargeted Base Pose selected from a pose within our Pose Asset is shown.
 
 ![](DefaultPose.png)(w:540)![](NewRetargetPose.png)(w:540)
-
diff --git a/Engine/Documentation/Source/Engine/Animation/Persona/MeshDetails/MeshDetails.INT.udn b/Engine/Documentation/Source/Engine/Animation/Persona/MeshDetails/MeshDetails.INT.udn
--- a/Engine/Documentation/Source/Engine/Animation/Persona/MeshDetails/MeshDetails.INT.udn
+++ b/Engine/Documentation/Source/Engine/Animation/Persona/MeshDetails/MeshDetails.INT.udn
@@ -14,22 +14,16 @@ tags:Animation
 tags:Skeletal Mesh
 tags:FBX Importing
 
-[TOC (start:2 end:2)]
 
 ![](MeshDetails.png)(w:940)
 
-[EXCERPT:Overview]
-The **Asset Details** panel, which is located inside of the **Skeletal Mesh Editor**, shows specific properties pertaining to a Skeletal Mesh (such as the Materials being applied to its surface) as well as skinning and LOD options. 
-You can also assign and adjust APEX clothing, or define and associate a **Physics Asset** for the Skeletal Mesh (which is needed for creating ragdoll type motions). 
-[/EXCERPT:Overview]
+The **Asset Details** panel, which is located inside of the **Skeletal Mesh Editor**, shows specific properties pertaining to a Skeletal Mesh (such as the Materials being applied to its surface) as well as skinning and LOD options. You can also assign and adjust APEX clothing, or define and associate a **Physics Asset** for the Skeletal Mesh (which is needed for creating ragdoll type motions).
 
 This page highlights each section within the Asset Details panel with links to additional documentation.
 
-## Materials 
+## Materials
 
-All materials applied to meshes modeled in external applications will be exported with the mesh. 
-When importing a mesh into Unreal Engine 4 (UE4), you can automatically import its materials as part of the [FBX Import Process](Engine/Content/FBX/Materials). 
-Materials will appear inside of the Skeletal Mesh Editor under the **Materials** section and can be overridden by any other Materials that exist within the UE4 project. 
+All materials applied to meshes modeled in external applications will be exported with the mesh. When importing a mesh into Unreal Engine 4 (UE4), you can automatically import its materials as part of the [FBX Import Process](Engine/Content/FBX/Materials). Materials will appear inside of the Skeletal Mesh Editor under the **Materials** section and can be overridden by any other Materials that exist within the UE4 project.
 
 ![](MaterialsTable.png)
 
@@ -39,63 +33,69 @@ You can add **Material Slots**, which can be used to override an LOD's section (
 
 ## LODs
 
-When importing a Skeletal Mesh, you can also import the [Level of Details](Engine/Content/ImportingContent/ImportingSkeletalLODs/) (**LODs**) that were created and associated with the mesh. 
+When importing a Skeletal Mesh, you can also import the [Level of Details](Engine/Content/ImportingContent/ImportingSkeletalLODs) (**LODs**) that were created and associated with the mesh.
 
 ![](LODSection.png)
 
-You can **Highlight** portions of the mesh, or **Isolate** portions to only show certain sections of the mesh by clicking your desired option's checkbox. 
-You can also assign Material Elements from the Material section to **Material Slots**, or override sections per LOD.
+You can **Highlight** portions of the mesh, or **Isolate** portions to only show certain sections of the mesh by clicking your desired option's checkbox. You can also assign Material Elements from the Material section to **Material Slots**, or override sections per LOD.
 
 [OBJECT:ComparisonSlider]
-	[PARAM:before]
-	![Highlight On](HighlightOption.png)(w:1000)
-	[/PARAM]
-	[PARAM:after]
-	![Isolate On](IsolateOption.png)(w:1000)
-	[/PARAM]
+ [PARAM:before]
+ ![Highlight On](HighlightOption.png) 
+ [/PARAM]
+ [PARAM:after]
+ ![Isolate On](IsolateOption.png) 
+ [/PARAM]
 [/OBJECT]
 
-_The comparison slider toggles between highlighting the helmet with its character, versus isolating the helmet from its associated character._
-
-The LODs section also allows you to define when to use your LODs with the **Screen Size** option (how near or far the LOD should become visible), as well as defining the amount of **LODHysteresis** (flickering) to avoid when you're on an LOD boundary. 
-Inside of the **Viewport**, you can click the **Preview LODs** button to switch between the available LODs (it may also help to enable **Display Info** for statistical information). 
-
-![](previewLODs.png)
-
-Using the LOD Previewer, you can inspect the visual differences between LODs. Also, you can inspect statistical information about each LOD. For example, the amount of triangles in LOD 2 is almost half the amount of triangles in LOD 0. 
+You can turn on/off casting shadows on a per element, per LOD basis (some areas might not require shadow casting).
 
 [OBJECT:ComparisonSlider]
-	[PARAM:before]
-	![LOD 0](LOD0.png)
-	[/PARAM]
-	[PARAM:after]
-	![LOD 2](LOD2.png)
-	[/PARAM]
+    [PARAM:before]
+    ![Cast Shadows On](castshadows1.png) 
+    [/PARAM]
+    [PARAM:after]
+    ![Cast Shadows Off](castshadows2.png) 
+    [/PARAM]
 [/OBJECT]
 
-You can turn on/off casting shadows on a per element, per LOD basis (some areas might not require shadow casting).
+Inside of the **Viewport**, you can click the **Preview LODs** button to switch between the available LODs (it may also help to enable **Display Info**  for statistical information).
+
+![](PreviewLODs.png "PreviewLODs.png")
+
+Using the LOD Previewer, you can inspect the visual differences between LODs. Also, you can inspect statistical information about each LOD. For example, in the image below, the amount of triangles in LOD 2 is almost half the amount of triangles in LOD 0.
 
 [OBJECT:ComparisonSlider]
-	[PARAM:before]
-	![Cast Shadows On](castshadows1.png)(w:1000)
-	[/PARAM]
-	[PARAM:after]
-	![Cast Shadows Off](castshadows2.png)(w:1000)
-	[/PARAM]
+    [PARAM:before]
+    ![LOD 0](LOD0.png) 
+    [/PARAM]
+    [PARAM:after]
+    ![LOD 2](LOD2.png) 
+    [/PARAM]
 [/OBJECT]
 
-## Reduction Settings
+### LOD Info
 
-[REGION:warning]
-These options require the use of Simplygon. For more information on LODs and Reduction settings, please see the [](Engine/Content/Types/StaticMeshes/LOD) page for more information. 
-[/REGION]
+The LOD Info section enables you to define how to render the LOD along with other settings related to the current LOD: 
 
-In addition to the options outlined above, each generated LOD allows for defining the **Reduction Settings** for optimization. 
+![LOD Info Section-2.png](LOD20Info20Section-2.png "LOD Info Section-2.png")
 
-![](ReductionSettings.png)
+| Property | Description |
+| --- | --- |
+| **Screen Size** | Screen Size to display the LOD. The Screen Size is based around the projected diameter of the bounding sphere of the model. For example, 0.5 means half the screen's maximum dimension. |
+| **LOD Hysteresis** | Used to avoid "flickering" when on LOD boundry. Only take into account when moving from complex to simple. |
+| **Source Import Filename** | The filename of the file that was used to import the LOD if it was not auto-generated. |
+| **Allow CPUAccess** | Keeps the LODs data on the CPU so it can be used for things such as sampling in FX. |
+| **Support Uniformly Distributed Sampling** | Mesh supports uniformly distributed sampling in constant time. Memory cost is 8 bytes per triangle. Example usage is uniform spawning of particles. |
+| **Bake Pose** | Enables you to bake a single frame into in LOD mesh (this is useful for when you are removing bones and want to retain a pose, see the example video below this table). |
+| **Bones to Remove** | Bones which should be removed from the Skeleton for the LOD level. |
+| **Reduction Settings** | [Reduction settings](https://docs.unrealengine.com/en-us/Engine/Animation/Persona/MeshDetails#reductionsettings) to apply when building render data. |
+| **Regenerate LOD** | Enables you to regenerate the current LOD. |
+| **Remove this LOD** | Enables you to remove the current LOD. |
+| **Reimport** | Enables you to reimport the LOD from your computer. |
+| **Reimport (New File)** | Enables you to reimport a new LOD to overwrite the existing LOD. |
 
-As of Unreal Engine 4.16, you can also bake a pose into an LOD using the **Bake Pose** reduction setting. This can be set to a single frame animation sequence which will be applied to the resulting LOD mesh. 
-This is useful when removing bones and still wanting to retain a pose, similar to the example video shown below. 
+The following is an example of using Bake Pose to retain the hand pose of our character while moving between LODs. 
 
 [OBJECT:EmbeddedVideo]
 	[PARAMLITERAL:videoid]
@@ -115,15 +115,71 @@ This is useful when removing bones and still wanting to retain a pose, similar t
 	[/PARAMLITERAL]
 [/OBJECT]
 
+## Reduction Settings
+
+[REGION:warning]
+These options require the use of Simplygon. For more information on LODs and Reduction settings, please see the [Simplygon LODs](Engine/Content/Types/StaticMeshes/LOD) page.
+[/REGION]
+
+In addition to the options outlined above, each generated LOD allows for defining the **Reduction Settings** for optimization.
+
+![](ReductionSettingsOnly.png "ReductionSettingsOnly.png")
+
+The following settings can be applied when building render data: 
+
+| Property | Description |
+| --- | --- |
+| **Reduction Method** | The method to use when optimizing the Skeletal Mesh LOD. You can set this to Triangles, Accuracy or Any (when either requirement is met). |
+| **Triangle Percentage** | The simplification uses this percentage of the source meshe's triangle count as a target. |
+| **Accuracy Percentage** | The simplification uses this as how much to deviate from the source mesh. Works better with hard surfaces. Setting this to 100% is the same as the Base LOD. |
+| **Remap Morph Targets** | Remap the morph targets from the base LOD onto the reduced LOD. |
+| **Silhouette** | How important the shape of the geometry is. |
+| **Texture** | How important texture density is. |
+| **Shading** | How important shading quality is. |
+| **Skinning** | How important the skinning quality is. |
+| **Welding Threshold** | The welding threshold distance. Vertices under this distance will be welded. |
+| **Recompute Normal** | Whether Normal smoothing groups should be preserved. If true, then Hard Edge Angle (Normals Threshold) is used. |
+| **Hard Edge Angle** | If the angle between two triangles are above this value, the normals will not be smooth over the edge between those two triangles. This is only used when Recalculate Normals is set to True. |
+| **Max Bones Influence** | Maximum number of bones that can be assigned to each vertex. |
+| **Base LOD** | Base LOD index to generate this LOD. By default, we generate from LOD 0. |
+
+Reduction Settings that use a Bone List also allows whitelist joints (as well as blacklist), so that you can just keep only certain joints. 
+
+![](BoneFilterAction_01.png "BoneFilterAction_01.png")
+
+Use the **Bone Filter Action ****Option **to define your desired setting after adding Bones to the Bone List.
+
+## Sharing LOD Settings
+
+Skeletal Meshes can use a shared Data Asset of the **LOD Settings **Asset type instead of setting LOD information directly (this includes the ability to share Reduction Settings). Inside the **Asset Details **panel for a Skeletal Mesh, under **LOD Settings**, you can select which LOD Settings asset to use or click the **Generate Asset** button to generate a new LOD Settings asset based on the current settings. 
+
+![](LODSettings.png "LODSettings.png")
+
+After selecting an LOD Settings Asset, options will be grayed out and use the information contained within the asset. 
+
+![](AssignedLODSettings.png "AssignedLODSettings.png")  
+
+**To create a new LOD Settings Asset: **  
+
+1\. In the **Content Browser **from the **Add New **button under **Miscellaneous**, select **Data Asset.**
+
+![](DataAsset.png "DataAsset.png")
+
+2\. After creating the asset, assign the **SkeletalMeshLODSettings **as the class. 
+
+![](DataAssetClass.png "DataAssetClass.png")
+
+3\. You can then set up your LOD Groups and settings to store with the asset that can be used with other Skeletal Meshes. 
+
+![](LOD_DataAsset.png "LOD_DataAsset.png")
 
 ## Clothing
 
-The **Clothing** section associates an [](Engine/Physics/Apex) clothing file with the mesh. Click the **Add APEX clothing file...** button to begin importing an .apb or .apx file, 
-which uses the same pipeline as importing an FBX file. 
+The **Clothing** section associates an [APEX](Engine/Physics/Apex) clothing file with the mesh. Click the **Add APEX clothing file...** button to begin importing an .apb or .apx file, which uses the same pipeline as importing an FBX file.
 
 ![](AddApex.png)
 
-Once you've imported your Apex file, you can adjust its [Physics Properties](Engine/Physics/PhAT/Reference) to achieve your desired results. 
+Once you've imported your Apex file, you can adjust its [Physics Properties](Engine/Physics/PhAT/Reference) to achieve your desired results.
 
 ![](OwenApexCloth.png)
 
@@ -131,54 +187,47 @@ Once you've imported your Apex file, you can adjust its [Physics Properties](Eng
 For example, this character has two Apex files associated with it; one for the character's tie, and the other for the character's coat.
 [/REGION]
 
-
 ## Mesh
 
-In the **Mesh** section, you can see what [Skeleton](Engine/Animation/Skeleton) asset the mesh is using and navigate to it with the **Content Browser**. 
-There are also options for setting **Positive Bounds Extension** and **Negative Bounds Extension** values, which can be used to offset the bounds values of the mesh when they're being imported. 
-The other option, **Use Full Precision UVs** can be turned on to use 32 bit UVs (or left off to save memory, thereby using 16 bit UVs). 
+In the **Mesh** section, you can see what [Skeleton](Engine/Animation/Skeleton) asset the mesh is using and navigate to it with the **Content Browser**. There are also options for setting **Positive Bounds Extension** and **Negative Bounds Extension** values, which can be used to offset the bounds values of the mesh when they're being imported. The other option, **Use Full Precision UVs** can be turned on to use 32 bit UVs (or left off to save memory, thereby using 16 bit UVs).
 
 ![](MeshSection.png)
 
 ## Physics
 
-The **Physics** section allows you to **Enable Per Poly Collision** for the mesh, which uses skinned data for collision information. Per Poly Collision cannot be used for simulation. In most cases, you're better off using a **Physics Asset**. 
-There is also the ability to assign a [Physics Asset](Engine/Physics/PhAT) to the Skeletal Mesh, which is used for processing physics and collision information. This is used for per-bone hit detection, accurate bounding box calculation, and ragdoll physics. 
+The **Physics** section allows you to **Enable Per Poly Collision** for the mesh, which uses skinned data for collision information. Per Poly Collision cannot be used for simulation. In most cases, you're better off using a **Physics Asset**. There is also the ability to assign a [Physics Asset](Engine/Physics/PhAT) to the Skeletal Mesh, which is used for processing physics and collision information. This is used for per-bone hit detection, accurate bounding box calculation, and ragdoll physics.
 
 ![](PhysicsSection.png)
 
 ## Lighting
 
-The **Shadow Physics Asset** option in the Lighting section allows you to specify a Physics Asset, with shapes that will be used for shadowing when components have their **Cast Character Capsule Direct Shadow** or **Cast Character Capsule Indirect Shadow** option(s) enabled. 
-Only spheres and sphyl shapes in the Physics Asset are supported. Additionally, it's imporant to remember that capsule shadows can result in higher rendering costs when more shapes are used. 
+The **Shadow Physics Asset** option in the Lighting section allows you to specify a Physics Asset, with shapes that will be used for shadowing when components have their **Cast Character Capsule Direct Shadow** or **Cast Character Capsule Indirect Shadow** option(s) enabled. Only spheres and sphyl shapes in the Physics Asset are supported. Additionally, it's imporant to remember that capsule shadows can result in higher rendering costs when more shapes are used.
 
 ![](Lighting.png)
 
-In-game, indirect capsule shadows are a good way to ground characters to the environment: 
+In-game, indirect capsule shadows are a good way to ground characters to the environment:
 
 ![](IndirectShadows.png)(w:840)
 
-Direct capsule shadows can be used to achieve soft character shadows in an efficient baked lighting environment. 
+Direct capsule shadows can be used to achieve soft character shadows in an efficient baked lighting environment.
 
 ![](DirectShadows.png)(w:840)
 
-The light's **Source Radius** or **Source Angle** determines how soft the shadows will be. Capsule shadowing is very efficient because it computes shadowing at half-resolution with a depth-aware up-sample, using screen tile culling to limit the shadowing work (where it's needed). 
-The GPU cost is proportional to the number of capsules and the number of pixels affected by the casted shadow. 
+The light's **Source Radius** or **Source Angle** determines how soft the shadows will be. Capsule shadowing is very efficient because it computes shadowing at half-resolution with a depth-aware up-sample, using screen tile culling to limit the shadowing work (where it's needed). The GPU cost is proportional to the number of capsules and the number of pixels affected by the casted shadow.
 
 ## Skeletal Mesh
 
-In this section, you can assign a **Post Process Anim Blueprint**, which is an [Animation Blueprint](Engine/Animation/AnimBlueprints) class, to run as a Post Process for the mesh. This Blueprint will be run before Physics, but after the Main Anim Instance for any [Skeletal Mesh Component](Engine/Components/SkeletalMesh/) using the mesh. 
+In this section, you can assign a **Post Process Anim Blueprint**, which is an [Animation Blueprint](Engine/Animation/AnimBlueprints) class, to run as a Post Process for the mesh. This Blueprint will be run before Physics, but after the Main Anim Instance for any [Skeletal Mesh Component](Engine/Components/SkeletalMesh) using the mesh.
 
 ![](SkeletalMesh.png)
 
-This allows you to set up [](Engine/Animation/NodeReference/SkeletalControls/AnimDynamics/) or other controllers that will always be applied to the associated mesh, whether that mesh is viewed with the Animation Tools, as part of a Sequencer cinematic, or through a Skeletal Mesh Actor placed in a level. 
-This is especially useful for complex characters that run multiple Animation Blueprints; because any Anim Dynamics, controllers, and IK (among other Anim Blueprint features) will be related to the mesh and not have to be duplicated in every Animation Blueprint intended to be used with the mesh. 
+This allows you to set up [AnimDynamics](Engine/Animation/NodeReference/SkeletalControls/AnimDynamics) or other controllers that will always be applied to the associated mesh, whether that mesh is viewed with the Animation Tools, as part of a Sequencer cinematic, or through a Skeletal Mesh Actor placed in a level. This is especially useful for complex characters that run multiple Animation Blueprints; because any Anim Dynamics, controllers, and IK (among other Anim Blueprint features) will be related to the mesh and not have to be duplicated in every Animation Blueprint intended to be used with the mesh.
 
-To assign a Post Process Anim Blueprint, select the drop-down menu and specify your default Animation Blueprint. 
+To assign a Post Process Anim Blueprint, select the drop-down menu and specify your default Animation Blueprint.
 
 ![](SelectPostProcessAnimBP.png)
 
-The example video illustrates how the Post Process Anim Blueprint is run during gameplay on a Skeletal Mesh. 
+The example video illustrates how the Post Process Anim Blueprint is run during gameplay on a Skeletal Mesh.
 
 [OBJECT:EmbeddedVideo]
 	[PARAMLITERAL:videoid]
@@ -198,8 +247,7 @@ The example video illustrates how the Post Process Anim Blueprint is run during
 	[/PARAMLITERAL]
 [/OBJECT]
 
-
-If you have a more complex animation setup that uses multiple Animation Blueprints, you may need to include a **Sub-Graph Input** node to the Post Process Anim Blueprint. 
+If you have a more complex animation setup that uses multiple Animation Blueprints, you may need to include a **Sub-Graph Input** node to the Post Process Anim Blueprint.
 
 ![](SubGraphInput.png)
 
@@ -207,22 +255,31 @@ If you have a more complex animation setup that uses multiple Animation Blueprin
 This will allow you to get the pose information from a parent Animation Blueprint (if applicable) rather than using the Skeletal Mesh's base pose.
 [/REGION]
 
+You can also control whether or not the Post Process instance runs on a Skeletal Mesh Component through Blueprints or get the Post Process Instance assigned. 
+
+![](TogglePostProcess.png)
+
+Above, we have a Skeletal Mesh Component called **Mesh **in which we can call **Toggle Disable Post Process Blueprint **to toggle the state of the assigned Post Process Blueprint. We also use the **Get Post Process Instance **node to return the active Post Process Instance if one is assigned. This is set on the mesh that the component is using, and is evaluated immediately after the main instance.
+
+  
+You can also set the value directly inside the **Details **panel when the Skeletal Mesh Component is selected. 
+
+![](DetailsPanel.png "DetailsPanel.png")
+
 ## Transform
 
-The **Transform** section inside of the **Skeletal Mesh Editor** allows you to offset the mesh from the world origin (when imported) to match its position in the source art program through the **Import Translation** option. 
-**Import Rotation** will rotate the mesh (when imported) to match its position in the source art program, while the **Import Uniform Scale** option will import the uniform scaling of the mesh from the source art program. 
+The **Transform** section inside of the **Skeletal Mesh Editor** allows you to offset the mesh from the world origin (when imported) to match its position in the source art program through the **Import Translation** option. **Import Rotation** will rotate the mesh (when imported) to match its position in the source art program, while the **Import Uniform Scale** option will import the uniform scaling of the mesh from the source art program.
 
 ![](TransformSection.png)
 
 ## Thumbnail
 
-**Thumbnail** options allow you to modify how the mesh is framed in the thumbnail that appears in the **Content Browser** or **Editor Toolbar**. 
+**Thumbnail** options allow you to modify how the mesh is framed in the thumbnail that appears in the **Content Browser** or **Editor Toolbar**.
 
 ![](ThumbnailSection.png)
 
 ## Import Settings
 
-The **Import Settings** section allows you to modify any of the settings used during the initial [FBX Skeletal Mesh Import Process](https://docs.unrealengine.com/latest/INT/Engine/Content/FBX/ImportOptions/#skeletalmeshoptions) for the mesh. 
-
-![](ImportSettingsSection.png)
+The **Import Settings** section allows you to modify any of the settings used during the initial [FBX Skeletal Mesh Import Process](Engine/Content/FBX/ImportOptions/#skeletalmeshoptions) for the mesh.
 
+![](ImportSettingsSection.png)
\ No newline at end of file
diff --git a/Engine/Documentation/Source/Engine/Blueprints/BP_HowTo/BasicUsage/BasicUsage.INT.udn b/Engine/Documentation/Source/Engine/Blueprints/BP_HowTo/BasicUsage/BasicUsage.INT.udn
--- a/Engine/Documentation/Source/Engine/Blueprints/BP_HowTo/BasicUsage/BasicUsage.INT.udn
+++ b/Engine/Documentation/Source/Engine/Blueprints/BP_HowTo/BasicUsage/BasicUsage.INT.udn
@@ -115,15 +115,7 @@ In the **Components Mode** you define what bits and pieces are used to construct
 
 * See [](Engine/Blueprints/Editor/Modes/Components) for more on the windows and options in this mode. 
 * See also [](Engine/Blueprints/UserGuide/Components) for working with **Components** in this mode.  
-
-### Graph Mode
-
-![](GraphTab.png)(w:740)
-
-The **Graph Mode** is considered to be the heart of a Blueprint. This is where you bring your Blueprints to life by adding and connecting Blueprint nodes together in a visual network that defines your Blueprint's logic. In the example image above, we have a **Character Blueprint** and in the graph, the logic that governs the character's movement is laid out. 
-
-* See [](Engine/Blueprints/Editor/Modes/GraphPanel) for more on the windows and options in this mode. 
-* See also [](Engine\Blueprints\BP_HowTo\NavigatingGraphs) for working within the **Graph Window** in this mode.  
+ 
 
 ### Defaults Mode
 
diff --git a/Engine/Documentation/Source/Engine/Blueprints/Editor/UIComponents/Menu/BlueprintEditorMenu.INT.udn b/Engine/Documentation/Source/Engine/Blueprints/Editor/UIComponents/Menu/BlueprintEditorMenu.INT.udn
--- a/Engine/Documentation/Source/Engine/Blueprints/Editor/UIComponents/Menu/BlueprintEditorMenu.INT.udn
+++ b/Engine/Documentation/Source/Engine/Blueprints/Editor/UIComponents/Menu/BlueprintEditorMenu.INT.udn
@@ -34,7 +34,7 @@ The **Menu** is displayed at the top of the Blueprint Editor by default.
 
 ### View
 
-[INCLUDE:ViewMenu]
+[INCLUDE:#ViewMenu]
 
 ### Debug
 
diff --git a/Engine/Documentation/Source/Engine/Blueprints/Overview/BlueprintOverview.INT.udn b/Engine/Documentation/Source/Engine/Blueprints/Overview/BlueprintOverview.INT.udn
--- a/Engine/Documentation/Source/Engine/Blueprints/Overview/BlueprintOverview.INT.udn
+++ b/Engine/Documentation/Source/Engine/Blueprints/Overview/BlueprintOverview.INT.udn
@@ -112,7 +112,6 @@ and more.
 
 A Blueprint's _mode_ determines what is visible in its window. Although Level Blueprints contain only one mode, the Graph mode, Blueprint Classes contain three different modes:
 
-* **[Graph Mode](Engine/Blueprints\Editor\Modes\GraphPanel)** - Where you add and connect scripting nodes to create gameplay functionality for your level.
 * **[Default Mode](Engine/Blueprints\Editor\Modes\Defaults)** - Where you set the defaults for the Blueprint.
 * **[Components Mode](Engine/Blueprints\Editor\Modes\Components)** - Where you add, remove, and modify the components that make up your Blueprint.
 
diff --git a/Engine/Documentation/Source/Engine/Blueprints/UserGuide/EventGraph/BlueprintEventGraph.INT.udn b/Engine/Documentation/Source/Engine/Blueprints/UserGuide/EventGraph/BlueprintEventGraph.INT.udn
--- a/Engine/Documentation/Source/Engine/Blueprints/UserGuide/EventGraph/BlueprintEventGraph.INT.udn
+++ b/Engine/Documentation/Source/Engine/Blueprints/UserGuide/EventGraph/BlueprintEventGraph.INT.udn
@@ -9,7 +9,7 @@ tags:Blueprints
 topic-image:Engine/Blueprints/UserGuide/EventGraph/eventgraph_topic.png
 
 
-![Graph Panel](Engine/Blueprints/Editor/Modes/GraphPanel/k2_graphview.png)
+![Graph Panel](k2_graphview.png)
 
 [EXCERPT:Intro]
 
diff --git a/Engine/Documentation/Source/Engine/Blueprints/UserGuide/Graphs/Graphs.INT.udn b/Engine/Documentation/Source/Engine/Blueprints/UserGuide/Graphs/Graphs.INT.udn
--- a/Engine/Documentation/Source/Engine/Blueprints/UserGuide/Graphs/Graphs.INT.udn
+++ b/Engine/Documentation/Source/Engine/Blueprints/UserGuide/Graphs/Graphs.INT.udn
@@ -55,7 +55,7 @@ For more information about this section, see [%Engine/Blueprints/UserGuide/Macro
 
 ## Working with Graphs
 
-Whether your Graph is a Construction Script, EventGraph, Function, or Macro, you will edit it in the [Graph](Engine/Blueprints/Editor/Modes/GraphPanel) mode of the [](Engine/Blueprints/Editor). Fundamentally,
+Whether your Graph is a Construction Script, EventGraph, Function, or Macro, you will edit it in the [Graph](Engine/Blueprints/Editor/UIBreakdowns/ClassBPUI) mode of the [](Engine/Blueprints/Editor). Fundamentally,
 all Graphs contain networks of Nodes connected by Wires.
 
 
diff --git a/Engine/Documentation/Source/Engine/Editor/ScriptingAndAutomation/Blueprints/Blueprints.INT.udn b/Engine/Documentation/Source/Engine/Editor/ScriptingAndAutomation/Blueprints/Blueprints.INT.udn
new file mode 100644
--- /dev/null
+++ b/Engine/Documentation/Source/Engine/Editor/ScriptingAndAutomation/Blueprints/Blueprints.INT.udn
@@ -0,0 +1,72 @@
+Availability: Public
+Crumbs: 
+Title: Scripting the Editor using Blueprints
+Description: Describes how to use Blueprints in the Editor to automate content production tasks.
+Type: 
+Version: 4.20
+Parent: Engine/Editor/ScriptingAndAutomation
+Order: 
+Tags: Blueprints, scripting, editor
+
+
+Blueprints are not only useful for creating runtime gameplay in your Project, they can also be a very effective tool for interacting with your Project content within the Unreal Editor. You can run your Blueprint graphs on demand in the Editor — these are sometimes referred to as Blutilities, short for Blueprint Utilities — to work with Assets, lay out content in your Levels, and even control the Editor UI itself.
+
+This page describes some of the basics involved in using Blueprints to script and automate the Unreal Editor.
+
+## Running Blueprints in the Editor
+
+Any time you use a Custom Event node in the Event Graph of a Blueprint class, you can set the **Graph > Call in Editor** option in the **Details** panel:
+
+![Call in Editor](bp-call-in-editor.png "Call in Editor")
+
+If you place an Actor with a Custom Event like this into a Level, and you select it in the Viewport or in the **World Outliner**, the **Details** panel will show a button for each Call in Editor event you set up. You'll find them in the **Default** section, where Blueprint classes also expose variables. Click this button to trigger execution of the Event Graph from your Custom Event node.
+
+![Custom event in the Details panel](bp-call-from-level.png "Custom event in the Details panel")
+
+This works for any Blueprint class that you can place into a Level. You can use it to test your runtime gameplay, as well as to build Editor-only content management tools.
+
+[REGION:note]
+If your Custom Event has any inputs, it will not be shown in the **Details** panel.
+[/REGION]
+
+## Using Editor-Only Features
+
+Most operations that involve modifying Asset files or working with the Unreal Editor UI can't work in your game at runtime. Therefore, you can only access Editor-only features — that is, any functions that are defined in a module whose type is set to `Editor` — from Blueprint classes that are also Editor-only.
+
+For example, if you use a Blueprint class that derives from a parent class that can be used at runtime — such as the base **Actor** class — you'll see a limited set of functions listed under the **Editor Scripting** category in your Blueprint Editor. However, if you use a class that derives from an Editor-only parent class — such as **PlacedEditorUtilityBase** — you'll see a much larger and more comprehensive set of functions that you can use:
+
+![Editor Scripting Blueprint nodes](bp-editor-scripting-nodes.png "Editor Scripting Blueprint nodes")
+
+[REGION:tip]
+If you've installed the Editor Scripting Utilities plugin, this is where you'll find the functions it exposes for working with Static Meshes and other kinds of Assets. See [Scripting and Automating the Editor](https://docs.unrealengine.com/en-us/Editor/Scripting-and-Automating-the-Editor).
+[/REGION]
+
+## Blutility Classes and the Blutility Shelf
+
+[REGION:note]
+Blutility classes and the Blutility shelf are currently experimental. As with all experimental features, we recommend using caution before relying on it in production. It may change in future versions of the Unreal Engine.
+[/REGION]
+
+When you're writing Blueprints to automate the Unreal Editor and work with Assets, the graphs you create may not have any dependency on being in a particular Level. When this is the case, you may not want to always have to drop an instance of your Blueprint class into a Level simply in order to invoke its Call in Editor events.
+
+For these situations, the Unreal Editor offers an experimental system for creating and managing Editor-only Blueprint utility classes that you can invoke directly from the Asset Browser or from a custom panel.
+
+To get started:
+
+1.  Open the Editor Preferences window, and find the **General > Experimental** category. Under **Tools**, check the option for **Editor Utility Blueprints (Blutility)**.  
+    ![Enable Editor Utility Blueprints](bp-experimental-blutility.png "Enable Editor Utility Blueprints")
+2.  Create a new Blutility class. In the Content Browser, right-click the folder where you want to create your new utility class, and choose **Blueprints > Blutility** from the contextual menu.  
+    ![Create new Blutility](bp-blutility-menu.png "Create new Blutility")
+3.  Choose the `GlobalEditorUtilityBase` as your parent class, and click **Select**.  
+    ![GlobalEditorUtilityBase parent class](bp-global-editor-utility-base.png "GlobalEditorUtilityBase parent class")  
+    
+4.  To edit the Blueprint graph for your new class, right-click it and choose **Edit Blueprint** from the contextual menu.
+5.  In the Event Graph, create a new Custom Event and enable its **Call in Editor** setting, as shown in the section above. Save and compile your Blueprint.
+6.  When you double-click your Blueprint class in the Content Browser, you'll open a new window that lists all of its Call in Editor events.  
+    ![Blutility events](bp-open-blutility.png "Blutility events")
+
+[REGION:tip]
+You can also access all of the Blutility classes you've created from the **Blutility Shelf** panel, which you can open by choosing **Window > Blutility Shelf** from the main menu.
+
+![Blutility Shelf](bp-blutility-shelf.png "Blutility Shelf")
+[/REGION]
\ No newline at end of file
diff --git a/Engine/Documentation/Source/Engine/Editor/ScriptingAndAutomation/HowTos/CreatingLODs/CreatingLODs.INT.udn b/Engine/Documentation/Source/Engine/Editor/ScriptingAndAutomation/HowTos/CreatingLODs/CreatingLODs.INT.udn
new file mode 100644
--- /dev/null
+++ b/Engine/Documentation/Source/Engine/Editor/ScriptingAndAutomation/HowTos/CreatingLODs/CreatingLODs.INT.udn
@@ -0,0 +1,153 @@
+Availability: Public
+Crumbs: 
+Title: Creating Levels of Detail in Blueprints and Python
+Description: Shows how to set up automatic Levels of Detail (LODs) in the Unreal Editor using Blueprints and Python.
+Type:how-to
+Version: 4.20
+Parent: Engine/Editor/ScriptingAndAutomation/HowTos
+SkillLevel:Intermediate
+Order: 1
+Tags: Blueprints, how-to, Python, scripting, editor, lods, levels of detail
+
+
+Creating Levels of Detail (LODs) for your meshes can be a good way to increase the performance and frame rate of your game without sacrificing visual quality.
+
+In general, the more triangles your meshes have, and the smaller those triangles are on screen, the harder it is for your GPU to render them. Trying to render too many detailed meshes all at the same time can slow down your frame rate. However, you typically don't need every mesh in the scene to be rendered with the same fidelity: when a mesh is far away, you are often able to swap it out for a less-detailed version of the mesh that contains fewer triangles without a noticeable difference in visual quality.
+
+Unreal Engine 4 offers a built-in LOD management system that automatically chooses the most appropriate version of a mesh to show at runtime, based on the amount of screen space the mesh is currently occupying in each frame. This system relies on you setting up these different alternative versions of your meshes in the Editor ahead of time.  
+
+The Editor is capable of generating these Levels of Detail for your Static Mesh assets automatically, by progressively simplifying the mesh to thresholds that you specify. For more information about how this system works and how to set it up in the Static Mesh Editor, see [Automatic LOD Generation](https://docs.unrealengine.com/en-us/Engine/Content/Types/StaticMeshes/HowTo/AutomaticLODGeneration).
+
+[REGION:lightbox]
+[![Level of Detail settings in the Static Mesh Editor](ALC_PM_00.png "Level of detail settings in the static mesh editor data-verified=")](ALC_PM_00.png)
+[/REGION]
+[REGION:caption]
+Click image for full size.
+[/REGION]
+
+However, even though the Editor can generate these LODs automatically, it may not be feasible for you to open every Static Mesh Asset in your project one-by-one in order to set them all up. You could bulk edit the Assets, but this doesn't help if you want to apply different settings to different Static Mesh Assets based on their characteristics — for example, to apply different reduction settings based on the number of existing triangles in the mesh, or based on an Asset naming convention. It also doesn't help if you want to create your LODs as a sub-step within a larger custom Asset pipeline that you've scripted. For these kinds of cases, you can use Blueprints or Python to script the automatic LOD creation system.
+
+[REGION:note]
+You'll need to install the **Editor Scripting Utilities** plugin, if you haven't already. For details, see [Scripting and Automating the Editor](Engine/Editor/ScriptingAndAutomation).
+[/REGION]
+
+%Globals:BlueprintOrPythonSelectVar%
+
+[OBJECT:ToggleButtonContent]
+	[PARAMLITERAL:category]
+	BlueprintOrCPPSelect
+	[/PARAMLITERAL]
+	[PARAMLITERAL:id]
+	blueprintimplementation
+	[/PARAMLITERAL]
+	[PARAMLITERAL:active]
+	[/PARAMLITERAL]
+	[PARAM:content]
+	You'll find the nodes you'll need to manage LODs under the **Editor Scripting > Static Mesh** category.
+
+	[REGION:note]
+	To use these nodes, your Blueprint class must be derived from an Editor-only class, such as the **PlacedEditorUtilityBase** class. For details, see [Scripting the Editor using Blueprints](Engine/Editor/ScriptingAndAutomation/Blueprints).
+	[/REGION]
+
+	*   The key node is **Set Lods**, which automatically creates levels of detail for the Static Mesh Asset you pass in. To use this node, you'll need to provide it with a set of reduction settings that define the screen size threshold and relative triangle percentage for each level of detail you want to create. See the example below.
+		[REGION:note]
+		The first **EditorScriptingMeshReductionSettings** item that you pass to the **Reduction Settings** input of the **EditorScriptingMeshReductionOptions** node has no effect, as LOD 0 is always considered to show the mesh at its most detailed level, with all of its triangles in place.
+		[/REGION]
+	*   You can use **Get Lod Count** and **Get Lod Screen Sizes** to get information about the Levels of Detail that are currently set up for a Static Mesh.
+	*   You can also use **Remove Lods** to remove all existing LODs (except LOD 0, which is always the fully detailed mesh).
+	*   Setting up LODs modifies the Static Mesh Asset. Assuming you want to keep the changes you make, you'll also need to use a node like **Save Asset** or **Save Loaded Asset** afterward.
+
+	The following example loads each Static Mesh Asset in an input path in turn. Whenever it finds one that has more vertices than a minimum threshold, it sets up three additional LODs for the Static Mesh, and saves it afterward.
+
+	[REGION:lightbox]
+	[![Set LODs in Blueprints Example](set_lods_example.png "Set LODs in Blueprints Example")](set_lods_example.png)
+	[/REGION]
+	[REGION:caption]
+	Click image for full size.
+	[/REGION]
+	[/PARAM]
+[/OBJECT]
+
+[OBJECT:ToggleButtonContent]
+	[PARAMLITERAL:category]
+	BlueprintOrCPPSelect
+	[/PARAMLITERAL]
+	[PARAMLITERAL:id]
+	pythonimplementation
+	[/PARAMLITERAL]
+	[PARAMLITERAL:active]
+	[/PARAMLITERAL]
+	[PARAM:content]
+	You'll find the LOD management functions in the `unreal.EditorStaticMeshLibrary` class.
+
+	*   The key function is `unreal.EditorStaticMeshLibrary.set_lods()`, which automatically creates levels of detail for the Static Mesh Asset you pass in. To use this function, you'll need to pass in an `EditorScriptingMeshReductionOptions` object, which contains a set of `EditorScriptingMeshReductionSettings` that define the screen size threshold and relative triangle percentage for each level of detail you want to create. See the example below.
+		[REGION:note]
+		The first `EditorScriptingMeshReductionSettings` item that you set in the `EditorScriptingMeshReductionOptions.reduction_settings` array has no effect, as LOD 0 is always considered to show the mesh at its most detailed level, with all of its triangles in place.
+		[/REGION]
+	*   You can use `unreal.EditorStaticMeshLibrary.get_lod_count()` and `unreal.EditorStaticMeshLibrary.get_lod_screen_sizes` to get information about the Levels of Detail that are currently set up for a Static Mesh.
+	*   You can also use `unreal.EditorStaticMeshLibrary.remove_lods()` to remove all existing LODs (except LOD 0, which is always the fully detailed mesh).
+	*   Setting up LODs modifies the Static Mesh Asset. Assuming you want to keep the changes you make, you'll also need to use a function like `unreal.EditorAssetLibrary.save_asset()` or  `unreal.EditorAssetLibrary.save_loaded_asset()` afterward.
+
+	The following example loads each Static Mesh Asset in an input path in turn. Whenever it finds one that has more vertices than a minimum threshold, it sets up three additional LODs for the Static Mesh, and saves it afterward.
+
+		import unrealasset_path = "/Game/studio"
+		# we define a function that generates new LODs for a specified Static Mesh Asset.
+		def apply_lods(static_mesh):
+		# check if the mesh is complex enough.
+		number_of_vertices = unreal.EditorStaticMeshLibrary.get_number_verts(static_mesh, 0)
+		if number_of_vertices < 10:
+			return
+		print("treating asset: " + static_mesh.get_name())
+		print("existing LOD count: " + str(unreal.EditorStaticMeshLibrary.get_lod_count(static_mesh)))
+		# set up options for auto-generating the levels of detail.
+		options = unreal.EditorScriptingMeshReductionOptions()
+		# we request three new levels of detail. Each has:
+		# - a screen space threshold at which this level of detail appears.
+		# - the percentage of the triangles in LOD 0 that should be retained at this level/
+		options.reduction_settings = [ unreal.EditorScriptingMeshReductionSettings(1.0, 1.0),
+			unreal.EditorScriptingMeshReductionSettings(0.8, 0.75),
+			unreal.EditorScriptingMeshReductionSettings(0.6, 0.5),
+			unreal.EditorScriptingMeshReductionSettings(0.4, 0.25)
+		]
+		# use the screen space thresholds set above, rather than auto-computing them.
+		options.auto_compute_lod_screen_size = False
+		# set the options on the Static Mesh Asset.
+		unreal.EditorStaticMeshLibrary.set_lods(static_mesh, options)
+		# save the changes.
+		unreal.EditorAssetLibrary.save_loaded_asset(static_mesh)
+		print("new LOD count: " + str(unreal.EditorStaticMeshLibrary.get_lod_count(static_mesh)))
+		# get a list of all Assets in the path.
+		all_assets = unreal.EditorAssetLibrary.list_assets(asset_path)
+		# load them all into memory.
+		all_assets_loaded = [unreal.EditorAssetLibrary.load_asset(a) for a in all_assets]
+		# filter the list to include only Static Meshes.
+		static_mesh_assets = unreal.EditorFilterLibrary.by_class(all_assets_loaded, unreal.StaticMesh)
+		# run the function above on each Static Mesh in the list.
+		map(apply_lods, static_mesh_assets)
+
+	An alternative approach is to set the **LOD Group** option for each Static Mesh Asset. This option makes the mesh use one of the pre-set LOD reduction settings that are defined in the `[StaticMeshLODSettings]` section of your project's **BaseEngine.ini** file, such as `LevelArchitecture`, `SmallProp`, `LargeProp`, or `HighDetail`.
+
+	For example:
+
+		import unreal
+		asset_path = "/Game/studio/"
+		def set_high_detail(static_mesh):
+		# Set the LOD group.
+		static_mesh.set_editor_property("lod_group", "HighDetail")
+		# Save the Asset.
+		unreal.EditorAssetLibrary.save_loaded_asset(static_mesh)
+		# get a list of all Assets in the path.
+		all_assets = unreal.EditorAssetLibrary.list_assets(asset_path)
+		# load them all into memory.
+		all_assets_loaded = [unreal.EditorAssetLibrary.load_asset(a) for a in all_assets]
+		# filter the list to include only Static Meshes.
+		static_mesh_assets = unreal.EditorFilterLibrary.by_class(all_assets_loaded, unreal.StaticMesh)
+		# run the function above on each Static Mesh in the list.
+		map(set_high_detail, static_mesh_assets)
+
+	[REGION:tip]
+	For more information about how this system works and how to use it in the Editor, see [Automatic LOD Generation](Engine/Content/Types/StaticMeshes/HowTo/AutomaticLODGeneration#usinglodgroups).
+	[/REGION]
+	[/PARAM]
+[/OBJECT]
+
diff --git a/Engine/Documentation/Source/Engine/Editor/ScriptingAndAutomation/HowTos/HowTos.INT.udn b/Engine/Documentation/Source/Engine/Editor/ScriptingAndAutomation/HowTos/HowTos.INT.udn
new file mode 100644
--- /dev/null
+++ b/Engine/Documentation/Source/Engine/Editor/ScriptingAndAutomation/HowTos/HowTos.INT.udn
@@ -0,0 +1,16 @@
+Availability: Public
+Crumbs: 
+Title: Editor Scripting How-Tos
+Description: Provides a list of example pages that show how to use Blueprints and Python to script and automate content production tasks in the Unreal Editor.
+Version: 4.20
+Parent: Engine/Editor/ScriptingAndAutomation
+Order: 0
+Tags: scripting
+Tags: editor
+
+
+The pages in this section show you how to use Blueprints and Python to automate some common content production tasks in the Unreal Editor. You can copy these examples directly, or use them as models for your own scripts.
+
+## Intermediate
+
+[DIR(output:"fancy" type:"how-to" skill:"Intermediate" parent:"Engine/Editor/ScriptingAndAutomation/HowTos" org:"hierarchy")]
diff --git a/Engine/Documentation/Source/Engine/Editor/ScriptingAndAutomation/HowTos/SettingUpCollisionProperties/SettingUpCollisionProperties.INT.udn b/Engine/Documentation/Source/Engine/Editor/ScriptingAndAutomation/HowTos/SettingUpCollisionProperties/SettingUpCollisionProperties.INT.udn
new file mode 100644
--- /dev/null
+++ b/Engine/Documentation/Source/Engine/Editor/ScriptingAndAutomation/HowTos/SettingUpCollisionProperties/SettingUpCollisionProperties.INT.udn
@@ -0,0 +1,285 @@
+Availability: Public
+Crumbs: 
+Title: Setting up Collisions with Static Meshes in Blueprints and Python
+Description: Describes how to set up collision properties for a Static Mesh in the Unreal Editor using Blueprints or Python.
+Version: 4.20
+type:how-to
+SkillLevel:Intermediate
+Parent: Engine/Editor/ScriptingAndAutomation/HowTos
+Order: 2
+Tags: Blueprints
+Tags: Physics
+Tags: Collision
+Tags: how-to
+Tags: Python 
+Tags: scripting
+Tags: editor
+
+
+In order for a Static Mesh to interact with physics in a Level — for example, if you want it to hit other objects, or if you want to detect its presence using high-performance collision tests at runtime — you have to set up the Static Mesh with a **Collision Mesh**. This represents the bounds of the Static Mesh object within the physics simulation. Although you could use the visible geometry of the Static Mesh, this is typically far too detailed. Physical interactions typically do not need this high level of accuracy in order to provide realistic results, so you can improve the performance of the physics system by making your collision meshes as simple as possible.
+
+You can create simplified collision representations for your Static Meshes automatically in the Static Mesh Editor:
+
+[REGION:lightbox]
+[![](set_collisions_menu.png "set_collisions_menu.png")](set_collisions_menu.png) 
+[/REGION]
+[REGION:caption]
+Click image for full size.
+[/REGION]
+
+
+
+For details, see [Setting Up Collisions With Static Meshes](Engine/Content/Types/StaticMeshes/HowTo/SettingCollision).
+
+In some cases, you may need to create these Collision Mesh representations programmatically, instead of doing it by hand in the Static Mesh Editor. For example, if you need to set up a large number of Static Mesh objects in the same Project, it may not be feasible to open each one in turn. Or, you may want to carry out the collision setup as one step in a larger automated pipeline for importing and managing content.
+
+The following sections illustrate how you can use Blueprints and Python to apply the different types of Collision Meshes shown above to your Static Mesh Assets automatically in the Unreal Editor.
+
+[REGION:note]
+You'll need to install the Editor Scripting Utilities plugin, if you haven't already. For details, see [Scripting and Automating the Editor](Engine/Editor/ScriptingAndAutomation).
+[/REGION]
+
+%Globals:BlueprintOrPythonSelectVar%
+
+[OBJECT:ToggleButtonContent]
+	[PARAMLITERAL:category]
+	BlueprintOrCPPSelect
+	[/PARAMLITERAL]
+	[PARAMLITERAL:id]
+	blueprintimplementation
+	[/PARAMLITERAL]
+	[PARAMLITERAL:active]
+	active_button_content
+	[/PARAMLITERAL]
+	[PARAM:content]
+	You'll find the nodes you'll need to manage Static Mesh collisions under the **Editor Scripting > Static Mesh** category.
+
+	[REGION:note]
+	To use these nodes, your Blueprint class must be derived from an Editor-only class, such as the **PlacedEditorUtilityBase** class. For details, see [Scripting the Editor using Blueprints](Engine/Editor/ScriptingAndAutomation).
+	[/REGION]
+
+	Setting up collisions modifies the Static Mesh Asset. Assuming you want to keep the changes you make, you'll also need to use a node like **Save Asset** or **Save Loaded Asset** afterward.
+	[/PARAM]
+[/OBJECT]
+
+[OBJECT:ToggleButtonContent]
+	[PARAMLITERAL:category]
+	BlueprintOrCPPSelect
+	[/PARAMLITERAL]
+	[PARAMLITERAL:id]
+	pythonimplementation
+	[/PARAMLITERAL]
+	[PARAMLITERAL:active]
+	[/PARAMLITERAL]
+	[PARAM:content]
+	You'll find most of the functions you'll need to manage Static Mesh collisions in the `unreal.EditorStaticMeshLibrary` class.
+
+	Setting up LODs modifies the Static Mesh Asset. Assuming you want to keep the changes you make, you'll also need to use a function like `unreal.EditorAssetLibrary.save_asset()` or `unreal.EditorAssetLibrary.save_loaded_asset()` afterward.
+
+	[/PARAM]
+[/OBJECT]
+
+
+
+## Adding Simplified Collision Shapes
+
+[OBJECT:ToggleButtonContent]
+	[PARAMLITERAL:category]
+	BlueprintOrCPPSelect
+	[/PARAMLITERAL]
+	[PARAMLITERAL:id]
+	blueprintimplementation
+	[/PARAMLITERAL]
+	[PARAMLITERAL:active]
+	active_button_content
+	[/PARAMLITERAL]
+	[PARAM:content]
+	To add a new simplified collision shape to a Static Mesh, use the **Add Simple Collisions** node. Use the **Shape Type** input to control what kind of collision shape you want to add. These options match the ones available in the **Collision** menu of the Static Mesh Editor:
+
+	![](set_collisions_simple_bp.png "set_collisions_simple_bp.png")
+	[/PARAM]
+[/OBJECT]
+
+[OBJECT:ToggleButtonContent]
+	[PARAMLITERAL:category]
+	BlueprintOrCPPSelect
+	[/PARAMLITERAL]
+	[PARAMLITERAL:id]
+	pythonimplementation
+	[/PARAMLITERAL]
+	[PARAMLITERAL:active]
+	[/PARAMLITERAL]
+	[PARAM:content]
+	To add a new simplified collision shape to a Static Mesh, use the `unreal.EditorStaticMeshLibrary.add_simple_collisions()` function. Pass it:
+
+	*   The `unreal.StaticMesh` object you want to modify.
+	*   An item from the `unreal.ScriptingCollisionShapeType` enumeration that indicates what type of collision primitive you want to create. These options match the ones available in the **Collision** menu of the Static Mesh Editor.
+
+	For example:
+
+		import unreal
+		asset_path = "/Game/ArchVis/Mesh"
+		def add_box_collision (static_mesh):
+			# You could instead use .SPHERE, .CAPSULE, .NDOP10_X, .NDOP10_Y, .NDOP10_Z, .NDOP18, .NDOP26
+			shape_type = unreal.ScriptingCollisionShapeType.BOX
+			unreal.EditorStaticMeshLibrary.add_simple_collisions(static_mesh, shape_type)
+			unreal.EditorAssetLibrary.save_loaded_asset(static_mesh)
+		# get a list of all Assets in the path.
+		all_assets = unreal.EditorAssetLibrary.list_assets(asset_path)
+		# load them all into memory.
+		all_assets_loaded = [unreal.EditorAssetLibrary.load_asset(a) for a in all_assets]
+		# filter the list to include only Static Meshes.
+		static_mesh_assets = unreal.EditorFilterLibrary.by_class(all_assets_loaded, unreal.StaticMesh)
+		# run the function above on each Static Mesh in the list.
+		map(add_box_collision, static_mesh_assets)
+	[/PARAM]
+[/OBJECT]
+
+Note that this adds a new collision shape to whatever other simplified collision shapes already existed for the Static Mesh, if any. If you want to remove the existing collision shapes first, see *Removing All Simple Collisions* below.
+
+![](set_collisions_simple_result.png "set_collisions_simple_result.png")
+
+## Auto-Generating Convex Collisions
+
+[OBJECT:ToggleButtonContent]
+	[PARAMLITERAL:category]
+	BlueprintOrCPPSelect
+	[/PARAMLITERAL]
+	[PARAMLITERAL:id]
+	blueprintimplementation
+	[/PARAMLITERAL]
+	[PARAMLITERAL:active]
+	[/PARAMLITERAL]
+	[PARAM:content]
+	To auto-generate a convex collision shape for a Static Mesh from its visible geometry, use the **Set Convex Decomposition Collisions** node.
+
+	![](set_collisions_complex_bp.png "set_collisions_complex_bp.png")
+
+	The inputs in this node exactly match the options you are asked to supply when you choose **Collisions > Auto Convex Collisions** in the Static Mesh Editor user interface. They control the complexity and fidelity of the generated Collision Meshes. In general, larger values lead to Collision Meshes that are closer to the visible geometry of the Static Mesh, but more expensive to simulate at runtime.
+	[/PARAM]
+[/OBJECT]
+
+[OBJECT:ToggleButtonContent]
+	[PARAMLITERAL:category]
+	BlueprintOrCPPSelect
+	[/PARAMLITERAL]
+	[PARAMLITERAL:id]
+	pythonimplementation
+	[/PARAMLITERAL]
+	[PARAMLITERAL:active]
+	[/PARAMLITERAL]
+	[PARAM:content]
+	To auto-generate a convex collision shape for a Static Mesh from its visible geometry, use the `unreal.EditorStaticMeshLibrary.set_convex_decomposition_collisions()` function. Pass it:
+
+	*   The `unreal.StaticMesh` object you want to modify.
+	*   Three integers that define the maximum hull count, maximum number of vertices per hull, and hull precision. These parameters exactly match the options you are asked to supply when you choose **Collisions > Auto Convex Collisions** in the Static Mesh Editor user interface. They control the complexity and fidelity of the generated collision meshes. In general, larger values lead to collision meshes that are closer to the visible geometry of the Static Mesh, but more expensive to simulate at runtime.
+
+	For example:
+
+		import unreal
+		asset_path = "/Game/ArchVis/Mesh"
+		def set_convex_collision (static_mesh):
+			unreal.EditorStaticMeshLibrary.set_convex_decomposition_collisions(static_mesh, 4, 12, 460000)
+			unreal.EditorAssetLibrary.save_loaded_asset(static_mesh)
+		# get a list of all Assets in the path.
+		all_assets = unreal.EditorAssetLibrary.list_assets(asset_path)# load them all into memory.
+		all_assets_loaded = [unreal.EditorAssetLibrary.load_asset(a) for a in all_assets]# filter the list to include only Static Meshes.
+		static_mesh_assets = unreal.EditorFilterLibrary.by_class(all_assets_loaded, unreal.StaticMesh)# run the function above on each Static Mesh in the list.
+		map(set_convex_collision, static_mesh_assets)
+	[/PARAM]
+[/OBJECT]
+
+All existing Collision Meshes are automatically removed from the Static Mesh before the new mesh is generated.
+
+Note that this method tends to produce less predictable and regular results than using a simplified collision primitive. It's best used on irregular meshes, or when you can tune the generation settings visually in order to make sure that the results you generate are simple enough and a good match for the visible geometry of your Static Mesh.
+
+![](set_collisions_convex_result.png "set_collisions_convex_result.png")
+
+## Removing All Simple Collisions
+[OBJECT:ToggleButtonContent]
+	[PARAMLITERAL:category]
+	BlueprintOrCPPSelect
+	[/PARAMLITERAL]
+	[PARAMLITERAL:id]
+	blueprintimplementation
+	[/PARAMLITERAL]
+	[PARAMLITERAL:active]
+	[/PARAMLITERAL]
+	[PARAM:content]
+	You can clear all the Collision Meshes assigned to your Static Mesh using the **Remove Collisions** node.
+
+	![](set_collisions_remove_bp.png "set_collisions_remove_bp.png")
+
+	After this, the mesh will not be found by any "simple" physics collision tests, but will still be found by "detailed" tests that consider the visible geometry of the Static Mesh. See also [Simple vs Complex Collision](Engine/Physics/SimpleVsComplex).
+	[/PARAM]
+[/OBJECT]
+
+[OBJECT:ToggleButtonContent]
+	[PARAMLITERAL:category]
+	BlueprintOrCPPSelect
+	[/PARAMLITERAL]
+	[PARAMLITERAL:id]
+	pythonimplementation
+	[/PARAMLITERAL]
+	[PARAMLITERAL:active]
+	[/PARAMLITERAL]
+	[PARAM:content]
+	You can clear all the Collision Meshes assigned to your Static Mesh using the `unreal.EditorStaticMeshLibrary.remove_collisions()` function.
+
+	For example:
+
+		import unreal
+		asset_path = "/Game/ArchVis/Mesh"
+		def remove_collisions (static_mesh):
+			unreal.EditorStaticMeshLibrary.remove_collisions(static_mesh)
+			unreal.EditorAssetLibrary.save_loaded_asset(static_mesh)
+		# get a list of all Assets in the path.
+		all_assets = unreal.EditorAssetLibrary.list_assets(asset_path)# load them all into memory.
+		all_assets_loaded = [unreal.EditorAssetLibrary.load_asset(a) for a in all_assets]# filter the list to include only Static Meshes.
+		static_mesh_assets = unreal.EditorFilterLibrary.by_class(all_assets_loaded, unreal.StaticMesh)# run the function above on each Static Mesh in the list.
+		map(remove_collision, static_mesh_assets)
+
+	After this, the mesh will not be found by any "simple" physics collision tests, but will still be found by "detailed" tests that consider the visible geometry of the Static Mesh. See also [Simple vs Complex Collision](Engine/Physics/SimpleVsComplex).
+	[/PARAM]
+[/OBJECT]
+
+[OBJECT:ToggleButtonContent]
+	[PARAMLITERAL:category]
+	BlueprintOrCPPSelect
+	[/PARAMLITERAL]
+	[PARAMLITERAL:id]
+	pythonimplementation
+	[/PARAMLITERAL]
+	[PARAMLITERAL:active]
+	[/PARAMLITERAL]
+	[PARAM:content]
+	## Using an LOD for Collisions
+
+	If you have already set up Levels of Detail (LODs) for your Static Mesh, you can use one of the less-detailed LODs as the Collision Mesh.
+
+	Call the `set_editor_property()` function on the `unreal.StaticMesh` object to set the `lod_for_collision` property to the index of the LOD you want to use. For example:
+
+		import unreal
+		asset_path = "/Game/ArchVis/Mesh"
+		def use_lod_for_collisions (static_mesh):
+			static_mesh.set_editor_property("lod_for_collision", 3)
+			unreal.EditorAssetLibrary.save_loaded_asset(static_mesh)
+		# get a list of all Assets in the path.
+		all_assets = unreal.EditorAssetLibrary.list_assets(asset_path)
+		# load them all into memory.
+		all_assets_loaded = [unreal.EditorAssetLibrary.load_asset(a) for a in all_assets]
+		# filter the list to include only Static Meshes.
+		static_mesh_assets = unreal.EditorFilterLibrary.by_class(all_assets_loaded, unreal.StaticMesh)
+		# run the function above on each Static Mesh in the list.
+		map(use_lod_for_collision, static_mesh_assets)
+
+	See also [How To Set LOD Collision](Engine/Content/Types/StaticMeshes/HowTo/LODCollision).
+
+	[/PARAM]
+[/OBJECT]
+[REGION:note]
+You can't yet set up a custom Collision Mesh via Blueprint or Python. To import a custom mesh and use it for a Static Mesh's Collision Mesh in the physics simulation, you have to either:
+
+*   Use the Static Mesh Editor user interface to import the Collision Mesh from a supported file format.
+*   Import the Collision Mesh at the same time as the visible Static Mesh, using a special naming convention to indicate that it represents the geometry you want to use for collision testing. For details, see [FBX Static Mesh Pipeline](Engine/Content/FBX/StaticMeshes#collision)<!--, or [Datasmith Overview](Studio/Unreal-Datasmith/Datasmith)-->.
+[/REGION]
\ No newline at end of file
diff --git a/Engine/Documentation/Source/Engine/Editor/ScriptingAndAutomation/Python/Python.INT.udn b/Engine/Documentation/Source/Engine/Editor/ScriptingAndAutomation/Python/Python.INT.udn
new file mode 100644
--- /dev/null
+++ b/Engine/Documentation/Source/Engine/Editor/ScriptingAndAutomation/Python/Python.INT.udn
@@ -0,0 +1,315 @@
+Availability: Public
+Crumbs: 
+Title: Scripting the Editor using Python
+Description: Describes how to use Python in the Unreal Editor to script content production tasks.
+Type: 
+Version: 4.20
+Parent: Engine/Editor/ScriptingAndAutomation
+Order: 4
+Tags: Unreal Studio
+Tags: Datasmith
+Tags: overview
+Tags: landing
+Tags: Python
+Tags: Experimental
+Tags: scripting
+Tags: editor
+Tags: Landing Page
+
+
+This page describes how to get started using Python in the Unreal Editor.
+
+## Why Python?
+
+In recent years, Python has become the de facto language for production pipelines and interoperability between 3D applications, particularly in the media and entertainment industry. This is partially due to the wide range of applications that support it. As the complexity of production pipelines continues to soar, and the number of applications involved continues to grow, having a common scripting language makes it easier to create and maintain large-scale asset management systems.
+
+Even without these outside considerations, or the need to work with other applications, Python is a great choice if you're looking to automate your workflows within the Unreal Editor. It's relatively easy for those new to programming to get started, it offers the ability to create complex and full-featured user interfaces through modules like PySide, and there are many other useful free modules available to the community to make your life easier.
+
+You can use Python in the Unreal Editor to do things like:
+
+*   Construct larger-scale asset management pipelines or workflows that tie the Unreal Editor to other 3D applications that you use in your organization.
+*   Automate time-consuming Asset management tasks in the Unreal Editor, like generating LODs for Static Meshes.
+*   Procedurally lay out content in a Level.
+*   Control the Unreal Editor from UIs that you create yourself in Python.
+
+## Set Up Your Project to Use Python
+
+Python support in the Unreal Editor is provided by the Python Editor Script Plugin. Before you can run Python scripts in the Editor, you'll need to enable this plugin for your current Project.
+
+[REGION:note]
+Currently, you must enable the plugin separately for each Project.
+[/REGION]
+
+To enable the plugin:
+
+1.  Open your Project, and choose **Edit > Plugins** from the main menu.
+2.  In the **Plugins** window, go to the **Scripting** section.  
+    Find the **Python Editor Script Plugin** in the right-hand panel, and check its **Enabled** box.  
+    ![Enable the Python Editor Script Plugin](install-python-plugin.png "Enable the Python Editor Script Plugin")  
+    You'll also want to enable the **Editor Scripting Utilities** plugin, which offers simplified APIs for many common Editor tasks. For details, see [Scripting and Automating the Editor](https://docs.unrealengine.com/en-us/Editor/Scripting-and-Automating-the-Editor).
+3.  Restart the Editor.
+
+### Python 2.7
+
+The Python Editor Script Plugin contains an embedded version of Python 2.7.
+
+This means that you don't have to install Python separately on your computer.
+
+[REGION:tip]
+Unreal uses Python 2.7 by default because it is an important part of the current [VFX Reference Platform](https://www.vfxplatform.com/). When the reference platform switches to version 3.x, we intend to follow.  
+In the meantime, if you need to use a different version of Python (including a 3.x version), you can set the `UE_PYTHON_DIR` environment variable to point to the installation you want to embed, then rebuild the Unreal Engine from source.
+[/REGION]
+
+## Ways to Run Python Code in Unreal Editor
+
+There are several different ways that you can run Python scripts in the Unreal Editor, each designed for a slightly different usage scenario. You can choose whichever fits your needs.
+
+[REGION:note]
+Unlike Blueprints, the Python environment is only available **in the Unreal Editor**, not when your Project is running in the Unreal Engine in any mode, including Play In Editor, Standalone Game, cooked executable, etc. That means that you can use Python freely for scripting and automating the Editor or building asset production pipelines, but you cannot currently use it as a gameplay scripting language.
+[/REGION]
+
+### The Python Console in the Output Log
+
+You can switch the Unreal Editor's console input bar to accept Python code instead of Unreal console commands.
+
+![The Python console](python-console-input.png "The Python console")
+
+You can do this in the **Output Log** panel, as shown above, or when you bring up the console input bar by pressing the `~` key.  
+
+When the console is in Python mode:
+
+*   You can enter lines of Python code into this console and have the Editor execute each one immediately, exactly as if you were using an interactive Python console in a command window. This is the only way to execute Python code line-by-line; all other approaches listed below run a script file that you specify.
+*   You can run multiple lines of code at a time by using **Shift+Enter** to separate each line, or by pasting in a multi-line block that you copy from a text editor.
+*   You can execute Python script files by simply typing the file name into the console. If your Python script requires additional command-line arguments, include them after the name of your script.
+
+[REGION:tip]
+Output from the built-in Python `print` function is also redirected to the **Output Log** panel.
+[/REGION]
+
+### The py Console Command
+
+In the normal console, you can use the `py` command to run the rest of the line as Python code, exactly as if you had typed it into the Python console described above.
+
+For example, this command runs the specified script file:
+
+![The py console command](python-py-command.png "The py console command")
+
+[REGION:note]
+We don't recommend running this command in the value of the `ExecCmd` command-line parameter when you start the Editor. This can cause your script to run before the Editor environment is ready — for example, before the startup Level is fully loaded. See the sections below for better options.
+[/REGION]
+
+### The File Menu
+
+The **File** menu in the main window of the Unreal Editor offers new options that you can use to run Python script files.
+
+*   Use **Execute Python Script** if you want to browse to a new script file on your computer that you haven't run before.
+*   Use the **Recent Python scripts** list to re-run any script that you ran previously. The file is read from disk again, so you changed the script in the meantime, your new version is run.
+
+![Execute Python from the File menu](python-file-menu-options.png "Execute Python from the File menu")
+
+### The Command Line
+
+If you start the Unreal Editor from the command line or from a script, you can specify a Python script file in the command line arguments. If your Python script requires additional command-line arguments, include them after the name of your script.
+
+There are two different ways to run a Python script from the command line. In both approaches, the Editor shuts down immediately after running your Python script.
+
+**Option 1: Full Editor**
+
+In this approach, the full Unreal Editor launches, opens your specified Project, loads the default startup level, then runs your Python script once everything is loaded and ready. This approach is good if you need your script to interact with content in your Project or in a Level that can take some time to load.
+
+Add the `ExecutePythonScript` argument to the command line, and set its value to the Python script you want to run. For example:
+
+    > UE4Editor-Cmd.exe MyProject.uproject -ExecutePythonScript="c:\my_script.py"
+
+[REGION:note]
+The approach above requires that you enable the Editor Scripting Utilities plugin for your Project.
+[/REGION]
+
+**Option 2: Commandlet**
+
+In this approach, the Editor starts up with a minimal environment, without UI or rendering. This approach is very fast to execute, but it can be trickier to load Levels and other kinds of Assets that your script needs to interact with.
+
+Add the following arguments to the command line: `-run=pythonscript -script=<script_file>`
+
+For example:
+
+    > UE4Editor-Cmd.exe -run=pythonscript -script="c:\\my_script.py"
+
+### The init_unreal.py File
+
+If the Editor detects a script file called `init_unreal.py` in any of the paths it is configured to use (see "Python Paths in the Unreal Editor" below), it automatically runs that script immediately.
+
+This is a good approach for situations where you are working on a Project or a Plugin and you know that everyone working with that content needs to run the same initialization code every time the Editor starts up. You could put your initialization code inside a script with this name, and put it into the **Content/Python** folder within that Project or Plugin.
+
+### Startup Scripts
+
+In your Project Settings, you can specify any number of Python scripts that you want to run every time you open that Project. The Editor runs these scripts after the default startup Level is fully loaded.
+
+Select **Edit > Project Settings...**. Under the **Plugins** list, select **Python**. Then, add your scripts to the **Startup scripts** setting:
+
+![Python startup scripts](python-startup-scripts.png "Python startup scripts")
+
+Restart the Unreal Editor when done. The next time the Editor loads your Project, it should run your new startup scripts.
+
+## Python Paths in the Unreal Editor
+
+When you use a relative path to run a Python script using the methods above, or to import another script module using the `import` command in one of your scripts, the script that you run or import can be in any path that is listed in the `sys.path` variable of the Python environment.
+
+The Unreal Editor automatically adds several paths to this `sys.path` list:
+
+*   The **Content/Python** sub-folder under your Project's folder.
+*   The **Content/Python** sub-folder in the main Unreal Engine installation.
+*   The **Content/Python** sub-folder under each enabled Plugin's folder.
+*   The **Documents/UnrealEngine/Python** folder inside your user directory. For example, on Windows 10, this is equivalent to `[REDACTED_USER_HOME]/Documents/UnrealEngine/Python`
+
+You can also add your own paths to this list using any of the following approaches:
+
+*   In your Project Settings. Select **Edit > Project Settings...**. Under the **Plugins** list, select **Python**. Then, add the paths to the **Additional Paths** setting. Restart the Unreal Editor when done.  
+    ![Additional Python paths](python-additional-paths.png "Additional Python paths")
+*   Add the paths to the value of the `PYTHONPATH` environment variable, then restart the Unreal Editor.
+*   Add the paths directly to the `sys.path` list within a Python script, or in the Python console.
+
+## About the Unreal Editor Python API
+
+The Python Editor Script Plugin exposes a wide range of classes and functions that you can use to interact with the Unreal Editor, the Assets in your Project, and the content in your Levels. This API is all contained in the `unreal` module. To access it, import this module at the beginning of any Python script you run in the Editor's Python environment:
+
+
+    import unreal
+
+
+The `unreal` module exposes nearly everything that is exposed from C++ to Blueprints in your Editor environment. It's not pre-generated; it automatically reflects whatever is available in Blueprints in your Editor. As you enable new plugins in the Unreal Editor, anything those plugins expose to Blueprints also becomes available in Python as well. The same goes for any C++ code that you write in your Project and expose to Blueprints.
+
+The Python API makes every effort to expose native Unreal APIs in a way that is as friendly as possible to Python developers. For example:
+
+*   Simple data types are transparently converted back and forth between Python and native types whenever necessary.  
+    When you pass in a Python list, set or dict, it is automatically converted to an  Unreal array, set, or map. When you retrieve a list, set, or dict returned by an API function, you are actually getting an instance of an Unreal class, but its API is fully consistent with the base Python list, set, or dict type.
+*   Python classes maintain the same inheritance hierarchy as the native types they represent. That means, for example, that you can use the built-in Python functions `isinstance()` and `type()` to test whether an object derives from or matches a given class.
+*   The API tries to strike a good balance between the naming conventions used in Unreal for C++ and Blueprints on one hand, and Python naming conventions on the other hand. Classes and objects in the Python API have the same names as they do in Blueprints. This is typically the same as their native types, omitting the prefix (e.g. `U` or `T`). Function and property names are automatically exposed as lower-case `snake_case`. So, for example, you typically call functions like `unreal.StaticMeshActor.get_actor_transform()`. Enumeration value names are automatically exposed as upper-case `SNAKE_CASE`.
+*   All exposed functions can accept either ordered parameters, or named parameters in any order. For example, the following two function calls are exactly equivalent:
+    
+    ~~~
+    unreal.EditorLevelLibrary.join_static_mesh_actors(list_of_actors, my_options)
+    unreal.EditorLevelLibrary.join_static_mesh_actors(join_options=my_options, actors_to_join=list_of_actors)
+    ~~~
+    
+
+### API Reference
+
+For details on all of the classes and functions exposed by the Unreal Python API, see the API Reference here:
+
+[**Unreal Editor Python API Reference**](https://api.unrealengine.com/INT/PythonAPI/)
+
+[REGION:note]
+The API Reference is not an exhaustive list of everything that may be exposed to Python by plugins. If you've installed additional plugins that aren't included in the API Reference, and you need to see the way their scripting features are exposed to Python, you can generate your own local version of the API Reference that contains docs for the plugins you need. For instructions, see the readme file under *Engine\Plugins\Experimental\PythonScriptPlugin\SphinxDocs* in your Unreal Engine installation folder.
+[/REGION]
+
+## Best Practices for Using the Python API
+
+This section covers a few things that it's important to keep in mind when you use the Python API.  
+
+### Working with Assets
+
+If you need to work with Assets in your Project, always use functions from the Unreal Python API to do it. Never use file management modules built in to Python to work directly with your Asset files on disk. For example, if you need to move an Asset to a different folder, do not use Python functions like `os.rename` or `shutil.move`. Unreal Projects and Assets contain internal content references that you can break if you do not respect this rule.
+
+Instead, we recommend using the `unreal.EditorAssetLibrary` API that is provided by the Editor Scripting Utilities plugin, or the `unreal.AssetTools` class built in to the Unreal Python API.
+
+### Changing Editor Properties
+
+You can use Python to get access to Objects in your Project and set up many configuration properties on those Objects programmatically. For example, your Python script could access Static Mesh Actors in the current Level, and set properties like whether the Actors can be damaged, or whether they should be hidden in the game. Or, you could retrieve their Static Mesh Components and set up properties on those Components, like their Lightmass settings, or even the Static Mesh Asset that they are linked to.
+
+These properties may be exposed to Python in two different ways:
+
+*   Items with the BlueprintReadOnly or BlueprintReadWrite flag are exposed as simple properties on the object.  
+    You can read and modify these properties like you access any Python object property.
+*   Items with the ViewAnywhere or EditAnywhere flag are exposed as "editor properties".  
+    You can read and write these values using a special pair of functions exposed by every object: `set_editor_property()` and `get_editor_property()`.
+
+In the API Reference for each class, you'll find a list of **Editor Properties** immediately following the class description. These are all the values that you can set and get using these `set_editor_property()` and `get_editor_property()` functions. Whenever you need to set or get a configuration property on an object, check this list first  to see if the property you want is listed there.
+
+*   When you need to read a value that is exposed both as an object property and as an editor property, the result of accessing the property directly is usually the same as by calling the `get_editor_property()` function. However, the `get_editor_property()` function often has access to properties that aren't exposed directly on the Python object.
+*   When you need to set a value that is exposed both as an object property and as an editor property, you should in most cases use the `set_editor_property()` function to set the value rather than set the value directly on the object. When you adjust properties in the UI, the Editor often performs additional operations behind the scenes: pre- and post-edit changes. These operations typically respond to the choices you make in some way, and keep the Editor UI in sync with the state of the object in the game world. If you modify these properties directly on the Python object, this Editor code won't be run automatically. On the other hand, when you call `set_editor_property()` to set the state of a property, you do trigger this pre- and post-edit code, exactly as if you changed the setting in the **Details** panel of the Editor UI.
+
+For example, Media Player objects have a **Play on Open** setting:  
+
+![Details](Details.png "Details")
+
+This is exposed in the `unreal.MediaPlayer` class in the `play_on_open` class member:
+
+    import unreal
+    obj = unreal.MediaPlayer()
+    # Modifying a property directly can have different results
+    # than changing settings in the Editor UI.
+    # Generally you'll want to avoid setting these values directly, like this:
+    obj.play_on_open = True
+    # This way of accessing the property will have exactly the same
+    # result as changing the setting in the Editor UI:
+    obj.set_editor_property("play_on_open", True)
+    # Both ways of reading the value are equivalent.
+    # When a class exposes a property in both ways, you can use either:
+    play_value = obj.play_on_open
+    play_value = obj.get_editor_property("play_on_open")
+
+### Use Unreal Types Whenever Possible
+
+Whenever you need utilities that are available in the Unreal Python API, like classes for math operations or manipulating 3D coordinates, we recommend using the Unreal utilities rather than using your own implementations. The Unreal versions are optimized for best performance in the Engine environment.
+
+For example, when you need to manipulate coordinates in 3D space, use the `unreal.Vector` class:
+
+    import unreal
+    v1 = unreal.Vector()
+    v1.x = 10
+    v2 = unreal.Vector(10, 20, 30)
+    v3 = (v1 + v2) * 2
+    print v3
+
+### Logging and Feedback
+
+The `unreal` object exposes functions that you can use in your code to send log, warning, and error messages through the same messaging system used by all Engine and Editor subsystems. We recommend using this standardized logging framework anytime your script needs to send a message to the user.
+
+*   Use `unreal.log()` for information messages. For your convenience, the Python `print()` function has also been implemented to pass through `unreal.log()` internally.
+*   Use `unreal.log_warning()` to alert users of potential problems.
+*   Use `unreal.log_error()` for severe problems that prevent your script from running as expected.
+
+Your messages appear in the **Output Log** panel, along with the messages sent by other subsystems:
+
+![Python log messages](python-log-feedback.png "Python log messages")
+
+### Supporting Undo and Redo
+
+Your scripts can take full advantage of the Undo / Redo system built in to the Unreal Editor.
+
+Each *transaction* that you define can contain any number of Python operations. Using these transactions, you can bundle large operations, or operations on many different objects, together into a single entry in the Undo / Redo history. Typically, if your script intends to perform a certain change on multiple objects, you don't want a separate entry in the Undo / Redo history for each change; you want one entry that will undo all changes to all objects.
+
+To define a transaction, use the `unreal.ScopedEditorTransaction` scope. For example, if you run this code:
+
+    import unreal
+    obj = unreal.MediaPlayer()
+    with unreal.ScopedEditorTransaction("My Transaction Test") as trans:
+        obj.set_editor_property(“play_on_open”, True)
+        obj.set_editor_property(“vertical_field_of_view”, 60)
+
+Your Editor's **Undo History** panel now lists that transaction by name:
+
+![Undo History](undoHistory.png "Undo History")
+
+As a general rule, your scoped transactions can contain any operations that are also undoable in the Editor UI. However, not every Editor operation is undoable. For example, you can't undo importing a model in the Editor UI, so trying to import a model inside an `unreal.ScopedEditorTransaction` will not work as you may be expecting.
+
+### Progress Dialogs for Slow Operations
+
+If your scripts need to work on many Assets or Actors in the same operation, they may take some time to complete. However, while the Unreal Editor is running a Python script, its UI becomes blocked to other user interactions. To give the user information about the progress of a large task, and avoid the Editor appearing to freeze or hang, you can use the `unreal.ScopedSlowTask` scope.
+
+For example:
+
+    import unreal
+    total_frames = 100
+    text_label = "Working!"
+    with unreal.ScopedSlowTask(total_frames, text_label) as slow_task:
+        slow_task.make_dialog(True)               # Makes the dialog visible, if it isn't already
+        for i in range(total_frames):
+            if slow_task.should_cancel():         # True if the user has pressed Cancel in the UI
+                break
+            slow_task.enter_progress_frame(1)     # Advance progress by one frame.
+                                                # You can also update the dialog text in this call, if you want.
+            ...                                   # Now do work for the current frame here!
+
+![Progress bar for a Scoped Slow Task](python-scoped-slow-task.png "Progress bar for a Scoped Slow Task")
\ No newline at end of file
