Copying //UE4/Dev-Documentation to //UE4/Samples-Main (Source: //UE4/Dev-Documentation @ 4321946)

diff --git a/Engine/Documentation/Source/Engine/Editor/ScriptingAndAutomation/ScriptingAndAutomation.INT.udn b/Engine/Documentation/Source/Engine/Editor/ScriptingAndAutomation/ScriptingAndAutomation.INT.udn
new file mode 100644
--- /dev/null
+++ b/Engine/Documentation/Source/Engine/Editor/ScriptingAndAutomation/ScriptingAndAutomation.INT.udn
@@ -0,0 +1,34 @@
+Availability: Public
+Crumbs: 
+Title: Scripting and Automating the Editor
+Description: Introduction to using Blueprints and Python to control the Unreal Editor programmatically.
+type: landing
+Version: 4.20
+Parent: Engine/Editor
+Order: 1
+Tags: Blueprints, Python, scripting
+
+
+In the Unreal Editor user interface, you have access to a vast range of visual tools for setting up your Project, designing and building Levels, creating gameplay interactions, and more. But sometimes, once you've figured out what you need the Editor to do, you might want to invoke its tools and commands programmatically in a reusable script.
+
+This can help you:
+
+*   Minimize or eliminate the need to repeat the same series of tasks over and over,
+*   Automate or randomize the placement, layout and settings of Actors in your Levels,
+*   Create your own Asset import and management pipelines,
+*   Interoperate with other third-party applications and pipeline scripts,
+*   Extend the Editor with additional tools, and even UIs, that you customize specifically for the needs of your Project or content.
+
+The pages in this section show how you can run these kinds of in-Editor scripts using [Blueprints](Engine/Editor/ScriptingAndAutomation/Blueprints) and [Python](Engine/Editor/ScriptingAndAutomation/Python) as you're working on your Project.
+
+## Installing the Editor Scripting Utilities Plugin
+
+If you're planning to do any Editor automation, regardless of whether you intend to use Blueprints, Python, or some other language, you will almost certainly want to install the **Editor Scripting Utilities** plugin. This plugin simplifies many of the most common operations you'll need to do in the Editor, handling uncommon edge cases so that you don't have to understand all the internal details of how the Editor works in order to do something conceptually simple.
+
+To install the plug-in:
+
+1.  From the main menu, choose **Editor > Plugins** to open the **Plugins** window.
+2.  Under the **Scripting** category, find the entry for the **Editor Scripting Utilities** and check its **Enabled** box.  
+    ![Enable the Editor Scripting Utilities plugin](install-editor-scripting.png "Enable the Editor Scripting Utilities plugin")
+3.  If you're interested in using Python, you can also check the **Enabled** box for the Python plugin while you're here. For more, see [Scripting the Editor using Python](Engine/Editor/ScriptingAndAutomation/Python).
+4.  Restart the Editor and reload your Project.
\ No newline at end of file
diff --git a/Engine/Documentation/Source/Engine/Landscape/Editing/Gizmos/Gizmos.INT.udn b/Engine/Documentation/Source/Engine/Landscape/Editing/Gizmos/Gizmos.INT.udn
--- a/Engine/Documentation/Source/Engine/Landscape/Editing/Gizmos/Gizmos.INT.udn
+++ b/Engine/Documentation/Source/Engine/Landscape/Editing/Gizmos/Gizmos.INT.udn
@@ -203,22 +203,19 @@ Heightmap data can be imported from and exported to the Gizmo through the **Gizm
 
 
 [COMMENT:none]
-
 ## Gizmo List
 
-![gizmolist.jpg](gizmolist.jpg)(w:267)
+![gizmolist.jpg]()(w:267)
 
 The Gizmo List displays a list of saved Gizmos so that they can easily be reused at a later time.
 
-![gizmolist_populated.jpg](gizmolist_populated.jpg)(w:266)
+![gizmolist_populated.jpg]()(w:266)
 
 | **Control** | **Description** |
 | --- | --- |
-| ![gizmolist_pin.jpg](gizmolist_pin.jpg)| Saves the Gizmo's current position, size, and rotation to the list. |
-| ![gizmolist_remove.jpg](gizmolist_remove.jpg)| Removes the Gizmo currently selected in the list. |
-| ![gizmolist_goto.jpg](gizmolist_goto.jpg)| Sets the Gizmo's position, size, and rotation to that of the currently selected Gizmo in the list. |
-
-
+| ![gizmolist_pin.jpg]()| Saves the Gizmo's current position, size, and rotation to the list. |
+| ![gizmolist_remove.jpg]()| Removes the Gizmo currently selected in the list. |
+| ![gizmolist_goto.jpg]()| Sets the Gizmo's position, size, and rotation to that of the currently selected Gizmo in the list. |
 [/COMMENT]
 
 [INCLUDE:Engine/Landscape#landscape002]
diff --git a/Engine/Documentation/Source/Engine/Landscape/Editing/PaintMode/PaintMode.INT.udn b/Engine/Documentation/Source/Engine/Landscape/Editing/PaintMode/PaintMode.INT.udn
--- a/Engine/Documentation/Source/Engine/Landscape/Editing/PaintMode/PaintMode.INT.udn
+++ b/Engine/Documentation/Source/Engine/Landscape/Editing/PaintMode/PaintMode.INT.udn
@@ -102,11 +102,9 @@ The Flatten tool directly sets the selected layer's weight to the value of the *
 | **Flatten Mode** | Determines whether the tool will increase or decrease the application of the selected layer's weight, or do both. |
 
 [COMMENT:none]
---
-| **Use Slope Flatten** | If checked, flattens along a Landscape's existing slope. |
-| **Pick Value Per Apply** | If checked, constantly selects new values to flatten toward, instead of only using the first clicked point. |
-| **Flatten Target** | Sets the target value toward which to flatten. |
---
+|**Use Slope Flatten**|If checked, flattens along a Landscape's existing slope.|
+|**Pick Value Per Apply**|If checked, constantly selects new values to flatten toward, instead of only using the first clicked point.|
+|**Flatten Target** | Sets the target value toward which to flatten.|
 [/COMMENT]
 
 
@@ -192,7 +190,7 @@ The primary purpose of layer info objects is to act as a unique key for painted
 | Option | Description |
 | ------ | ----------- |
 | **Phys Material** | The [Physical Material](Engine/Physics/PhysicalMaterials) (if any) assigned to areas of Landscape where this layer is dominant. |
-| **Hardness** | The value used by the [Erosion](Engine/Landscape/Editing/SculptMode/#Erosion) tool. |
+| **Hardness** | The value used by the [Erosion](Engine/Landscape/Editing/SculptMode/Erosion) tool. |
 
 ### Orphaned Layers
 
diff --git a/Engine/Documentation/Source/Engine/Landscape/Materials/Materials.INT.udn b/Engine/Documentation/Source/Engine/Landscape/Materials/Materials.INT.udn
--- a/Engine/Documentation/Source/Engine/Landscape/Materials/Materials.INT.udn
+++ b/Engine/Documentation/Source/Engine/Landscape/Materials/Materials.INT.udn
@@ -320,9 +320,6 @@ You can use the same Material for your Landscape Material and your Landscape Hol
 	[/PARAMLITERAL]
 	[/OBJECT]
 
-[REGION:note]
-For information about using the painting tools to create a Landscape hole, see [Using Opacity Masking to Create a Hole](Engine\Landscape\Editing\SculptMode\#UsingOpacityMaskingtoCreateaHole).
-[/REGION]
 
 
 ## Assigning Materials to Landscapes
@@ -350,9 +347,6 @@ You can assign a different Material to act as the **Landscape Hole Material**, w
 	![](Landscape_HoleSettings.png)
 
 
-[REGION:note]
-For information about using Landscape Hole Materials, see [Using Opacity Masking to Create a Hole](Engine\Landscape\Editing\SculptMode\#UsingOpacityMaskingtoCreateaHole).
-[/REGION]
 
 [COMMENT:none]
 
diff --git a/Engine/Documentation/Source/Engine/Landscape/QuickStart/6/LSQS_6.INT.udn b/Engine/Documentation/Source/Engine/Landscape/QuickStart/6/LSQS_6.INT.udn
--- a/Engine/Documentation/Source/Engine/Landscape/QuickStart/6/LSQS_6.INT.udn
+++ b/Engine/Documentation/Source/Engine/Landscape/QuickStart/6/LSQS_6.INT.udn
@@ -30,7 +30,6 @@ While the previous quick start tutorial for Landscape will get you up and runnin
 [COMMENT:none]
 Before choosing the size of your Landscape there are many things that you need to consider first like target platform and performance impact. To make things a little easier for you, here is a chart of recommended Landscape component counts and sizes that might help in choosing the best size for your Landscape. 
 
-[INCLUDE:Engine/Landscape/Creation#RecommendedLSS]
 [/COMMENT]
 
 * When using the **Paint Tools**, you might find it easier to paint over what you would like to erase than to try and erase it using **Shift** + **Left Mouse Button**.
diff --git a/Engine/Documentation/Source/Engine/LevelStreaming/StreamingVolumes/StreamingVolumes.INT.udn b/Engine/Documentation/Source/Engine/LevelStreaming/StreamingVolumes/StreamingVolumes.INT.udn
--- a/Engine/Documentation/Source/Engine/LevelStreaming/StreamingVolumes/StreamingVolumes.INT.udn
+++ b/Engine/Documentation/Source/Engine/LevelStreaming/StreamingVolumes/StreamingVolumes.INT.udn
@@ -14,8 +14,6 @@ tags:Level Streaming
 
 You can adjust how a Level Streaming Volume handles level streaming by adjusting the properties from the **Details** panel, pictured below. 
 
-![](levelstreamingvolume.png)
-
 For more on Level Streaming Volumes, see [](Engine/LevelStreaming/StreamingVolumes).
 
 |Property| Description |
diff --git a/Engine/Documentation/Source/Engine/Levels/LevelsWindow/LevelsWindow.INT.udn b/Engine/Documentation/Source/Engine/Levels/LevelsWindow/LevelsWindow.INT.udn
--- a/Engine/Documentation/Source/Engine/Levels/LevelsWindow/LevelsWindow.INT.udn
+++ b/Engine/Documentation/Source/Engine/Levels/LevelsWindow/LevelsWindow.INT.udn
@@ -14,31 +14,31 @@ topic-icon:%ROOT%/levels_icon.png
 Levels are managed with the **Levels** window. You can access the **Levels** window through the **Windows** menu.
 
 [OBJECT:ToggleButtonContent]
-[PARAMLITERAL:category]
-OS
-[/PARAMLITERAL]
-[PARAMLITERAL:id]
-windows
-[/PARAMLITERAL]
-[PARAMLITERAL:active]
-[/PARAMLITERAL]
-[PARAM:content]
-![](w:500)(LevelsMenu_Windows.png)
-[/PARAM]
+	[PARAMLITERAL:category]
+	OS
+	[/PARAMLITERAL]
+	[PARAMLITERAL:id]
+	windows
+	[/PARAMLITERAL]
+	[PARAMLITERAL:active]
+	[/PARAMLITERAL]
+	[PARAM:content]
+	![](LevelsMenu_Windows.png)(w:500)
+	[/PARAM]
 [/OBJECT]
 
 [OBJECT:ToggleButtonContent]
-[PARAMLITERAL:category]
-OS
-[/PARAMLITERAL]
-[PARAMLITERAL:id]
-mac
-[/PARAMLITERAL]
-[PARAMLITERAL:active]
-[/PARAMLITERAL]
-[PARAM:content]
-![](w:500)(LevelsMenu_Mac.png)
-[/PARAM]
+	[PARAMLITERAL:category]
+	OS
+	[/PARAMLITERAL]
+	[PARAMLITERAL:id]
+	mac
+	[/PARAMLITERAL]
+	[PARAMLITERAL:active]
+	[/PARAMLITERAL]
+	[PARAM:content]
+	![](LevelsMenu_Mac.png)(w:500)
+	[/PARAM]
 [/OBJECT]
 
 You will always have a **Persistent Level**, and you can have one or more sublevels that are always loaded or are streamed in with Level Streaming Volumes, Blueprints, or C++ code.
@@ -74,31 +74,31 @@ Adding a new sublevel makes it the current Level, so if you want to continue wor
 1. Select your Level to add in the **Open Level** dialog, then click on **Open**.
 	
 	[OBJECT:ToggleButtonContent]
-	[PARAMLITERAL:category]
-	OS
-	[/PARAMLITERAL]
-	[PARAMLITERAL:id]
-	windows
-	[/PARAMLITERAL]
-	[PARAMLITERAL:active]
-	[/PARAMLITERAL]
-	[PARAM:content]
-	![](AddExisting_Windows.png)(w:700)
-	[/PARAM]
+		[PARAMLITERAL:category]
+		OS
+		[/PARAMLITERAL]
+		[PARAMLITERAL:id]
+		windows
+		[/PARAMLITERAL]
+		[PARAMLITERAL:active]
+		[/PARAMLITERAL]
+		[PARAM:content]
+		![](AddExisting_Windows.png)(w:700)
+		[/PARAM]
 	[/OBJECT]
 
 	[OBJECT:ToggleButtonContent]
-	[PARAMLITERAL:category]
-	OS
-	[/PARAMLITERAL]
-	[PARAMLITERAL:id]
-	mac
-	[/PARAMLITERAL]
-	[PARAMLITERAL:active]
-	[/PARAMLITERAL]
-	[PARAM:content]
-	![](AddExisting_Mac.png)(w:700)
-	[/PARAM]
+		[PARAMLITERAL:category]
+		OS
+		[/PARAMLITERAL]
+		[PARAMLITERAL:id]
+		mac
+		[/PARAMLITERAL]
+		[PARAMLITERAL:active]
+		[/PARAMLITERAL]
+		[PARAM:content]
+		![](AddExisting_Mac.png)(w:700)
+		[/PARAM]
 	[/OBJECT]
 
 ### Creating an Empty Sublevel
@@ -110,31 +110,31 @@ Adding a new sublevel makes it the current Level, so if you want to continue wor
 1. Select a save location and a name for your Level, then click on **Save**. 
 
 	[OBJECT:ToggleButtonContent]
-	[PARAMLITERAL:category]
-	OS
-	[/PARAMLITERAL]
-	[PARAMLITERAL:id]
-	windows
-	[/PARAMLITERAL]
-	[PARAMLITERAL:active]
-	[/PARAMLITERAL]
-	[PARAM:content]
-	![](SaveAs_Windows.png)(w:700)
-	[/PARAM]
+		[PARAMLITERAL:category]
+		OS
+		[/PARAMLITERAL]
+		[PARAMLITERAL:id]
+		windows
+		[/PARAMLITERAL]
+		[PARAMLITERAL:active]
+		[/PARAMLITERAL]
+		[PARAM:content]
+		![](SaveAs_Windows.png)(w:700)
+		[/PARAM]
 	[/OBJECT]
 
 	[OBJECT:ToggleButtonContent]
-	[PARAMLITERAL:category]
-	OS
-	[/PARAMLITERAL]
-	[PARAMLITERAL:id]
-	mac
-	[/PARAMLITERAL]
-	[PARAMLITERAL:active]
-	[/PARAMLITERAL]
-	[PARAM:content]
-	![](SaveAs_Mac.png)(w:700)
-	[/PARAM]
+		[PARAMLITERAL:category]
+		OS
+		[/PARAMLITERAL]
+		[PARAMLITERAL:id]
+		mac
+		[/PARAMLITERAL]
+		[PARAMLITERAL:active]
+		[/PARAMLITERAL]
+		[PARAM:content]
+		![](SaveAs_Mac.png)(w:700)
+		[/PARAM]
 	[/OBJECT]
 
 	The new Level will be added as a sublevel of your current persistent Level, and it will also be made the current Level for working in the Viewport.
@@ -153,31 +153,31 @@ from the original.
 1. Select a save location and a name for your Level, then click on **Save**. 
 
 	[OBJECT:ToggleButtonContent]
-	[PARAMLITERAL:category]
-	OS
-	[/PARAMLITERAL]
-	[PARAMLITERAL:id]
-	windows
-	[/PARAMLITERAL]
-	[PARAMLITERAL:active]
-	[/PARAMLITERAL]
-	[PARAM:content]
-	![](SaveAs_Windows.png)(w:700)
-	[/PARAM]
+		[PARAMLITERAL:category]
+		OS
+		[/PARAMLITERAL]
+		[PARAMLITERAL:id]
+		windows
+		[/PARAMLITERAL]
+		[PARAMLITERAL:active]
+		[/PARAMLITERAL]
+		[PARAM:content]
+		![](SaveAs_Windows.png)(w:700)
+		[/PARAM]
 	[/OBJECT]
 
 	[OBJECT:ToggleButtonContent]
-	[PARAMLITERAL:category]
-	OS
-	[/PARAMLITERAL]
-	[PARAMLITERAL:id]
-	mac
-	[/PARAMLITERAL]
-	[PARAMLITERAL:active]
-	[/PARAMLITERAL]
-	[PARAM:content]
-	![](SaveAs_Mac.png)(w:700)
-	[/PARAM]
+		[PARAMLITERAL:category]
+		OS
+		[/PARAMLITERAL]
+		[PARAMLITERAL:id]
+		mac
+		[/PARAMLITERAL]
+		[PARAMLITERAL:active]
+		[/PARAMLITERAL]
+		[PARAM:content]
+		![](SaveAs_Mac.png)(w:700)
+		[/PARAM]
 	[/OBJECT]
 
 	All your selected Actors will be removed from their original Level and added to the new one, which will be added as a sublevel of your current persistent Level and made the current Level for working in the Viewport. 
@@ -210,62 +210,62 @@ To set Level Streaming Volumes, you need to open the **Level Details** for your
 No additional details are visible for persistent Levels, although there is a dropdown menu you can use to change to another Level. 
 
 [OBJECT:ToggleButtonContent]
-[PARAMLITERAL:category]
-OS
-[/PARAMLITERAL]
-[PARAMLITERAL:id]
-windows
-[/PARAMLITERAL]
-[PARAMLITERAL:active]
-[/PARAMLITERAL]
-[PARAM:content]
-![](LevelDetailsPersistent_Windows.png)(w:700)
-[/PARAM]
+	[PARAMLITERAL:category]
+	OS
+	[/PARAMLITERAL]
+	[PARAMLITERAL:id]
+	windows
+	[/PARAMLITERAL]
+	[PARAMLITERAL:active]
+	[/PARAMLITERAL]
+	[PARAM:content]
+	![](LevelDetailsPersistent_Windows.png)(w:700)
+	[/PARAM]
 [/OBJECT]
 
 [OBJECT:ToggleButtonContent]
-[PARAMLITERAL:category]
-OS
-[/PARAMLITERAL]
-[PARAMLITERAL:id]
-mac
-[/PARAMLITERAL]
-[PARAMLITERAL:active]
-[/PARAMLITERAL]
-[PARAM:content]
-![](LevelDetailsPersistent_Mac.png)(w:700)
-[/PARAM]
+	[PARAMLITERAL:category]
+	OS
+	[/PARAMLITERAL]
+	[PARAMLITERAL:id]
+	mac
+	[/PARAMLITERAL]
+	[PARAMLITERAL:active]
+	[/PARAMLITERAL]
+	[PARAM:content]
+	![](LevelDetailsPersistent_Mac.png)(w:700)
+	[/PARAM]
 [/OBJECT]
 
 For sublevels, you can set the Level's offset **Position** and **Rotation**, as well as the **Streaming Volumes** to use and the debug **Level Color**. Advanced settings like the 
 [minimum time between unload requests](Engine/LevelStreaming/StreamingVolumes#addinghysteresistounloadingrequests) to improve performance are also accessible here.
 
 [OBJECT:ToggleButtonContent]
-[PARAMLITERAL:category]
-OS
-[/PARAMLITERAL]
-[PARAMLITERAL:id]
-windows
-[/PARAMLITERAL]
-[PARAMLITERAL:active]
-[/PARAMLITERAL]
-[PARAM:content]
-![](LevelDetailsStreaming_Windows.png)(w:700)
-[/PARAM]
+	[PARAMLITERAL:category]
+	OS
+	[/PARAMLITERAL]
+	[PARAMLITERAL:id]
+	windows
+	[/PARAMLITERAL]
+	[PARAMLITERAL:active]
+	[/PARAMLITERAL]
+	[PARAM:content]
+	![](LevelDetailsStreaming_Windows.png)(w:700)
+	[/PARAM]
 [/OBJECT]
 
 [OBJECT:ToggleButtonContent]
-[PARAMLITERAL:category]
-OS
-[/PARAMLITERAL]
-[PARAMLITERAL:id]
-mac
-[/PARAMLITERAL]
-[PARAMLITERAL:active]
-[/PARAMLITERAL]
-[PARAM:content]
-![](LevelDetailsStreaming_Mac.png)(w:700)
-[/PARAM]
+	[PARAMLITERAL:category]
+	OS
+	[/PARAMLITERAL]
+	[PARAMLITERAL:id]
+	mac
+	[/PARAMLITERAL]
+	[PARAMLITERAL:active]
+	[/PARAMLITERAL]
+	[PARAM:content]
+	![](LevelDetailsStreaming_Mac.png)(w:700)
+	[/PARAM]
 [/OBJECT]
 
 ## Visualizing Sublevels
diff --git a/Engine/Documentation/Source/Engine/Niagara/CascadeToNiagara/CascadeToNiagara.INT.udn b/Engine/Documentation/Source/Engine/Niagara/CascadeToNiagara/CascadeToNiagara.INT.udn
--- a/Engine/Documentation/Source/Engine/Niagara/CascadeToNiagara/CascadeToNiagara.INT.udn
+++ b/Engine/Documentation/Source/Engine/Niagara/CascadeToNiagara/CascadeToNiagara.INT.udn
@@ -1,65 +1,67 @@
-Availability: Docs
-Crumbs: %ROOT%
+Availability: Public
+Crumbs: 
 Title: Cascade To Niagara
 Description: Showing the differences and similarities between Cascade and Niagara
-Type: Overview
-SkillLevel: Beginner
-Version: 4.18
+Version: 4.20
 Parent: Engine/Niagara
-Order: 1
-Tags: Cascade
+Order: 0
 Tags: Getting Started
-topic-image:NiagaraToCascadeTopicImage.png
+Tags: Cascade
+Tags: Early Access
 
-[TOC(start:1 end:3)]
+[TOC(start:2 end:3)]
 
 While both **Cascade** and **Niagara** can be used to make visual effects (VFX) inside of Unreal Engine 4 (UE4), the manner in which you use Cascade to create and adjust visual effects is vastly different from how you use Niagara to achieve the same ends. In the following document, we will take a look at how to transition from Cascade to Niagara to make all of your visual effects. The reason for creating this new tool is to put more power into the hands of the artist to create additional functionality without the help of a programmer. We take some inspiration from movie FX tools and try to make the system as adaptable and flexible as possible while keeping usability high and accommodating the 90% use cases.
 
 ## Niagara Overview
 
-Niagara is Unreal Engine's next-generation visual effects system. It's a complete replacement for Cascade, with which it will coexist until Niagara functionally is deemed to fully supersede that of Cascade.
+Niagara is Unreal Engine's next-generation visual effects system. It's a complete replacement for Cascade, with which it will coexist until Niagara functionally is deemed to fully supersede that of Cascade. In the image below you can see a high level breakdown that covers how Niagara works with UE4.
+
+![CascadeVS Niagara.png](CascadeVS20Niagara.png "CascadeVS Niagara.png")
 
-## Why is Niagara Better Than Cascade?
+## How is Niagara Different Than Cascade?
 
 The fundamental power of Niagara lies predominantly in the ability to share arbitrary particle data in whatever way the user sees fit. This helps to enable users to script their own behaviors rather than requiring code intervention to create new hard-coded features. In Cascade, particle simulations are controlled by hard-coded modules, each providing a specific set of functionality to modify particle attributes. In contrast, Niagara simulation is controlled by node graphs that artists can create and edit.
 
 [OBJECT:ComparisonSlider]
-    [PARAM:before]
-    ![Cascade](Cascade.png)(w:1000)
-    [/PARAM]
-    [PARAM:after]
-    ![Niagara](Ninagra.png)
-    [/PARAM]
+	[PARAM:before]
+	![Cascade](Cascade.png)(w:800)
+	[/PARAM]
+	[PARAM:after]
+	![Niagara](Ninagra.png)(w:800)
+	[/PARAM]
 [/OBJECT]
 
 One of the other reasons Niagara is superior to Cascade is that Niagara's node graphs are compiled to bytecode that runs in a very low-overhead interpreter, resulting in very high performance even on CPU simulations. GPU simulations are planned, and switching between the two is expected to be transparent - with a few caveats, any emitter should be able to be simulated on CPU or GPU, based on user choice.
 
-## Cascade & Niagara UI 
+## Cascade & Niagara UI
 
 When you first open up Niagara, the User Interface (UI) will look completely different than Cascade. However, if you take a close look, you will begin to notice that the Niagara interface does share some of the following similarities to Cascade:
 
-* Realtime 3D viewport to visualize effects while they are being created.
-
-* Curves Editor for adjusting R, G, B and Alpha values.
-
-* Details Panel for adjusting various properties.
+*   Realtime 3D viewport to visualize effects while they are being created.
+    
+*   Curves Editor for adjusting R, G, B and Alpha values.
+    
+*   Details Panel for adjusting various properties.
+    
 
 [OBJECT:ComparisonSlider]
-    [PARAM:before]
-    ![Cascade](CTN_Cas_VFX.png)(w:1000)
-    [/PARAM]
-    [PARAM:after]
-    ![Niagara](CTN_Nia_VFX.png)
-    [/PARAM]
+	[PARAM:before]
+	![Cascade](CTN_Cas_VFX.png)(w:800)
+	[/PARAM]
+	[PARAM:after]
+	![Niagara](CTN_Nia_VFX.png)(w:800) 
+	[/PARAM]
 [/OBJECT]
 
 ## Niagara Emitters & Systems
 
-One of the biggest differences between Cascade and Niagara is how you build emitters. When using Cascade, emitters are built as a self-contained unit. While this made the emitters easy to place and spawn, it also made it hard to share behavior between different emitter or make updates as everything had to be done manually.  Niagara changes all of this by giving visual effect artist the ability to control, add, or modify any part of a visual effects behavior right down to the pixel level. To accomplish this, Niagara particle systems have been broken up into the following two parts:
+One of the biggest differences between Cascade and Niagara is how you build emitters. When using Cascade, emitters are built as a self-contained unit. While this made the emitters easy to place and spawn, it also made it hard to share behavior between different emitter or make updates as everything had to be done manually. Niagara changes all of this by giving visual effect artist the ability to control, add, or modify any part of a visual effects behavior right down to the pixel level. To accomplish this, Niagara particle systems have been broken up into the following two parts:
 
-* **Niagara Emitters:** This is a single visual effect that will be combined with other Niagara Emitters to make a Niagara System. If we are making a fire effect, the Niagara Emitters would be each of the different elements that make up the fire effect like the smoke, fire, distortion, and embers.
- 
+*   **Niagara Emitters:** This is a single visual effect that will be combined with other Niagara Emitters to make a Niagara System. If we are making a fire effect, the Niagara Emitters would be each of the different elements that make up the fire effect like the smoke, fire, distortion, and embers.
+    
     ![](CTN_NiagaraEmitters.png)
-* **Niagara Systems:** This contains everything needed to make up a single effect. Continuing with our fire example, this would be the completed fire effect that contains the Niagara Emitters for the smoke, fire, distortion, and embers.
- 
-    ![](CTN_NiagaraSystem.png)
+    
+*   **Niagara Systems:** This contains everything needed to make up a single effect. Continuing with our fire example, this would be the completed fire effect that contains the Niagara Emitters for the smoke, fire, distortion, and embers.
+    
+    ![](CTN_NiagaraSystem.png)
\ No newline at end of file
diff --git a/Engine/Documentation/Source/Engine/Niagara/HowTo/EnableNiagara/Enable Niagara.INT.udn b/Engine/Documentation/Source/Engine/Niagara/HowTo/EnableNiagara/Enable Niagara.INT.udn
--- a/Engine/Documentation/Source/Engine/Niagara/HowTo/EnableNiagara/Enable Niagara.INT.udn	
+++ b/Engine/Documentation/Source/Engine/Niagara/HowTo/EnableNiagara/Enable Niagara.INT.udn	
@@ -20,7 +20,7 @@ Before you can begin to create visual effects using Niagara you will need to ena
 
 	![](HT_EnableNP_00.png)
 
-1. In the Plugins menu, locate the **FX** section and then under Niagara, click on the **Enabled** checkmark box.
+1. In the Plugins menu, locate the **FX** section and then under Niagara, and **Niagara Extras**, click to check the box for **Enabled.**
 
 	![](HT_EnableNP_02.png)
 
diff --git a/Engine/Documentation/Source/Engine/Niagara/HowTo/GPUParticles/GPUParticles.INT.udn b/Engine/Documentation/Source/Engine/Niagara/HowTo/GPUParticles/GPUParticles.INT.udn
--- a/Engine/Documentation/Source/Engine/Niagara/HowTo/GPUParticles/GPUParticles.INT.udn
+++ b/Engine/Documentation/Source/Engine/Niagara/HowTo/GPUParticles/GPUParticles.INT.udn
@@ -1,49 +1,70 @@
-Availability: Docs
-Crumbs: %ROOT%, 
-Title:Create a GPU based effects
-Description:Describes how you can create use your GPU to spawn millions of particles.
-Type: how-to
-SkillLevel:Beginner
-Version: 4.17
-parent:Engine/Niagara/HowTo
-order:5
-tags:Niagara
-topic-image:
-
-## Steps
-
-## End Result
-
-
-
-
-
-
-
-
-
-
-
-
-
-
-
-
-
-
-
-
-
-
-
-
-
-
-
+Availability: Public
+Crumbs: 
+Title: Create a GPU Particle Effect
+Description: Describes how you can use your GPU to spawn millions of particles.
+Type: how-to 
+Version: 4.20
+Parent: Engine/Niagara/HowTo
+Order: 5
+Tags: Niagara
+Tags: Early Access
+prereqs: Engine/Niagara/HowTo/EnableNiagara
 
 
+For some effects, you may need to spawn tens of thousands of particles. However, using the standard CPU to generate this many particles can cause the game to slow down. In the following How-To we will demonstrate how to spawn particles using the GPU instead of the CPU.
 
+![](GPU_Banner_01-1.png "GPU_Banner_01-1.png")
 
+## Steps
 
+1.  First, we will need to create a particle effect. Right-click in the **Content Browser**, and from the displayed menu select **FX > Niagara Emitter** to create a new Niagara Emitter.  
+    ![](SE_01.png)
+1.  Name the newly created Niagara Emitter **FX_GPUParticle** and then double-click it to open it.  
+    ![](CreateGPU_Steps_Step2-CreateEmitter.png)
+1.  With the Niagara Editor now open, go to the **Selected Emitter** panel and locate the **Render** section.
+1.  Expand the **Particle Spawn** section. Locate the **Sphere Location** section.
+1.  In the **Sphere Location** section, set the **Sphere Radius** to a value of **100**.  
+    ![CreateGPUParticle_Steps_Step05-SphereRadius.png](CreateGPUParticle_Steps_Step05-SphereRadius.png)
+1.  Locate the **Set Variables** section and then set the following parameters with the following values: ![](CreateGPUParticle_Steps_Step06-SetVariables.png)
+    
+    | Particle Module Name | Value |
+    | --- | --- |
+    | **Particles.Lifetime** | 5 |
+    | **Particles.SpriteSize** | X: 5, Y: 5 |
+    
+1.  Locate the **Emitter Properties** section and click the dropdown next to **Sim Target**. From the displayed list, select the **GPUCompute Sim** option.  
+    ![](CreateGPUParticle_Steps_Step07-ChangeToGPU.png)
+[REGION:note]
+This is where you are changing from the CPU to the GPU. You should see an increase in the number of particles generated.
+[/REGION]9.  Next, find the **Spawn Rate** section and set the **SpawnRate** to a value of **5000**.  
+![](CreateGPUParticle_Steps_Step08-SpawnRate.png)
+1.  Click the **Compile**, **Apply** and **Save** buttons to compile and save your work. Then close the **FX_GPUParticle** emitter.  
+    ![](ME_09.png)
+1.  Now right-click in the **Content Browser**, and select **FX > Niagara System**, giving the new Niagara System the name **GPUParticle**.  
+    ![](SE_16.png)
+1.  Double-click the **GPUParticle** Niagara System to open it. Then drag the **FX_GPUParticle** Niagara Emitter to the **Timeline**.  
+    ![](CreateGPUParticle_Steps_Step11-AddEmitterToSystem.png)
+1.  Click the **Compile** and **Save** buttons to compile and save your work, and then close the Particle Light Niagara System. In the **Content Browser**, drag the Mesh Niagara System into the level viewport.  
+    ![](CreateGPUParticle_Steps_Step12-AddSystemToLevel.png)
 
+## End Result
 
+When all steps have been completed, you will have a particle effect that is larger and more complex than the same particle effect using the CPU. An example is shown in the video below.
+
+[OBJECT:EmbeddedVideo]
+    [PARAMLITERAL:videoid]
+    zz3dFQzNQY8
+    [/PARAMLITERAL]
+    [PARAMLITERAL:width]
+    78
+    [/PARAMLITERAL]
+    [PARAMLITERAL:height]
+    56.25
+    [/PARAMLITERAL]
+    [PARAMLITERAL:units]
+    %
+    [/PARAMLITERAL]
+    [PARAMLITERAL:type]
+    simple
+    [/PARAMLITERAL]
+[/OBJECT]
diff --git a/Engine/Documentation/Source/Engine/Niagara/HowTo/MeshEffect/MeshEffect.INT.udn b/Engine/Documentation/Source/Engine/Niagara/HowTo/MeshEffect/MeshEffect.INT.udn
--- a/Engine/Documentation/Source/Engine/Niagara/HowTo/MeshEffect/MeshEffect.INT.udn
+++ b/Engine/Documentation/Source/Engine/Niagara/HowTo/MeshEffect/MeshEffect.INT.udn
@@ -1,90 +1,79 @@
-Availability: Docs
-Crumbs: %ROOT%, 
-Title:Create a Mesh-Based Effect with Niagara
-Description:Describes how you can create a static mesh based visual effect using Niagara.
-Type: how-to
-SkillLevel:Beginner
-Version: 4.17
-parent:Engine/Niagara/HowTo
-order:3
-tags:Niagara
-topic-image:
-
-Having the ability to spawn a Static Mesh instead of a camera facing sprite can help with the believability and realism of the effects you are creating. In the following How - To we will take a look at setting up a Niagara emitter to spawn Static Meshes instead of Sprites.
-
-![](HeroImage_00.png)
+Availability: Public
+Crumbs: 
+Title: Create a Mesh Particle Effect in Niagara
+Description: Describes how you can create a static mesh particle effect using Niagara.
+Type:how-to
+Version: 4.20
+Parent: Engine/Niagara/HowTo
+Order: 3
+Tags: Niagara
+Tags: Early Access
+prereqs Engine/Niagara/HowTo/EnableNiagara
+
+Having the ability to spawn a Static Mesh instead of a camera facing sprite can help with the believability and realism of the effects you are creating. In the following How-To, we will take a look at setting up a Niagara emitter to spawn Static Meshes instead of Sprites.
+
+![](Mesh_Banner_01.png "Mesh_Banner_01.png")
+
+[REGION:note]
+**Prerequisite Step:** This How-To uses the **Shape_Sphere** Static Mesh, which is found with the **Starter Content**. If you have not done so already, make sure that this Static Mesh or the Starter Content has been added to your project.
+[/REGION]
 
 ## Steps
 
-[region:note]
-This How-To will make use of the **Shape_Sphere** Static Mesh that can be found with the **Starter Content**. If you have not done so already, make sure that this Static Mesh or the Starter Content has been added to your project.
-[/region]
-
-1. First, we will need to create a **Niagara Emitter** by right-clicking in the **Content Browser** and from the displayed menu go to **FX** and then click on the **Niagara Emitter**. 
-
-    ![](Engine\Niagara\HowTo\SpriteEffect\SE_01.png)
-
-1. Name the newly created Niagara Emitter, **FX_Mesh** and then double click on it to open it up.
-    
+1.  First, we will need to create a Niagara Emitter. Right-click in the **Content Browser**, and from the displayed menu select **FX > Niagara Emitter**.   
+    ![](SE_01.png)
+1.  Name the newly created Niagara Emitter **FX_Mesh** and then double-click it to open it.  
     ![](ME_02.png)
-
-1. Go to the **Render** section and press the **Trashcan** icon to remove the **Sprite Rendering** component. 
-
-    ![](ME_03.png)
-
-1. Press the **Plus** sign icon and from the displayed list, select the **Niagara Mesh Renderer Properties** option.
-
-	 ![](ME_12.png)
-
-1. Expand the **Mesh Rendering** section and then in the **Particle Mesh**, input the **Spahe_Sphere** Static Mesh.
-
-    ![](ME_04.png)
-
-1. Next, go to the **Particle Spawn** section and in the **Add Velocity** section, set the **Minimum Z** to **400** and the **Maximum Z** to **980**. 
-
-    ![](ME_05.png)
-
-1. Down in the **Sphere Location**, set the **Sphere Radius** to **32**.
-
-    ![](ME_06.png)
-
-1. Then in the **Set Particles.Lifetime** section set the **Minimum** lifetime to **1.25** and the **Maximum** lifetime to **5.0**. 
-
-    ![](ME_07.png)
-
-1. Next, to the far right of the **Set Particles.SpriteSize** and **Set Particles.SpriteRotation** there will be an icon of a **Trash Can**. Press this icon to remove both Set Particles.SpriteSize and Set Particles.SpriteRotation from the Particle Spawn as it is no longer needed.
-
-    ![](ME_08.png)
-
-1. Then press the **Compile**, **Apply** and **Save** buttons to apply and save the changes and then close the **FX_Mesh** emitter.
-
-    ![](ME_09.png)
-
-1. Now over in the Content Browser, right - click and from the displayed list, go to **FX** and then select the **Niagara System** option, giving it a name of **Mesh**.
-
-    ![](Engine\Niagara\HowTo\SpriteEffect\SE_16.png)
-
-1. Double - click on the Mesh Niagara System to open it up and when open, drag the **FX_Mesh** from the Content Browser to the **Timeline**.
-
-    ![](ME_10.png)
-
-1. Press the **Compile** and **Save** buttons and then close the Mesh Niagara System down and then over in the Content Browser, drag the Mesh Niagara System into the level viewport.
-
-    [REGION:lightbox]
-    [![](ME_11.png)(w:834)](ME_11.png)
-    [/REGION]
-
-    [REGION:caption]
-    Click for full image.
+1.  With the Niagara Editor now open, go to the **Selected Emitter** panel and locate the **Render** section.  To remove the **Niagara ****Sprite Rendering** component, click the **Trash** icon.  
+    ![](CreateMeshBasedEffect_Steps_Step03-RemoveSpriteRender.png)
+1.  Click the **Plus** sign icon (**+**) and from the displayed list, select the **Niagara Mesh Renderer Properties** option. This inserts the **NiagaraMeshRendererProperties** section.  
+    ![](CreateMeshBasedEffect_Steps_Step04-AddMeshRender.png)
+1.  Expand the **Mesh Rendering** section. Click the dropdown next to **Particle Mesh**, and select the **Shape_Sphere** static mesh.   
+    ![](CreateMeshBasedEffect_Steps_Step05-AddShapeSphere.png)
+1.  Collapse the Render section, and expand the **Particle Spawn** section.
+1.  Locate the **Add Velocity** section.
+1.  Set Velocity **Minimum** and **Maximum** fields to the following values.   
+    ![](CreateMeshBasedEffect_Steps_Step08-VelocityMinMax.png)   
+	| Property Name | Value |
+	| --- | --- |
+	| **Minimum** | X:0, Y:0, Z: 400 |
+	| **Maximum** | X: 1, Y: 1, Z: 980 |
+    
+1.  In the **Sphere Location** section, set the **Sphere Radius** to **100**.  
+    ![](Mesh_Step9_01.png "Mesh_Step9_01.png")  
+    
+1.  Locate the **Set Variables** section. Expand**Particles.Lifetime** to display the **Minimum** and **Maximum** fields.
+1.  Set the **Minimum** and **Maximum** fields to the following values.  
+    ![](CreateMeshBasedEffect_Steps_Step11-ParticlesLifetimeMinMax.png)    
+	| Property Name | Value |
+	| --- | --- |
+	| **Minimum** | 2.0 |
+	| **Maximum** | 5.0 |
+    
+1.  Next, click the small white triangle to the right of the **Particles.SpriteSize** and  **Particles.SpriteRotation** properties, and select **Remove**.  
+    ![](CreateMeshBasedEffect_Steps_Step12-RemoveSpriteVariables.png)
+1.  Locate the **Acceleration Force** section.
+1.  Set the **Z** value to **-500**.  
+    ![](CreateMeshBasedEffect_Steps_Step14-AccelerationForce.png)
+    [REGION:note]
+    This negative value causes the mesh particles to go up briefly, then fall down. To have the mesh particles go up without falling, use a positive Z value.
     [/REGION]
+1.  Click the **Compile**, **Apply** and **Save** buttons to apply and save the changes, and then close the **FX_Mesh** emitter.   
+    ![](ME_09.png)
+1.  Now over in the **Content Browser**, right-click and select **FX > Niagara System**. Name the new system **Mesh**.   
+    ![](SE_16.png)
+1.  Double-click the **Mesh** Niagara System to open it, then drag the **FX_Mesh** Emitter from the Content Browser to the **Timeline**.   
+    ![](CreateMeshBasedEffect_Steps_Step17_EmittertoSystem.png)
+1.  Click the **Compile** and **Save** buttons to compile and save your work, and then close the Mesh Niagara System. In the **Content Browser**, drag the **Mesh** Niagara System into the level viewport.  
+    ![](CreateMeshParticle_Steps_Step18_AddToLevel.png)
 
 ## End Result
 
-When completed you will now have a Niagara System that will play your Niagara Emitter when the Niagara System is placed in a level like in the following video.
+When complete, you will have a Niagara System playing the Niagara Emitter when you place the Niagara System in a level (like in the following video).
 
 [OBJECT:EmbeddedVideo]
     [PARAMLITERAL:videoid]
-    SjCLcU67gcQ
+    cFgfGav5UjE
     [/PARAMLITERAL]
     [PARAMLITERAL:width]
     100
@@ -98,6 +87,4 @@ When completed you will now have a Niagara System that will play your Niagara Em
     [PARAMLITERAL:type]
     simple
     [/PARAMLITERAL]
-[/OBJECT]
-
-
+[/OBJECT]
\ No newline at end of file
diff --git a/Engine/Documentation/Source/Engine/Niagara/HowTo/ParticleLights/Particle Lights.INT.udn b/Engine/Documentation/Source/Engine/Niagara/HowTo/ParticleLights/Particle Lights.INT.udn
--- a/Engine/Documentation/Source/Engine/Niagara/HowTo/ParticleLights/Particle Lights.INT.udn	
+++ b/Engine/Documentation/Source/Engine/Niagara/HowTo/ParticleLights/Particle Lights.INT.udn	
@@ -1,75 +1,61 @@
-Availability: Docs
-Crumbs: %ROOT%, 
-Title:Create a Particle Light
-Description:Describes how you can create light within your particle system.
-Type: how-to
-SkillLevel:Beginner
-Version: 4.17
-parent:Engine/Niagara/HowTo
-order:5
-tags:Niagara
-topic-image:
-
-Giving particles the ability to illuminate the world around them can add an extra layer of realism to any of your project's visual effects. In the following How - To we will take a look at setting up a Niagara emitter so that particles and lights are spawned at the same time.
-
-![](PL_Hero_Image.png)
+Availability: Public
+Crumbs: 
+Title: Create a Particle Light
+Description: Describes how you can create light within your particle system.
+Type:how-to
+Version: 4.20
+Parent:Engine/Niagara/HowTo 
+Order:5 
+Tags: Niagara
+Tags:Early Access
+prereqs:Engine/Niagara/HowTo/EnableNiagara
+    
+
+Giving particles the ability to illuminate the world around them can add an extra layer of realism to any of your project's visual effects. In the following How-To, we will take a look at setting up a Niagara emitter so that particles and lights are spawned at the same time.
+
+![](Light_Banner_01.png "Light_Banner_01.png")
 
 ## Steps
 
-1. First, we will need to create a **Niagara Emitter** by right-clicking in the **Content Browser**, then from the displayed menu, go to **FX** and click **Niagara Emitter**. 
-
-    ![](Engine\Niagara\HowTo\SpriteEffect\SE_01.png)
-
-1. Name the newly created Niagara Emitter, **FX_Light** and then double-click on it to open it up.
-
+1.  First, we will need to create a **Niagara Emitter** by right-clicking in the **Content Browser**, then from the displayed menu, go to **FX** and click **Niagara Emitter**.  
+    ![](SE_01.png)
+2.  Name the newly created Niagara Emitter, **FX_Light** and then double-click on it to open it up.  
     ![](PL_01.png)
-
-1. Next in the **Render** section of the FX_Light Niagara Emitter, press the **Plus** sign icon and from the displayed list, select the **Niagara Light Renderer Properties** option.
-
-    ![](PL_02.png)
-
-1. In the Niagara Light Renderer Properties, set the **Radius Scale** and **Color Add** to a value of **10**.
-
-    ![](PL_03.png)
-
-1. Then press the **Compile**, **Apply** and **Save** buttons to apply and save the changes and then close the **FX_Mesh** emitter.
-
-    ![](Engine\Niagara\HowTo\MeshEffect\ME_09.png)
-
-1. Now over in the **Content Browser**, right-click and from the displayed list, go to **FX** and then select the **Niagara System** option, giving it a name of **ParticleLight**.
-
-    ![](Engine\Niagara\HowTo\SpriteEffect\SE_16.png)
-
-1. Double-click on the Particle Light Niagara System to open it up and then drag the **FX_Light** Niagara Emitter to the **Timeline** of the Particle Light Niagara System.
-
+3.  With the Niagara Editor now open, go to the **Selected Emitter** panel and locate the **Render** section.  
+    
+4.  Click the **Plus** sign icon (**+**) next to **Render**, and from the displayed list, select the **Niagara Light Renderer Properties** option.  
+    ![](CreateParticleLight_Steps_Step4-AddLightRenderer.png)
+5.  Locate the new **NiagaraLightRendererProperties** section at the bottom of the Render module. Set the following properties to the following values.  
+    ![ParticleLights/CreateParticleLight_Steps_Step5-RadiusScaleColorAdd.png](CreateParticleLight_Steps_Step5-RadiusScaleColorAdd.png) 
+	| Property Name | Value |
+	| --- | --- |
+	| **Radius Scale** | 1 |
+	| **Color Add X** | 10 |
+    
+6.  Click the **Compile**, **Apply** and **Save** buttons to compile and save your work. Then close the **FX_Light** Niagara Emitter.  
+    ![](ME_09.png)
+7.  Now right-click in the **Content Browser**, and select **FX > Niagara System**. Name the new Niagara System **ParticleLight**.  
+    ![](SE_16.png)
+8.  Double-click on the **ParticleLight** **Niagara** System to open it, and then drag the **FX_Light** Niagara Emitter to the **Timeline** of the Particle Light Niagara System.  
     [REGION:lightbox]
     [![](PL_04.png)(w:317)](PL_04.png)
     [/REGION]
-
-    [REGION:caption]
-    Click for full image.
-    [/REGION]
-
-1. Press the **Compile** and **Save** buttons and then close the Particle Light Niagara System down and then over in the Content Browser, drag the Mesh Niagara System into the level viewport.
-
-    [REGION:lightbox]
-    [![](PL_05.png)(w:779)](PL_05.png)
-    [/REGION]
-
     [REGION:caption]
-    Click for full image.
+    *Click for full image.*
     [/REGION]
+9.  Click the **Compile** and **Save** buttons to compile and save your work, and then close the Particle Light **Niagara System**. In the **Content Browser**, drag the **Mesh Niagara System** into the level viewport.  
+    ![](CreateParticleLight_Steps_Step9_AddToLevel.png)
 
 ## End Result
 
-When completed you will now have a VFX that will cast light into the world when the particles are spawned.
+When completed you will now have a VFX that will cast light into the world when the particles are spawned. You can see an example in the video below.
 
 [OBJECT:EmbeddedVideo]
     [PARAMLITERAL:videoid]
-    7D0edbLt2oE
+    cRcjPhPHmSM
     [/PARAMLITERAL]
     [PARAMLITERAL:width]
-    100
+    78
     [/PARAMLITERAL]
     [PARAMLITERAL:height]
     56.25
@@ -80,11 +66,4 @@ When completed you will now have a VFX that will cast light into the world when
     [PARAMLITERAL:type]
     simple
     [/PARAMLITERAL]
-[/OBJECT]
-
-
-
-
-
-
-
+[/OBJECT]
\ No newline at end of file
diff --git a/Engine/Documentation/Source/Engine/Niagara/HowTo/RibbonEffect/RibbonEffect.INT.udn b/Engine/Documentation/Source/Engine/Niagara/HowTo/RibbonEffect/RibbonEffect.INT.udn
--- a/Engine/Documentation/Source/Engine/Niagara/HowTo/RibbonEffect/RibbonEffect.INT.udn
+++ b/Engine/Documentation/Source/Engine/Niagara/HowTo/RibbonEffect/RibbonEffect.INT.udn
@@ -1,143 +1,113 @@
-Availability: Docs
-Crumbs: %ROOT%, 
-Title:Create a Ribbon Based Effect with Niagara
-Description:Describes how you can create a ribbon-based visual effect using Niagara.
+Availability: Public
+Crumbs: 
+Title: Create a Ribbon Effect in Niagara
+Description: Describes how you can create a ribbon visual effect using Niagara.
 Type: how-to
-SkillLevel:Beginner
-Version: 4.17
-parent:Engine/Niagara/HowTo
-order:5
-tags:Niagara
-topic-image:
+Version: 4.20
+Parent: Engine/Niagara/HowTo
+Order: 5
+Tags: Niagara
+Tags: Early Access
+prereqs:Engine/Niagara/HowTo/EnableNiagara
+    
 
-Trying to simulate natural phenomena like smoke or vapor trails can be quite challenging to accomplish with sprite- or mesh-based particles. A better approach to simulating this type of object would be to use a **Ribbon** emitter, and in the following how-to, we will take a look at how to set up a Niagara Emitter to emit a continuous ribbon-style particle effect into the world.
+Simulating natural phenomena is challenging, especially when using sprite or mesh-based particles to simulate smoke or vapor trails. **Ribbon Emitters** are an excellent solution for simulating these objects, and in the following How-To, you will learn how to set up a Niagara Emitter to emit a continuous ribbon-style particle effect into the world.
 
-![](HeroImage_00.png)
+![](Ribbon_Banner_01.png "Ribbon_Banner_01.png")
 
 ## Steps
 
-1. First, we will need to create a new Material to color our Ribbon Emitter. To do this, right-click in the **Content Browser**, and select the **Material** option from the displayed list, naming the Material **MAT_Ribbon**.
-
+1.  First, we will need to create a new Material to color our Ribbon Emitter. To do this, right-click in the **Content Browser**, and under **Create a Basic Asset**, select the **Material** option from the displayed list. Name the Material **M_Ribbon**.  
     ![](RE_01.png)
-
-1. Double-click on the MAT_Ribbon Material to open it up and then, while holding down the **3** key, left-click in the Material graph to add a **Constant 3 Vector** Material Expression node.
-
-    [REGION:lightbox]
-    [![](RE_02.png)(w:639)](RE_02.png)
-    [/REGION]
-
-    [REGION:caption]
-    Click for full image.
-    [/REGION]
-
-1. Connect the **Constant 3 Vector** to the **Base Color** and **Emissive Color** input and then set the color to red.
-
-    [REGION:lightbox]
-    [![](RE_03.png)(w:639)](RE_03.png)
-    [/REGION]
-
-    [REGION:caption]
-    Click for full image.
-    [/REGION]
-
-1. In the **Details** panel, under **Usage**, make sure that **Used with Niagara Ribbons** is checked and then press the **Compile** and **Save** buttons to save your work.
-
-    [REGION:lightbox]
-    [![](RE_04.png)(w:639)](RE_04.png)
-    [/REGION]
-
-    [REGION:caption]
-    Click for full image.
-    [/REGION]
-
-1. Now, right-click in the **Content Browser** and select **FX** from the context menu. Next, click on the **Niagara Emitter** entry to create a new Niagara Emitter.
-
-    ![](Engine\Niagara\HowTo\SpriteEffect\SE_01.png)
-
-1. Name the newly-created Niagara Emitter **FX_Ribbon**, and then double-click on it to open it up.
-
+2.  Double-click the **M_Ribbon Material** to open it. Press and hold the **3** key, and left-click in the **Material** grid to add a **Constant 3 Vector Material** Expression node.  
+    
+3.  Hover the cursor over the small circle on the right side of the node, then click and drag to connect the **Constant 3 Vector** to the **Base Color** and **Emissive Color** inputs in the M_Ribbon box on the right.   
+    ![](CreateRibbonBasedEffect_Steps_Step03-ConnectNodes.png)
+4.  Double-click the **Constant 3 Vector** to open the color picker. Set the **R** value to 1, then click **OK**. This turns the material red.  
+    ![](CreateRibbonBasedEffect_Steps_Step04-PickColor.png)
+5.  Click in the grid to take focus off the Constant 3 Vector. This changes the **Details** panel display on the left. In the **Details** panel, scroll down to locate **Usage**, and expand it. Check the box for **Used with Niagara Ribbons**.  
+    ![](CreateRibbonBasedEffect_Steps_Step05-DetailsUsage.png)
+6.  Click the **Apply** and **Save** buttons to apply changes and save your work.
+7.  Now, right-click in the **Content Browser**, and select **FX > Niagara Emitter** to create a new Niagara Emitter.   
+    ![](SE_01.png)
+8.  Name the newly-created Niagara Emitter **FX_Ribbon**, and then double-click it to open it.  
     ![](RE_05.png)
+9.  "Locate the **Render** section of the **FX_Ribbon** Emitter, and click the **Trash** icon to delete the **Sprite Rendering** module.  
+    ![CreateRibbonBasedEffect_Steps_Step09-RemoveSpriteRender.png?](CreateRibbonBasedEffect_Steps_Step09-RemoveSpriteRender.png)
+10. Press the **Plus** sign icon **(+) **and from the displayed list, select the **Niagara Ribbon Renderer** Properties.  
+    ![](CreateRibbonBasedEffect_Steps_Step10-AddRibbonRender.png)
+11. Next, click the small white triangle to expand the **Ribbon Rendering** section. In the **Material** section, click the dropdown to select the **M_Ribbon** Material that was created in steps one through five.   
+    ![](CreateRibbonBasedEffect_Steps_Step11-AddMaterial.jpg)
+[REGION:note]
+The value in the **SpawnRate** module (found in the **Emitter Update** section) is set to **20.0** by default. The **Spawn Rate** affects the roundness or angularity of the resulting ribbon effect. Lower values create a more angular or polygonal effect shape. Higher values create a more round effect shape.
+[/REGION]
+
+[OBJECT:ComparisonSlider]
+    [PARAM:before]
+    ![Low Spawn Rate value](CreateRibbonBasedEffect_Steps_Step11a-SpawnRateLow.png) 
+    [/PARAM]
+    [PARAM:after]
+    ![High Spawn Rate value](CreateRibbonBasedEffect_Steps_Step11b-SpawnRateHigh.png) 
+    [/PARAM]
+[/OBJECT]
 
-1. Go to the **Render** section of the FX_Ribbon Emitter and press the **Trashcan** icon to delete the **Sprite Rendering**  rendering module.  
+1.  Locate and expand the **Particle Spawn** section.
+1.  Delete the following modules by clicking the **Trash** icon.  
     
-     ![](RE_16.png)
-
-1. Press the **Plus** sign icon and from the displayed list, select the **Niagara Ribbon Renderer Properties**.
-
-    ![](RE_06.png)
-
-1. Next, click on the small white triangle to expand the **Ribbon Rendering** section. In the **Material** section, input the MAT_Ribbon Material that was created.
-
-    ![](RE_11.png)
-
-1. Now, over in the **Emitter Update** section, go to the **Spawn Rate** section and set the **Spawn Rate** to a value of **20**.
-
-    ![](RE_07.png)
-
-1. In the **Particle Spawn** section, delete the following modules by pressing the **Trash Can** icon:
-
-    ![](RE_08.png)
-
-    * **Add Velocity**
-    * **Sphere Location**
-    * **Set Particle Sprite Size**
-    * **Set Particle Sprite Rotation**
-
-1. When those entries have been deleted, press the **Plus** sign icon at the bottom of the Particle Spawn and from the displayed list add the following modules:
-
-    ![](RE_09.png)
-
-    * **Set Particle Ribbon Width**
-    * **Set Particles Velocity**
-    * **Point Spiral Location** 
-
-1. With the needed modules now added to the Ribbon effect, input the following values into the following modules:
-
-    ![](RE_10.png)
-
-    |Module Name|Value|
-    |----|---|
-    |**Set Particle Ribbon Width**| 32|
-    |**Set Particles Velocity**| X:1 y:0 Z:1|
-    |**Point Spiral Location** | Speed: 500, RotSpeed: 1|
-    |**Set Particle Lifetime**| 5|
-
-1. Now, over in the **Particle Update** section, set the **Z** value of the  **Acceleration Force** to **-100**.
-
-    ![](RE_12.png)
-
-1. Once that has all been completed, press the **Compile**, **Apply** and **Save** buttons to compile and save our work.
-
-    ![](RE_13.png)
-
-1. In the Content Browser, right-click. From the displayed list, go to **FX** and then select the **Niagara System** option, giving it a name of **FX_Ribbon**.
-
-    ![](Engine\Niagara\HowTo\SpriteEffect\SE_16.png)
-
-1. Double-click on the Ribbon Niagara System to open it up and when open, drag the **FX_Ribbon** from the Content Browser to the **Timeline**.
-
-    ![](RE_14.png)
-
-1. Press the **Compile** and **Save** buttons, and then close the Mesh Niagara System down. In the Content Browser, drag the Mesh Niagara System into the level viewport.
-
-    [REGION:lightbox]
-    [![](RE_15.png)(w:834)](RE_15.png)
-    [/REGION]
-
-    [REGION:caption]
-    Click for full image.
+    *   **Add Velocity**
+    *   **Sphere Location**
+    
+    ![](CreateRibbonBasedEffect_Steps_Step13-RemoveSpriteVariables.png)
+1.  In the **Set Variables** section, click the small white arrow next to the following variables and select **Remove**.
+    
+    *   **Particles.SpriteRotation**
+    *   **Particles.SpriteSize**
+    
+    ![](CreateRibbonBasedEffect_Steps_Step14-RemoveSpriteVariables2.png)
+1.  When those variables have been deleted, click the **Plus** sign icon (**+**) next to **Particle Spawn** to display a list of options. Expand **Set Specific Parameters**, and add the following variables:
+    
+    *   **Particles.RibbonWidth**
+    *   **Particles.Velocity**
+    
+    ![](CreateRibbonBasedEffect_Steps_Step15-AddParticlesParameters.png)
+1.  Click the **Plus** sign icon (**+**) next to **Particle Spawn** to display a list of options again. Select **Experimental > Point Spiral Location**.
+    [REGION:note]
+    The **Speed** value of this variable expands the base diameter of the overall effect.
+    [/REGION]![](CreateRibbonBasedEffect_Steps_Step16-AddPointSpiralLocation.png)
+1.  With the variables now added to the Ribbon emitter, set the following Variables to the following values.
+    
+    | Variable Name | Value |
+    | --- | --- |
+    | **Particles.RibbonWidth** | 35.0 |
+    | **Particles.Velocity** | X:1.0, Y: 0, Z: 1.0 |
+    | **Point Spiral Location** | Speed: 500.0, RotSpeed: 4.0 |
+    | **Particles.Lifetime** | 5.0 |
+    
+1.  Locate and expand the **Acceleration Force** module.
+1.  Set the **Acceleration Z** value to **-100.0**.  
+    ![](CreateRibbonBasedEffect_Steps_Step19-SetAccelerationForce.png)
+    [REGION:note]
+    A negative value in the **Z** field produces a cone-shaped ribbon effect. If you put a positive value in the **Z** field, an inverted cup-shape is produced.
     [/REGION]
+1.  Once all that has been completed, click the **Compile**, **Apply** and **Save** buttons to compile and save your work.  
+    ![](ME_09.png)
+1.  Now right-click in the **Content Browser**, and select **FX > Niagara System**. Name the new system **Ribbon**.   
+    ![](SE_16.png)
+1.  Double-click on the Ribbon Niagara System to open it, then drag the **FX_Ribbon** Niagara Emitter from the Content Browser to the **Timeline**.  
+    ![](CreateRibbonBasedEffect_Steps_Step22-AddtoSystem.png)
+1.  Click the **Compile** and **Save** buttons to compile and save your work, and then close the Ribbon Niagara System. In the **Content Browser**, drag the **Ribbon Niagara System** into the level viewport.  
+    ![](CreateRibbonBasedEffect_Steps_Step23-AddToLevel.png)
 
 ## End Result
 
-When completed, you will now have a Niagara System that will play your Niagara Emitter when the Niagara System is placed in a Level, as shown in the following video:
+When completed, you have a Niagara System that plays your Niagara Emitter when the Niagara System is placed in a level, as shown in the following video.
 
 [OBJECT:EmbeddedVideo]
     [PARAMLITERAL:videoid]
-    wSRBBGirYig
+    pxQEQ4zWLWI
     [/PARAMLITERAL]
     [PARAMLITERAL:width]
-    100
+    78
     [/PARAMLITERAL]
     [PARAMLITERAL:height]
     56.25
diff --git a/Engine/Documentation/Source/Engine/Niagara/HowTo/SpriteEffect/SpriteEffect.INT.udn b/Engine/Documentation/Source/Engine/Niagara/HowTo/SpriteEffect/SpriteEffect.INT.udn
--- a/Engine/Documentation/Source/Engine/Niagara/HowTo/SpriteEffect/SpriteEffect.INT.udn
+++ b/Engine/Documentation/Source/Engine/Niagara/HowTo/SpriteEffect/SpriteEffect.INT.udn
@@ -9,114 +9,102 @@ parent:Engine/Niagara/HowTo
 order:2
 tags:Niagara
 topic-image:
+prereq: Engine/Niagara/HowTo/EnableNiagara
 
-One of the most used visual effects techniques is to render Texture and Materials to a 2D camera facing plans, called sprites. In the following How - to we will take a look setting up a Niagara Emitter to work with sprites so that when the effect is placed in a level, it will work and be displayed correctly. 
+One of the most used visual effects techniques is to render Texture and Materials to a 2D camera facing planes, called sprites. In the following How-To, we will take a look setting up a Niagara Emitter to work with sprites so that when the effect is placed in a level, it will work and be displayed correctly.
 
-![](HeroImage_00.png)
-
-## Steps
-
-[region:note]
-This How - to will make use of the **M_smoke_subUV** Material that can be found with the **Starter Content**. If you have not done so already, make sure that this Material or the Starter Content has been added to your project.
-[/region]
+![](Sprite_Banner_01.png "Sprite_Banner_01.png")
 
+[REGION:note]
+**Prerequisite Steps:**
+This How-To uses the **M_smoke_subUV** Material, which can be found with the **Starter Content**. If you have not done so already, make sure that this Material or the Starter Content has been added to your project.
+[/REGION]
 
-1. First, we will need to create a **Niagara Emitter** by right-clicking in the **Content Browser** and from the displayed menu go to **FX** and then click on the **Niagara Emitter**. 
-
-    ![](SE_01.png)
-
-1. Name the newly created Niagara Emitter, **FX_Smoke** and then double - click on it to open it up.
+## Steps
 
+1.  First, we will need to create a **Niagara Emitter** by right-clicking in the **Content Browser**, and from the displayed menu select **FX > Niagara Emitter**.   
+	![](SE_01.png)
+2.  Name the newly created Niagara Emitter **FX_Smoke** and then double-click it to open it.   
     ![](SE_02.png)
-
-1. With the Niagara Editor now open, go to the **Selected Emitter** panel and look for the **Render** section.
-
-	![](SE_03.png)
-
-1. Expand the **Sprite Render** section and set the following properties with the following options.
-
-    ![](SE_04.png)
-
-    |Property Name| Value|
-    |------|-----|
-    |**Material**|M_smoke_subUV|
-    |**Sub Image Size**| X: 8, Y: 8|
-    |**Sub UV Blending Enabled**| Checked|
-
-1. Now under **Emitter Update** in the **Spawn Rate** section, change the **Spawn Rate*** to a value of **50**.
-
-    ![](SE_05.png)
-
-1. Next, under **Particle Spawn** section go to the **Add Velocity** section and set the **Minimum** and **Maximum** to the following.
-
-    ![](SE_06.png)
-
-    |Property Name| Value|
-    |------|-----|
-    |**Minimum**|0,0,50|
-    |**Maximum**|0,0,200|
-
-1. Over in the **Sphere Location** set the **Sphere Radius** to **64** and then click on the **Expand** icon and make sure that **Surface Only** has been enabled.
-
-    ![](SE_07.png)
-
-    [region:note]
-    If you do not see the Surface Only option, press the small white triangle to expand more options.
-    [/region]
-
-1. Next, in the **Set Particles.Lifetime** section, click on the small white triangle next to **Lifetime** and from the drop-down list select the **Uniform Ranged Float** option and then input **2** into the **Minimum** lifetime and **5** into the **Maximum** lifetime.
-
-    ![](SE_08.png)
-
-1. In the **Set Particles. Sprite Size** set the **Minimum** value to **200** and set the **Maximum** value to **700**.
-
-    ![](SE_09.png)
-
-1. Go to the bottom of the Particle Update section and press the **Plus** sign icon to add a new module to the list. From the displayed list select the **Set Particles.SubImageIndex** option.
-
-    ![](SE_11.png)
-
-1. In the Set Particles.SubImageIndex section, press the white triangle next to the **SubImagIndex** option and from the displayed list, select the **Float From Curve** option.
-
-    ![](SE_12.png)
-
-1. In the **Curve** input window, select the first key and make sure that its **Time** and **Value** are both set to 0. Then select 2nd key and make sure that it's Time is set to **1** and the Value is set to **64**. 
-
-    ![](SE_13.png)
-
-1. Next, at the bottom of the Set Particles.SubImageIndex section, press the small grey triangle next to **CurveIndex** and from the dropdown list select **Normalized Age** option.
-
-    ![](SE_14.png)
-
-1. Then press the **Compile**, **to Apply** and **Save** buttons to apply and save the changes and then close the **FX_Smoke** emitter.
-
+3.  With the Niagara Editor now open, go to the **Selected Emitter** panel and locate the **Render** section.
+4.  Expand the **Sprite Rendering** section and set the following properties with the following values.  
+    ![](CreateSpriteBasedEffect_Steps_Step4-SpriteRender.png)  
+    
+    | Property Name | Value |
+    | --- | --- |
+    | **Material** | M_smoke_subUV |
+    | **Sub Image Size** | X: 8.0, Y: 8.0 |
+    | **Sub UV Blending Enabled** | Checked |
+    
+5.  Collapse the Render section, then locate and expand the **Emitter Update** section.
+6.  In the **Spawn Rate** section, change the **Spawn Rate** to a value of **50.0**.  
+    ![](CreateSpriteBasedEffect_Steps_Step6-EmitterUpdate.png)
+7.  Collapse the Emitter Update section, then locate and expand the **Add Velocity** section. 
+8.  Set the **Minimum** and **Maximum** velocity to the following values.  
+    ![](CreateSpriteBasedEffect_Steps_Step8-AddVelocityMinMax.png)
+    
+    | Property Name | Value |
+    | --- | --- |
+    | **Minimum** | X: 0.0, Y: 0.0, Z: 50.0 |
+    | **Maximum** | X: 0.0, Y: 0.0, Z: 200.0 |
+    
+9.  Locate the **Sphere Location** section.
+10.  Set the **Sphere Radius** value to **64.0**. Click the white arrow to expand the section, and make sure that the **Surface Only** box has been checked.  
+    ![](CreateSpriteBasedEffect_Steps_Step10-SphereRadius.png "CreateSpriteBasedEffect_Steps_Step10-SphereRadius.png")
+11.  Locate the **Set Variables** section.
+12.  Expand the **Particles.Lifetime** variable. This displays the **Minimum** and **Maximum **fields.  
+    
+13.  Set the **Minimum** and **Maximum** lifetime fields to the following values.  
+    ![](CreateSpriteBasedEffect_Steps_Step13-ParticlesLifetimeMinMax.png)
+    
+    | Property Name | Value |
+    | --- | --- |
+    | **Minimum** | 2.0 |
+    | **Maximum** | 3.0 |
+    
+14.  Expand the **Particles.SpriteSize** variable to display the **Minimum** and **Maximum** fields.
+15.  Set the **Minimum** and **Maximum** fields to the following values.  
+    ![](CreateSpriteBasedEffect_Steps_Step15-ParticlesSpriteSizeMinMax.png)
+    
+    | Property Name | Value |
+    | --- | --- |
+    | **Minimum** | X: 200.0, Y: 200.0 |
+    | **Maximum** | X: 700.0, Y: 700.0 |
+    
+16.  Collapse the Particle Spawn section, then locate and expand the **Particle Update** section. 
+17. Locate the **Acceleration Force** section. Set the **Z** value to **500**.  
+    ![](CreateSpriteBasedEffect_Steps_Step17_AccelerationForce.png)
+    [REGION:note]
+    If your particle flow is going up and then down (like a fountain), make sure your **Z** value is **positive** and not negative. Positive goes up, negative goes down.
+    [/REGION]
+18. Click the **Plus** sign icon (**+**) next to **Particle Update** to display a list of options. Select **Set Specific Parameters > Particles.SubImageIndex**.  This adds a **Set Variables** section for Particles.SubImageIndex.  
+    [REGION:imagetable]
+    | ![](CreateSpriteBasedEffect_Steps_Step18-ParticlesSubImageIndex1.png) | ![](CreateSpriteBasedEffect_Steps_Step18-ParticlesSubImageIndex2.png) |
+    | --- | --- |
+    [/REGION]
+19. In the **Particles.SubImageIndex** section, click the white triangle next to **Particle.SubImageIndex** to display a list of options. Select **Dynamic Inputs > Float From Curve**.  
+    ![](CreateSpriteBasedEffect_Steps_Step19-ParticlesSubImageIndexFloatCurve.png)
+20. In the **Curve** input window, click the first key on the left side of the curve. For this key, set its **Time** and **Value** to **0.0**.  
+    ![](CreateSpriteBasedEffect_Steps_Step20-FloatCurveFirstKey.png)
+21. Click the second key on the right side of the curve, and make sure that its **Time** is set to **1.0** and the **Value** is set to **0****.64**.  
+    ![](CreateSpriteBasedEffect_Steps_Step21-FloatCurveSecondKey.png)
+22. In the **Particles.SubImageIndex** variable, press the small white triangle next to **CurveIndex** to display a list of options. Select **Link Inputs > Particle Attributes > Normalized Age**.  
+    ![](CreateSpriteBasedEffect_Steps_Step22-CurveIndexNormalizedAge.png)
+23. Then click the **Compile**, **Apply** and **Save** buttons to apply and save the changes, and close the **FX_Smoke** Niagara Emitter.  
     ![](SE_15.png)
-
-1. Now over in the Content Browser, right - click and from the displayed list, go to **FX**  and then select the **Niagara System** option, giving it a name of **Smoke**.
-
+24. Now over in the **Content Browser**, right-click and select **FX > Niagara System**. Name the new Niagara System **Smoke**.  
     ![](SE_16.png)
-
-1. Double - click on the Smoke Niagara System to open it up and when open, drag the **FX_Smoke** from the Content Browser to the **Timeline**.
-
+25. Double-click the **Smoke** Niagara System to open it, and drag the **FX_Smoke** emitter from the Content Browser to the **Timeline**.  
     ![](SE_17.png)
-
-1. Press the **Compile** and **Save** buttons and then close the Smoke Niagara System down. Then, over in the Content Browser, drag the Smoke Niagara System into the level viewport.
-
-    [REGION:lightbox]
-    [![](SE_18.png)(w:834)](SE_18.png)
-    [/REGION]
-
-    [REGION:caption]
-    Click for full image.
-    [/REGION]
+26. Click the **Compile** and **Save** buttons to apply and save your changes, then close the Smoke Niagara System. In the Content Browser, drag the **Smoke** Niagara System into the level viewport.
 
 ## End Result
 
-When completed you will now have a Niagara System that will play your Niagara Emitter when the Niagara System is placed in a level like in the following video.
+When complete, you will have a Niagara System playing the Niagara Emitter when you place the Niagara System in a level (like in the following video).
 
 [OBJECT:EmbeddedVideo]
     [PARAMLITERAL:videoid]
-    h_3CVsDqKBc
+    3mGxYRZDORY
     [/PARAMLITERAL]
     [PARAMLITERAL:width]
     100
@@ -130,5 +118,4 @@ When completed you will now have a Niagara System that will play your Niagara Em
     [PARAMLITERAL:type]
     simple
     [/PARAMLITERAL]
-[/OBJECT]
-
+[/OBJECT]
\ No newline at end of file
diff --git a/Engine/Documentation/Source/Engine/Niagara/Niagara.INT.udn b/Engine/Documentation/Source/Engine/Niagara/Niagara.INT.udn
--- a/Engine/Documentation/Source/Engine/Niagara/Niagara.INT.udn
+++ b/Engine/Documentation/Source/Engine/Niagara/Niagara.INT.udn
@@ -14,18 +14,14 @@ tags:Niagara
 ![](NiagaraBanner.png)(convert:false)
 [/REGION]
 
-The **Niagara Editor** is one of two tools you can use to create and adjust visual effects inside Unreal Engine 4 (UE4). Before Niagara, the primary way to create and edit visual effects in UE4 was to use [**Cascade**](Engine\Rendering\ParticleSystems). While Niagara does have many of the same of particle manipulation methods that Cascade offers, the manner in which you interact and build visual effects with Niagara is vastly different. 
+The **Niagara Editor** is one of two tools you can use to create and adjust visual effects inside Unreal Engine 4 (UE4). Before Niagara, the primary way to create and edit visual effects in UE4 was to use [**Cascade**](Engine\Rendering\ParticleSystems). While Niagara does have many of the same of particle manipulation methods that Cascade offers, the manner in which you interact and build visual effects with Niagara is vastly different.
 
-On this page, you will find links to documentation regarding the different aspects of the Niagara Editor. If you are new to Niagara, it is recommended that you check out the **Essentials** section which includes a high-level overview of the Niagara Editor. In this section, you will also find an Editor Reference which breaks down components of the Niagara Editor, a Quick Start guide to quickly get you up and running, multiple How-to's which are one-off pages of how to perform certain tasks with Niagara as well as a section on Workflow tips. 
+On this page, you will find links to documentation regarding the different aspects of the Niagara Editor. If you are new to Niagara, it is recommended that you check out the **Essentials** section which includes a high-level overview of the Niagara Editor. In this section, you will also find an Editor Reference which breaks down components of the Niagara Editor, a Quick Start guide to quickly get you up and running, multiple How-to's which are one-off pages of how to perform certain tasks with Niagara as well as a section on Workflow tips.
 
 ## Starting Out
 
 [DIR(output:"topic" parent:"Engine/Niagara" org:"hierarchy" end:"1" tags:"Getting Started")]
 
-[COMMENT:none]
-## Essentials
-[/COMMENT]
-
 ## Guides
 
 [DIR(output:"fancy" parent:"Engine/Niagara/HowTo" org:"hierarchy" end:"1" )]
diff --git a/Engine/Documentation/Source/Engine/Niagara/NiagaraKeyConcepts/NiagaraKeyConcepts.INT.udn b/Engine/Documentation/Source/Engine/Niagara/NiagaraKeyConcepts/NiagaraKeyConcepts.INT.udn
--- a/Engine/Documentation/Source/Engine/Niagara/NiagaraKeyConcepts/NiagaraKeyConcepts.INT.udn
+++ b/Engine/Documentation/Source/Engine/Niagara/NiagaraKeyConcepts/NiagaraKeyConcepts.INT.udn
@@ -1,27 +1,26 @@
-Availability:Docs
-Crumbs: %ROOT%
+Availability: Public
+Crumbs: 
 Title: Niagara Key Concepts
-Description:Critical concepts for working with particles in Unreal and the Niagara, particle editor.
-Type: Overview
-SkillLevel:Beginner
-Version: 4.17
-Parent:Engine/Niagara 
+Description: Critical concepts for working with particles in Unreal and the Niagara, particle editor.
+Type:overview
+Version: 4.20
+Parent:Engine/Niagara
 Order: 2
-Tags:Getting Started
+Tags: Getting Started
+Tags: Early Access
 topic-image:Engine\Niagara\NiagaraSocial.png
 
 [TOC(start:2 end:2)]
 
-The Niagara Editor gives the user the ability to create a wide range of visual effects that can meet any Unreal Engine 4 (UE4) projects needs and do so in real-time. From effects for the player to effects for the environment, Niagara is flexible and powerful enough to handle whatever visual effects your project requires. In the following document, we will take a high-level look at how Niagara is different from Cascade and also how all the various pieces that make up the Niagara Editor work together. 
+The Niagara Editor gives the user the ability to create a wide range of visual effects that can meet any Unreal Engine 4 (UE4) projects needs and do so in real-time. From effects for the player to effects for the environment, Niagara is flexible and powerful enough to handle whatever visual effects your project requires. In the following document, we will take a high-level look at how Niagara is different from Cascade and also how all the various pieces that make up the Niagara Editor work together.
 
 ## A Modular Approach to Particle Effects
 
-The primary overreaching concept of Cascade was to incorporate modularity into the creation of effects whenever possible. This was achieved by supplying a default set of modules when any new effect was created and then allowing effect artist to add additional modules or entirely new emitters based on what the effect needed to achieve. While this system works remarkably well for some years, it had a few minor issues that made things like sharing effect data between effects or updating many effects at once a time-consuming task that was prone to error. 
+The primary overreaching concept of Cascade was to incorporate modularity into the creation of effects whenever possible. This was achieved by supplying a default set of modules when any new effect was created and then allowing effect artist to add additional modules or entirely new emitters based on what the effect needed to achieve. While this system works remarkably well for some years, it had a few minor issues that made things like sharing effect data between effects or updating many effects at once a time-consuming task that was prone to error.
 
 [REGION:lightbox]
 [![](NKC_00.png)(w:764)](NKC_00.png)
 [/REGION]
-
 [REGION:caption]
 This image shows that a **Cascade Particle System** is made up of a collection of emitters.
 [/REGION]
@@ -31,7 +30,6 @@ To address these issues, Niagara has fully embraced this idea of modularity by d
 [REGION:lightbox]
 [![](NKC_01.png)(w:764)](NKC_01.png)
 [/REGION]
-
 [REGION:caption]
 This image shows that a **Niagara System** is made up of a collection of **Niagara Emitters**.
 [/REGION]
@@ -44,18 +42,18 @@ When you create or open up a Niagara Emitter, it will come with some default mod
 [![](NKC_02.png)(w:620)](NKC_02.png)
 [/REGION]
 
-* **Emitter Spawn**
-    * The Emitter Spawn section is run when the emitter starts up. This is similar to a Construction Script for Blueprints in that it does one-time setup that can be used in later scripts. It cannot reference individual particles, but it can be used to set emitter variables, run emitter modules, and set up variables that will be used in the particle scripts.
-* **Emitter Update**
-    * The Emitter Update section is run every frame that the Emitter is active. Users will determine how many particles to spawn, compute constant values for use in particle update scripts, and do other per-frame "tick" logic for the emitter here. It is important to have the Emitter Lifecycle module in this section to ensure that all of the lifecycle logic is executed for each emitter, allowing the system to properly shut down.
-* **Particle Spawn**
-    * The Particle Spawn section is run once every time a particle is born. Users will use this section to set initial state for the particle. If the Emitter is not set to use Interpolated Spawn, the output from this section will define the state of the particle for the first frame it is rendered. If Interpolated Spawning does occur, the particles will be evenly spread out in the time range from the previous frame if not an instantaneous burst and the update script will be run for each particle for the fractional amount of the frame since it was created. This allows users to achieve smooth trails and can also be used to allow the particle update script to define the state of the particle for the first frame that it is rendered. If you need randomness to differentiate each particle, it is best to choose your random values in the spawn script and store that value for use in each subsequent update "tick."
-* **Particle Update**
-    * The Particle Update section controls how an individual particle will behave over its life. Changes to physical simulation attributes, rendering attributes, etc. can all be done within this section. If the user wishes to broadcast data to other emitters, modules that emit events can occur in this section.  It is important to have the Update Age module in this section to ensure that each particle continues to age and eventually die once it has passed its lifetime.
-* **Add Event Handler**
-    * The Add Event Handler section determines how this emitter will respond to incoming events from collision or other emitters. 
+*   **Emitter Spawn**
+    *   The Emitter Spawn section is run when the emitter starts up. This is similar to a Construction Script for Blueprints in that it does one-time setup that can be used in later scripts. It cannot reference individual particles, but it can be used to set emitter variables, run emitter modules, and set up variables that will be used in the particle scripts.
+*   **Emitter Update**
+    *   The Emitter Update section is run in every frame that the Emitter is active. Users will determine how many particles to spawn, compute constant values for use in particle update scripts, and do other per-frame "tick" logic for the emitter here. It is important to have the Emitter Lifecycle module in this section to ensure that all of the lifecycle logic is executed for each emitter, allowing the system to properly shut down.
+*   **Particle Spawn**
+    *   The Particle Spawn section is run once every time a particle is born. Users will use this section to set initial state for the particle. If the Emitter is not set to use Interpolated Spawn, the output from this section will define the state of the particle for the first frame it is rendered. If Interpolated Spawning does occur, the particles will be evenly spread out in the time range from the previous frame if not an instantaneous burst and the update script will be run for each particle for the fractional amount of the frame since it was created. This allows users to achieve smooth trails and can also be used to allow the particle update script to define the state of the particle for the first frame that it is rendered. If you need randomness to differentiate each particle, it is best to choose your random values in the spawn script and store that value for use in each subsequent update "tick."
+*   **Particle Update**
+    *   The Particle Update section controls how an individual particle will behave over its life. Changes to physical simulation attributes, rendering attributes, etc. can all be done within this section. If the user wishes to broadcast data to other emitters, modules that emit events can occur in this section. It is important to have the Update Age module in this section to ensure that each particle continues to age and eventually die once it has passed its lifetime.
+*   **Add Event Handler**
+    *   The Add Event Handler section determines how this emitter will respond to incoming events from collision or other emitters.
 
-## Niagara System 
+## Niagara System
 
 To use a Niagara Emitter in your project, it will first need to be placed into a **Niagara System**. A Niagara System is a used to display the individual Niagara Emitters that have been created to make an effect. This also makes it very easy to mix and match different Niagara emitters together to create entirely new effects using only the existing emitters.
 
@@ -63,7 +61,7 @@ To use a Niagara Emitter in your project, it will first need to be placed into a
 [![](NKC_08.png)(w:620)](NKC_08.png)
 [/REGION]
 
- To add a Niagara Emitter to a Niagara System, you will need to drag your Niagara Emitter from the Content Browser onto the **Track** section of the **Timeline**. 
+To add a Niagara Emitter to a Niagara System, you will need to drag your Niagara Emitter from the Content Browser onto the **Track** section of the **Timeline**.
 
 [REGION:lightbox]
 [![](NKC_07.png)(w:600)](NKC_07.png)
@@ -73,8 +71,4 @@ You can adjust the properties of Niagara Emitters that are placed in the Timelin
 
 [REGION:lightbox]
 [![](NKC_09.png)(w:250)](NKC_09.png)
-[/REGION]
-
-
-
-
+[/REGION]
\ No newline at end of file
diff --git a/Engine/Documentation/Source/Engine/Physics/Cloth/Overview/ClothOverview.INT.udn b/Engine/Documentation/Source/Engine/Physics/Cloth/Overview/ClothOverview.INT.udn
--- a/Engine/Documentation/Source/Engine/Physics/Cloth/Overview/ClothOverview.INT.udn
+++ b/Engine/Documentation/Source/Engine/Physics/Cloth/Overview/ClothOverview.INT.udn
@@ -354,44 +354,12 @@ If you were to repaint your LOD0 mesh but don't want to have to re-paint all you
 	
 1. Select **Rebuild Asset Parameter Masks** to automate the process and carry the new painted cloth values over to your other LODs.
 
-	[REGION:imagetable]
-	| --- | --- |
-	| [INCLUDE:#lod0] | [INCLUDE:#lod1] |
-	[/REGION]
 
 	You will now have the newly painted values for LOD0 passed down to any lower LOD meshes like the comparison you can see here with LOD1 before and after.
 
-<!----
-
-[EXCERPT:lod0]
-[OBJECT:ComparisonSlider]
-	[PARAM:before]
-	![LOD0 Before](Tonemapper_On.png)
-	[/PARAM]
-	[PARAM:after]
-	![LOD0 After](Tonemapper_Off.png)
-	[/PARAM]
-[/OBJECT]
-[/EXCERPT:lod0]
 
-----  SAVE FOR 4.18 Release Update -
 [/COMMENT]
 
-
-[EXCERPT:lod1]
-[OBJECT:ComparisonSlider]
-	[PARAM:before]
-	![LOD1 Before](Tonemapper_On.png)
-	[/PARAM]
-	[PARAM:after]
-	![LOD1 After](Tonemapper_Off.png)
-	[/PARAM]
-[/OBJECT]
-[/EXCERPT:lod1]
-
---->
-
-
 ## APEX (Legacy) Assets
 
 If you have created your clothing asset using the APEX Clothing Tool provided by NVIDIA for 3DS Max or Maya, you can continue to use this workflow. There are two ways to go about
diff --git a/Engine/Documentation/Source/Engine/Physics/Constraints/ConstraintsReference/Constraints.INT.udn b/Engine/Documentation/Source/Engine/Physics/Constraints/ConstraintsReference/Constraints.INT.udn
--- a/Engine/Documentation/Source/Engine/Physics/Constraints/ConstraintsReference/Constraints.INT.udn
+++ b/Engine/Documentation/Source/Engine/Physics/Constraints/ConstraintsReference/Constraints.INT.udn
@@ -229,31 +229,31 @@ Whether this angular motor uses SLERP (Spherical Linear Interpolation) or decomp
 
 | Property | Description |
 | -------- | ----------- |
-|[REGION:tablesection]Constraint[/REGION]||
-| **Component Name 1** |  |
-| **Component Name 2** |   |
-| **Joint Name** |  |
-| **Constraint Bone 1** |  |
-| **Constraint Bone 2** |  |
-| **Disable Collision** |  |
+| [REGION:tablesection] Constraint [/REGION] ||
+| **Component Name 1** ||
+| **Component Name 2** ||
+| **Joint Name** ||
+| **Constraint Bone 1** ||
+| **Constraint Bone 2** ||
+| **Disable Collision** ||
 
-* For information on using **Physics Constraints** as a Blueprint Component, see: [**%Engine/Physics\Constraints\ConstraintsBlueprints:title%**](Engine/Physics\Constraints\ConstraintsBlueprints)
+* For information on using **Physics Constraints** as a Blueprint Component, see: [Engine/Physics\Constraints\ConstraintsBlueprints](Engine/Physics\Constraints\ConstraintsBlueprints)
 
 ### Physics Asset Tool Constraint
 
 | Property | Description |
 | -------- | ----------- |
-|[REGION:tablesection]Constraint[/REGION]|
+|[REGION:tablesection]Constraint[/REGION]||
 | **Joint Name** |  |
 | **Constraint Bone 1** |  |
 | **Constraint Bone 2** |  |
 | **Disable Collision** |  |
 
-* For further reference on the **Physics Asset Tool** (**PhAT**), see: [**%Engine/Physics\PhAT\Reference:title%**](Engine/Physics\PhAT\Reference)
+* For further reference on the **Physics Asset Tool** (**PhAT**), see: [Engine/Physics\PhAT\Reference]()
 
-* For information on using **Physics Constraints** in the **Physics Asset Tool** (**PhAT**), see: [**%Engine/Physics\PhAT\UserGuide:title%**](Engine/Physics\PhAT\UserGuide)
+* For information on using **Physics Constraints** in the **Physics Asset Tool** (**PhAT**), see: [Engine/Physics\PhAT\UserGuide]()
 
-* For information on using the **Physics Asset Tool** (**PhAT**), see: [**%Engine/Physics\PhAT:title%**](Engine/Physics\PhAT)
+* For information on using the **Physics Asset Tool** (**PhAT**), see: [**%Engine/Physics\PhAT:title%**]()
 
 ### Shared Properties
 
diff --git a/Engine/Documentation/Source/Engine/Physics/Constraints/DampingAndFriction/DampingAndFriction.INT.udn b/Engine/Documentation/Source/Engine/Physics/Constraints/DampingAndFriction/DampingAndFriction.INT.udn
--- a/Engine/Documentation/Source/Engine/Physics/Constraints/DampingAndFriction/DampingAndFriction.INT.udn
+++ b/Engine/Documentation/Source/Engine/Physics/Constraints/DampingAndFriction/DampingAndFriction.INT.udn
@@ -69,7 +69,7 @@ it which may look odd.
 To simulate 'friction' within a joint, you can use the joint drive properties in the RB_ConstraintInstance (you can get to these properties in PhAT by selecting the joint and 
 pressing the 'I' key or the 'i' icon on the toolbar). First you want to turn on the appropriate bTwistVelocityDrive or bSwingVelocityDrive drive, set the desired angular velocity 
 to zero and modify the AngularDriveForceLimit to control how much friction is in the joint. You can find more information on this on the 
-[PhysicsConstraintReference](Engine/Physics/PhysX/Constraints) page. 
+[PhysicsConstraintReference](Engine/Physics/Constraints) page. 
 
 The key thing here is that the damping is applied between two jointed object. Imagine two objects jointed together, but then thrown through space. If you used PhysicalMaterial 
 damping here, it would slow down how quickly the pair of objects would move through the air, when what you really want is to modify how they move relative to each other. 
diff --git a/Engine/Documentation/Source/Engine/Physics/FrictionRestitutionAndDamping/FrictionRestitutionAndDamping.INT.udn b/Engine/Documentation/Source/Engine/Physics/FrictionRestitutionAndDamping/FrictionRestitutionAndDamping.INT.udn
--- a/Engine/Documentation/Source/Engine/Physics/FrictionRestitutionAndDamping/FrictionRestitutionAndDamping.INT.udn
+++ b/Engine/Documentation/Source/Engine/Physics/FrictionRestitutionAndDamping/FrictionRestitutionAndDamping.INT.udn
@@ -68,7 +68,7 @@ it which may look odd.
 To simulate 'friction' within a joint, you can use the joint drive properties in the RB_ConstraintInstance (you can get to these properties in PhAT by selecting the joint and 
 pressing the 'I' key or the 'i' icon on the toolbar). First you want to turn on the appropriate bTwistVelocityDrive or bSwingVelocityDrive drive, set the desired angular velocity 
 to zero and modify the AngularDriveForceLimit to control how much friction is in the joint. You can find more information on this on the 
-[PhysicsConstraintReference](Engine/Physics/PhysX/Constraints) page. 
+[PhysicsConstraintReference](Engine/Physics/Constraints) page. 
 
 The key thing here is that the damping is applied between two jointed object. Imagine two objects jointed together, but then thrown through space. If you used PhysicalMaterial 
 damping here, it would slow down how quickly the pair of objects would move through the air, when what you really want is to modify how they move relative to each other. 
diff --git a/Engine/Documentation/Source/Engine/Rendering/CurveAtlasesInMaterials/CurveAtlasesInMaterials.INT.udn b/Engine/Documentation/Source/Engine/Rendering/CurveAtlasesInMaterials/CurveAtlasesInMaterials.INT.udn
new file mode 100644
--- /dev/null
+++ b/Engine/Documentation/Source/Engine/Rendering/CurveAtlasesInMaterials/CurveAtlasesInMaterials.INT.udn
@@ -0,0 +1,143 @@
+Availability: Public
+Crumbs: 
+Title: Curve Atlases in Materials
+Description: Curve Atlases store a combination of Curve assets that enables you to access Curve Linear Color data through your Materials.
+Version: 4.20
+Parent: Engine/Rendering/Materials
+Order: 6
+
+
+![](MaterialGradientCurvesBanner.png "MaterialGradientCurvesBanner.png")
+
+Curve Atlases store a combination of [Curve](Engine/UI/CurveEditor) assets that enables you to access Curve Linear Color data through your Materials. The Curve Atlas stores any number of Curves in a texture that can be used with [Materials](https://docs.unrealengine.com/en-us/Engine/Rendering/Materials). When creating [Material Instances](https://docs.unrealengine.com/en-us/Engine/Rendering/Materials/MaterialInstances) and changing [Material Instance Dynamics](https://docs.unrealengine.com/en-us/Engine/Rendering/Materials/MaterialInstances#materialinstancedynamic) (MIDs) through Blueprint, you are able to set a curve as an override, enabling you to quickly iterate and change RGBA curve values without changing your base Material.
+
+## Curve Linear Color Asset
+
+A **Curve Linear Color** asset is used to store the interpolated RGBA points that are evaluated over a given range producing a color gradient that can be used with Materials.  
+![](CurveGradient.png)
+
+[REGION:caption]
+Example of a Curve Gradient created with a Curve Linear Color asset
+[/REGION]
+
+Use the **Content Browser** to create your Curve asset by selecting **Add New > Miscellaneous > Curve**.  
+
+[REGION:lightbox]
+[![](Create_Curve.png)(w:350)](Create_Curve.png)
+[/REGION]
+[REGION:caption]
+Click image for full size.
+[/REGION]
+
+Then select **CurveLinearColor** from the **Pick Curve Class** window.  
+![](PickCurveClassWindow.png)
+
+A new Curve Linear Color asset is added to the Content Browser displaying a preview of its current gradient. When you open the Curve Asset Editor, you can add, adjust, and subtract keys to the individual RGBA curves in the graph.  
+![](MyCurveAsset.png)
+
+When you open the Curve in its Asset Editor, you’ll be able to set the individual RGBA curves, adjust color values for the Curve, and preview the Curve Gradient Result.  
+
+[REGION:lightbox]
+[![](CurveAssetEditor.png)](CurveAssetEditor.png)
+[/REGION]
+[REGION:caption]
+Click image for full size.
+[/REGION]
+
+[REGION:caption]
+1.  The **Curve Graph** is where you adjust, add, and subtract any keys for the individual RGBA curves.
+2.  The **Color** panel is where you adjust color values that affect all keys.
+3.  The **Curve Gradient Result** shows you the resulting gradient.
+
+[REGION:note]
+You can add keys to your Curve by holding SHIFT and clicking on individual curve to add a single key. If you click in the empty graph, a new key will be added to all curves at the point where you clicked.
+[/REGION]
+[/REGION]
+
+## Curve Atlas Asset
+
+A **Curve Atlas** asset is used to store and access multiple Curve assets enabling management of gradient lookup tables (LUT). The Curve Atlas Asset Editor is similar to the Texture Editor, enabling you to make adjustments for Brightness, Vibrance, Hue, and more.  
+![](CurveAtlas.png)(w:200 h:200)
+
+[REGION:caption]
+Example of the Curve Atlas LUT
+[/REGION]
+
+The curves assigned to the Curve Atlas as **Gradient Curves** make up the atlas. The Material Graph uses the texture created to perform the lookup for Materials applied to Actors.
+
+Use the **Content Browser** to create your Curve Atlas asset by selecting **Add New > Miscellaneous > Curve Atlas**.  
+
+[REGION:lightbox]
+[![](Create_CurveAtlas.png)(w:350)](Create_CurveAtlas.png)
+[/REGION]
+[REGION:caption]
+Click image for full size.
+[/REGION]
+
+When you open the Curve Atlas in its Asset Editor, you’ll be able to set the number of curves it can store, assign your curves, and adjust color values for all assigned curves - similar to the functionality of the Texture Editor.  
+
+[REGION:lightbox]
+[![](CurveAtlasAssetEditor_1.png)](CurveAtlasAssetEditor_1.png)
+[/REGION]
+[REGION:caption]
+Click image for full size.
+[/REGION]
+
+1.  The main viewport displays the **Curve Atlas** LUT that displays all assigned **Gradient Curves** for the **Texture Size** that is applied.
+2.  The **Curves** panel is where you sets the **Texture Size** for the number of **Gradient Curves** you need. The default is 256.
+3.  The **Adjustments** panel makes texture adjustments to all **Gradient Curves** assigned to the Curve Atlas.
+
+[REGION:note]
+For maximum texture efficiency, it’s best to use a **power of 2** value (such as 32, 64, 128). The texture sizes here are for demonstration purposes to include the entirety of the Editor Interface.
+[/REGION]
+[REGION:warning]
+The **Texture Size** used can have an impact on the fidelity of your gradients, so it’s best not to use a small atlas size for a complex curve. However, you can save texture size by using a small atlas for a set of simple curves.
+[/REGION]
+
+To add a new Gradient Curve to the Curve Atlas, click the **Plus** (**+**) icon to add array elements, and the **Trash Can** icon to remove elements.  
+![](CurveAtlasAssetEditor_AddRemove.png)
+
+## Using Curve Atlas with Materials
+
+After creating your Curves and applying them to a Curve Atlas, you can create a Material that references the Curve Atlas and a Curve assigned to the atlas.
+
+To sample the curve from an atlas, create a new Material and in the graph, right-click and add a **Curve Atlas Row Parameter** node.  
+![](CurveAtlasRowParameter.png)
+
+[REGION:note]
+This node is similar to a Scalar Parameter in that it enables you to work with [Material Instance Dynamics](Engine/Rendering/Materials/MaterialInstances#materialinstancedynamic) (MIDs) that can then be used through Blueprint for the V-axis of the atlas UVs, but the node will do the sampling for you and return a Vector 3, along with R, G, B, and A, masks.
+[/REGION]
+
+When you select the node, in the **Details** panel, you can assign the **Atlas** and a default **Curve** from that atlas to use with this Material.  
+![](MatDetails_AtlasAndCurveAssignment.png)
+
+The Curve Atlas bakes down when it is compiled, which means there is currently no runtime support for changing what is in the atlas or the curve data that is stored in the atlas at runtime. However, you can store a large amount of data in a single Curve Atlas, and use Blueprint to override the curve being sampled from a Material Instance.
+
+As an example, here is a “Rock” Material that is using multiple Curves assigned to a Curve Atlas.  
+
+[REGION:lightbox]
+[![](RockMaterialSetup_2.png)](RockMaterialSetup_2.png)
+[/REGION]
+[REGION:caption]
+Click image for full size.
+[/REGION]
+
+Then when you create any **Material Instances**, you can change the Scalar Parameter to select the **Curve** asset referenced by your **Curve Atlas** that is applied to your mesh.
+
+[REGION:imagetable]
+The examples here show a Material Instance applied and the Curve that is being referenced from the Curve Atlas.
+| ![](Curve_Base.png)(w:250) | ![](Curve_1.png)(w:250) | ![](Curve_2.png)(w:250) | ![](Curve_3.png)(w:250) |
+| --- | --- | --- | --- |
+| Base | Curve 1 | Curve 2 | Curve 3 |
+[/REGION]
+
+### Accessing Curve Atlases through Blueprint
+
+In Blueprint, you can set the scalar parameter value on a Dynamic Material Instance using the **Get Curve Position** node. Get Curve Position takes the Curve Atlas as an input, passes the scalar value to the **Set Scalar Parameter Value**, and returns a Boolean indicating whether the curve was found in the atlas.  
+
+[REGION:lightbox]
+[![](BlueprintCurve.png)](BlueprintCurve.png)
+[/REGION]
+[REGION:caption]
+Click image for full size.
+[/REGION]
\ No newline at end of file
diff --git a/Engine/Documentation/Source/Engine/Rendering/LightingAndShadows/CapsuleShadows/CapsuleShadowsLanding.INT.udn b/Engine/Documentation/Source/Engine/Rendering/LightingAndShadows/CapsuleShadows/CapsuleShadowsLanding.INT.udn
--- a/Engine/Documentation/Source/Engine/Rendering/LightingAndShadows/CapsuleShadows/CapsuleShadowsLanding.INT.udn
+++ b/Engine/Documentation/Source/Engine/Rendering/LightingAndShadows/CapsuleShadows/CapsuleShadowsLanding.INT.udn
@@ -5,7 +5,7 @@ crumbs:
 version:4.15
 parent:Engine/Rendering/LightingAndShadows
 order:
-related: Engine/Rendering/LightingAndShadows/DistanceFields/HowTo/DFHT_4
+related: Engine/Rendering/LightingAndShadows/MeshDistanceFields/HowTo/DFHT_4
 type:landing
 tags: Lighting
 topic-image:CapsuleShadows_Topic.png
diff --git a/Engine/Documentation/Source/Engine/Rendering/LightingAndShadows/DistanceFieldAmbientOcclusion/DistanceFieldAmbientOcclusion.INT.udn b/Engine/Documentation/Source/Engine/Rendering/LightingAndShadows/DistanceFieldAmbientOcclusion/DistanceFieldAmbientOcclusion.INT.udn
--- a/Engine/Documentation/Source/Engine/Rendering/LightingAndShadows/DistanceFieldAmbientOcclusion/DistanceFieldAmbientOcclusion.INT.udn
+++ b/Engine/Documentation/Source/Engine/Rendering/LightingAndShadows/DistanceFieldAmbientOcclusion/DistanceFieldAmbientOcclusion.INT.udn
@@ -191,7 +191,7 @@ By default, non-directional AO is applied to specular. You can use **r.AOSpecula
 
 ### Foliage 
 
-For Actors painted with the [Foliage Tool](Engine/Rendering/LightingAndShadows/MeshDistanceFields#folaigetoolsettings), you must first enable the option **Affect Distance Field Lighting** in tool settings.
+For Actors painted with the [Foliage Tool](Engine/Rendering/LightingAndShadows/MeshDistanceFields#foliagetoolsettings), you must first enable the option **Affect Distance Field Lighting** in tool settings.
 
 ![](Engine/Rendering/LightingAndShadows/MeshDistanceFields/FoliageToolSettings.png)
 
diff --git a/Engine/Documentation/Source/Engine/Rendering/LightingAndShadows/LightTypes/SkyLight/SkyLight.INT.udn b/Engine/Documentation/Source/Engine/Rendering/LightingAndShadows/LightTypes/SkyLight/SkyLight.INT.udn
--- a/Engine/Documentation/Source/Engine/Rendering/LightingAndShadows/LightTypes/SkyLight/SkyLight.INT.udn
+++ b/Engine/Documentation/Source/Engine/Rendering/LightingAndShadows/LightTypes/SkyLight/SkyLight.INT.udn
@@ -22,7 +22,7 @@ Social-Image:SkyLight_1.png
 
 
 
-The Sky Light captures the distant parts of your level and applies that to the scene as a light. That means the sky’s appearance and its lighting/reflections will match, even if your sky is 
+The Sky Light captures the distant parts of your level and applies that to the scene as a light. That means the skyï¿½s appearance and its lighting/reflections will match, even if your sky is 
 coming from atmosphere, or layered clouds on top of a skybox, or distant mountains. You can also manually specify a cubemap to use. 
 
 ## Scene Capture
@@ -33,7 +33,7 @@ The Sky Light will only capture the scene under certain circumstances:
 * For **Stationary** or **Movable Sky Lights**, updates happen once on load and only updates further when **Recapture Sky** is called. This can be done via the **Details** panel or via Blueprint call in-game.
 
 [REGION:note]
-If you change the texture that the skybox is using, it will not automatically know to update. You’ll be required to use one of the methods above for it to update.
+If you change the texture that the skybox is using, it will not automatically know to update. Youï¿½ll be required to use one of the methods above for it to update.
 [/REGION]
 
 A Sky Light should be used instead of the Ambient Cubemap to represent the sky's light because Sky Lights support local shadowing, which prevents indoor areas from getting lit by the sky.
@@ -172,15 +172,12 @@ For additional information visit the [Bent Normal Map](Engine/Rendering/Lighting
 [/REGION]
 
 
-
-
-
-
 ### Multi-Bounce
 
 Multiple bounces of indirect light from Sky Lights is supported by adjusting the **Num Sky Lighting Bounces** setting in **World Settings** > **Lightmass**. 
 
 ![](LightmassSettings.png)
+
 [REGION:caption] Setting for the number of sky light bounces to use. [/REGION]
 
 [OBJECT:ComparisonSlider]
@@ -201,7 +198,6 @@ Each added bounce increases non-distributed build times, so use with care when n
 [/REGION]
 
 [COMMENT:none]
--
 [OBJECT:ComparisonSlider]
 	[PARAM:before]
 	![Number of Bounces: 1](SkyLight_6.png)(w:1000)
@@ -210,7 +206,6 @@ Each added bounce increases non-distributed build times, so use with care when n
 	![Number of Bounces: 10](SkyLight_5.png)(w:1000)
 	[/PARAM]
 [/OBJECT]
--
 [/COMMENT]
 
 ## Properties
@@ -280,43 +275,7 @@ These features would be useful with Sky Lights but are not yet implemented:
 	
 	![](RecaptureSkyLightBP.png)(w:700)
 
-
-
-
-
-
-
-
-
-
-
-
-
-
-
-
-
-
-
-
-
-
-
-
-
-
-
-
-
-
-
-
-
-
-
 [COMMENT:none]
--
-
 
 [REGION:fullwidth raw]
 ![](Skylight_Header.png)
@@ -345,11 +304,11 @@ Pictured below are example images from a scene which utilizes a Sky Light.
 	![Sky Light](Skylight.png)(w:800)
 	[/PARAM]
 	[PARAM:after]
-	![No Sky Light](NoSkylight.png)(w:800)
+	![No Sky Light](NoSkyLight.png)(w:800)
 	[/PARAM]
 [/OBJECT]
 
-<!--
+
 [REGION:lightbox]
 
 [![](skyLightEx_01.png)(w:884)](skyLightEx_01.png) [![](skyLightEx_02.png)(w:884)](skyLightEx_02.png)
@@ -451,5 +410,3 @@ The properties of a Sky Light are broken up into 2 categories: Light and Sky Lig
 | **Sky Distance Threshold** | Distance from the Sky Light at which any geometry should be treated as part of the sky (also used by [REDACTED_PERSON]). |
 | **Lower Hemisphere is Black** | Whether all lighting from the lower hemisphere should be set to zero.  This is useful to prevent leaking from the lower hemisphere. |
 | **Recapture Scene** | When the Sky Light Actor is set to `SLS_CapturedScene`, this will recapture the image the Sky Light is using to light the scene. |
-
----->
\ No newline at end of file
