Copying //UE4/Dev-Documentation to //UE4/Samples-Main (Source: //UE4/Dev-Documentation @ 4321946)

diff --git a/Engine/Documentation/Source/Engine/Rendering/LightingAndShadows/ReflectionEnvironment/ReflectionEnvironment.INT.udn b/Engine/Documentation/Source/Engine/Rendering/LightingAndShadows/ReflectionEnvironment/ReflectionEnvironment.INT.udn
--- a/Engine/Documentation/Source/Engine/Rendering/LightingAndShadows/ReflectionEnvironment/ReflectionEnvironment.INT.udn
+++ b/Engine/Documentation/Source/Engine/Rendering/LightingAndShadows/ReflectionEnvironment/ReflectionEnvironment.INT.udn
@@ -54,7 +54,7 @@ It is also extremely important to setup your level's Materials to work well with
 used in flat areas as this will help them better show off reflections.
 
 [REGION:imagetable]
-| ![Curvy and Sharp](1CurvyAndSharp.png)(w:290) | ![Flat and Rough](1FlatAndRough.png)(w:290) | ![Flat and Sharp](1FlatAndSharp.png)(w:290)|
+| ![Curvy and Sharp](1CurvyAndSharp.png) | ![Flat and Rough](1FlatAndRough.png) | ![Flat and Sharp](1FlatAndSharp.png)|
 | ------ | ------ | ------ |
 | Smooth surface but curvy geometry: Good quality reflections | Rough surface but flat geometry: Good quality reflections |Smooth surface AND flat geometry: Evident that reflections do not match up
 [/REGION]
@@ -418,16 +418,19 @@ You can blend between multiple different cubemap reflections by providing the Re
 
 	[OBJECT:EmbeddedVideo]
 		[PARAMLITERAL:width]
-		640
+		50
 		[/PARAMLITERAL]
 		[PARAMLITERAL:height]
-		360
+		56.25
+		[/PARAMLITERAL]
+		[PARAMLITERAL:units]
+		%
 		[/PARAMLITERAL]
 		[PARAMLITERAL:videoid]
 		QYbSsYP2Nkk
 		[/PARAMLITERAL]
 	[/OBJECT]
-
+	
 ### Visualizing
 
 The Reflection Override viewmode has been added to make it easier to see how well the reflections are set up. 
diff --git a/Engine/Documentation/Source/Engine/Rendering/Materials/ExpressionReference/ExpressionReference.INT.udn b/Engine/Documentation/Source/Engine/Rendering/Materials/ExpressionReference/ExpressionReference.INT.udn
--- a/Engine/Documentation/Source/Engine/Rendering/Materials/ExpressionReference/ExpressionReference.INT.udn
+++ b/Engine/Documentation/Source/Engine/Rendering/Materials/ExpressionReference/ExpressionReference.INT.udn
@@ -183,21 +183,7 @@ Certain material expressions are parameters, meaning their values can be modifie
 		[RELATIVE:Engine/Rendering/Materials/ExpressionReference/Particles]
 	[/PARAM]
 [/OBJECT]
-[COMMENT:none]
--
-[OBJECT:TopicCompact]
-	[PARAM:image]
-		![%Engine/Rendering/Materials/ExpressionReference/Parameters:title%](exp_parameter.png)
-	[/PARAM]
-	[PARAM:title]
-		SpeedTree
-	[/PARAM]
-	[PARAM:path]
-		[RELATIVE:Engine/Rendering/Materials/ExpressionReference/SpeedTree]
-	[/PARAM]
-[/OBJECT]
--
-[/COMMENT]
+
 [OBJECT:TopicCompact]
 	[PARAM:image]
 		![%Engine/Rendering/Materials/ExpressionReference/Landscape:title%](exp_landscape.png)
@@ -420,14 +406,7 @@ This list exists as a mere convenience for those who would like to bookmark this
 * [ParticleSubUV](Engine/Rendering/Materials\ExpressionReference\Particles#ParticleSubUV)
 * [TextureSampleParameterSubUV](Engine/Rendering/Materials\ExpressionReference\Particles#TextureSampleParameterSubUV)
 [/REGION]
-[COMMENT:none]
--
-[**SpeedTree**](Engine/Rendering/Materials\ExpressionReference\SpeedTree)
-[REGION:indentedlist]
-* [SpeedTree](Engine/Rendering/Materials\ExpressionReference\SpeedTree#SpeedTree)
-[/REGION]
--
-[/COMMENT]
+
 [**Terrain**](Engine/Rendering/Materials\ExpressionReference\Landscape)
 [REGION:indentedlist]
 * [LanscapeLayerBlend](Engine/Rendering/Materials\ExpressionReference\Landscape#LandscapeLayerBlend)
diff --git a/Engine/Documentation/Source/Engine/Rendering/Materials/ExpressionReference/Vector/VectorExpressions.INT.udn b/Engine/Documentation/Source/Engine/Rendering/Materials/ExpressionReference/Vector/VectorExpressions.INT.udn
--- a/Engine/Documentation/Source/Engine/Rendering/Materials/ExpressionReference/Vector/VectorExpressions.INT.udn
+++ b/Engine/Documentation/Source/Engine/Rendering/Materials/ExpressionReference/Vector/VectorExpressions.INT.udn
@@ -73,20 +73,6 @@ The **CameraVector** expression outputs a three-channel vector value representin
 [REGION:warning]
 This expression has been deprecated in Unreal Engine 4 as lighting calculations are now deferred.
 [/REGION:warning]
-[COMMENT:none]
-
-The LightVector expression is a three-channel vector value representing the direction of the light with respect to the surface.
-
-**Example Usage:** LightVector can be used to calculate custom lighting algorithms, such as a simple Lambert lighting model.
-
-![LightVectorExample.jpg](LightVectorExample.jpg)
-
-[REGION:note]
-	The LightVector expression is not valid while evaluating the Emissive output, and can only be used in LightFunction or DeferredDecal materials.
-[/REGION]
-
-[/COMMENT]
-
 
 ## Object Bounds
 
diff --git a/Engine/Documentation/Source/Engine/Rendering/Materials/Functions/Reference/FunctionReference.INT.udn b/Engine/Documentation/Source/Engine/Rendering/Materials/Functions/Reference/FunctionReference.INT.udn
--- a/Engine/Documentation/Source/Engine/Rendering/Materials/Functions/Reference/FunctionReference.INT.udn
+++ b/Engine/Documentation/Source/Engine/Rendering/Materials/Functions/Reference/FunctionReference.INT.udn
@@ -180,18 +180,4 @@ Below is breakdown of the various **Material Functions** included within Unreal
 [/EXCERPT:Categories]
 
 
-[COMMENT:none]
 
-[OBJECT:TopicCompact]
-	[PARAM:image]
-		![%Engine/Rendering/Materials/Functions/Reference/ComplexMotion:title%](ComplexMotion.png)
-	[/PARAM]
-	[PARAM:title]
-		%Engine/Rendering/Materials/Functions/Reference/ComplexMotion:title%
-	[/PARAM]
-	[PARAM:path]
-		[RELATIVE:Engine/Rendering/Materials/Functions/Reference/ComplexMotion]
-	[/PARAM]
-[/OBJECT]
-
-[/COMMENT]
diff --git a/Engine/Documentation/Source/Engine/Rendering/ParticleSystems/ParticleLights/ParticleLights.INT.udn b/Engine/Documentation/Source/Engine/Rendering/ParticleSystems/ParticleLights/ParticleLights.INT.udn
--- a/Engine/Documentation/Source/Engine/Rendering/ParticleSystems/ParticleLights/ParticleLights.INT.udn
+++ b/Engine/Documentation/Source/Engine/Rendering/ParticleSystems/ParticleLights/ParticleLights.INT.udn
@@ -88,12 +88,4 @@ Remember: the **Light** module can only be used on CPU particles, so none of the
 
 
 
-[COMMENT:none]
 
-## Related
-
-[INCLUDE:Engine/Rendering/LightingAndShadows/LightProperties#LightRef005]
-
-[INCLUDE:Engine/Rendering/LightingAndShadows/LightProperties#LightRef007]
-
-[/COMMENT]
diff --git a/Engine/Documentation/Source/Engine/Rendering/PostProcessEffects/DepthOfField/Cinematic/CinematicDepthofField.INT.udn b/Engine/Documentation/Source/Engine/Rendering/PostProcessEffects/DepthOfField/Cinematic/CinematicDepthofField.INT.udn
new file mode 100644
--- /dev/null
+++ b/Engine/Documentation/Source/Engine/Rendering/PostProcessEffects/DepthOfField/Cinematic/CinematicDepthofField.INT.udn
@@ -0,0 +1,476 @@
+title:Cinematic Depth of Field
+description:The following depth of field methods provides a cinematic look that closely resembles photography and film for desktop and console platforms using the Deferred rendering path.
+crumbs:
+availability:Public
+version:4.20
+parent:Engine/Rendering/PostProcessEffects/DepthOfField
+order:1
+type:overview
+tags:Post Process
+tags:Depth of Field
+topic-image:DOF_Cine_Topic.png
+
+The following depth of field method provides a cinematic look that closely resembles photography and film for desktop and console platforms using the Deferred Shading renderer and Clustered Forward renderer.
+
+## Cinematic
+
+[REGION:note]
+This depth of field method replaces **Circle** DOF as the default cinematic method used in Unreal Engine 4, version 4.20. It is enabled by default on supported hardware.
+[/REGION]
+
+**Cinematic** DOF closely matches real-world cameras, similarly to Circle and Bokeh DOF, you can see circular shaped Bokeh with sharp [High Dynamic Range](Engine/Rendering/HDRDisplayOutput) (HDR) content. This method uses a procedural Bokeh simulation that provides dynamic resolution stability and alpha channel support while being faster, scalable, and optimized for projects developing on desktop and consoles.
+
+[OBJECT:ComparisonSlider]
+	[PARAM:before]
+	![Depth of Field Disabled](image_0.png)
+	[/PARAM]
+	[PARAM:after]
+	![Cinematic Depth of Field](image_1.png)
+	[/PARAM]
+	[PARAMLITERAL:position]
+	bottom
+	[/PARAMLITERAL]
+[/OBJECT]
+
+### Focusing the Subject of the Shot
+
+The key to achieving an aesthetically pleasing depth of field effect is all in how the subject(s) of the shot is focused. There are three core factors that affect setting up Depth of Field for any given shot:
+
+* Decide on the **Focal Length** of the lens.
+
+* Choose an appropriate **Aperture** (f-stop number).
+
+* Choose a **Distance to your Subject(s) **from the camera.
+
+Let's break down the elements that make up the camera and scene being captured in order to understand what's going on when we adjust these settings:
+
+![image_2.png](image_2.png)
+
+1. Points in the Scene
+
+    1. Red = Out of Focus
+
+    2. Blue = Perfectly in Focus
+
+    3. Green = Marginally in Focus
+
+2. Focus Distance to Subject of the Shot (Blue Dot)
+
+3. Lens of the camera
+
+4. Diaphragm of the Lens measured in F-stop
+
+5. Focal Length of the Lens
+
+6. Filmback / Image Sensor
+
+7. Rendered Image Result
+
+[REGION:note]
+Take note that the rendered image on the right is inverted through the lens of the camera.
+[/REGION]
+
+Here we have the multiple points (1) representing objects captured by the camera to a defined **Focus Distance** (2), in this case, the blue subject. The **Aperture** (4) defines how blurred objects in the foreground and background will be that are not in focus, and, finally, the **Focal Length** (5) of the lens controls the field of view, or how zoomed in the image will be. 
+
+#### Aperture
+
+The **Aperture** defines how sharp or blurred the foreground and background are based on the diameter of the **Diaphragm** which is controlled in f-stop. 
+
+![image_3.png](image_3.png)
+
+1. Lens of the Camera
+
+2. Aperture Diaphragm of the Lens measured in F-stop
+
+This diagram demonstrates the Diaphragm (2) blocking light passing through the Lens (1). The size of the Aperture (or f-stop number) sets how much light passes through the Lens in turn controlling how much the foreground and background is blurred  based on the Focus Distance of the shot. 
+
+There are two elements that make up the Aperture; the **F-stop** number and the **Diaphragm**.
+
+![image_4.png](image_4.png)
+
+1. The** Aperture** is the diameter of the opening through which light passes. The Aperture is defined by [REDACTED_PERSON] divided by an F-stop number.
+
+2. The **Diaphragm** is the mechanical structure made up of multiple blades used to block light. It opens and closes based on the supplied F-stop.
+
+The len's Aperture is defined by the diameter of the Diaphragm which decreases as the F-stop number increases controlling how shallow or wide the depth of field effect is. See Focus Distance for a diagram showing this effect. 
+
+To demonstrate how the Aperture works, drag the slider to change the F-stop number between 1.4, 2.8, and 5.6. Smaller F-stop numbers create a shallower DOF with more background and foreground blur. Larger F-stops increase the depth of DOF with less blur. Note that only the F-stop is being changed here, the Focus Distance and Focal Length remain constant at 7m and 75mm, respectively.
+
+[OBJECT:SliderSequence]
+	[PARAMLITERAL:id]
+	aperture
+	[/PARAMLITERAL]
+	[PARAMLITERAL:max]
+	2
+	[/PARAMLITERAL]
+	[PARAMLITERAL:caption]
+	Drag the slider to change the Aperture F-stop.
+	[/PARAMLITERAL]
+	[PARAMLITERAL:ext]
+	jpg
+	[/PARAMLITERAL]
+	[PARAM:image_list]
+	![aperture_0.png](aperture_0.png)
+	![aperture_1.png](aperture_1.png)
+	![aperture_2.png](aperture_2.png)
+	[/PARAM]
+[/OBJECT]
+
+[REGION:note]
+Normally when you would adjust the Aperture settings of a real camera, you'd have to also adjust the Exposure settings at the same time to maintain same light intensity received by the filmback / image sensor. However, in UE4, this is not a real camera, so adjusting the F-stop and Diaphragm does not control the light intensity. 
+[/REGION]
+
+#### Focal Length
+
+The **Focal Length** is the distance from the center of the lens to the **Filmback** (or image sensor) measured in millimeters (mm). The focal length of the lens your image: if the filmback grows, the field of view (FOV) will grow with it. If you were to decrease the Focal Length, the FOV would increase too. 
+
+![image_8.png](image_8.png)
+
+1. Lens of the camera.
+
+2. Focal Length
+
+To demonstrate how this works, drag the slider to change the Focal Length of the lens between 50mm, 75mm, and 100mm. Take notice of how effectively it changes the Field of View (FOV) of the shot; as the focal length increase, the FOV decreases. You can think of the focal length as being similar to zooming in except you also zoom in on the out-of-focus areas of the shot as well, increasing the blur for the foreground and background areas of the image. 
+
+Note that only the Focal Length is being changed here, the Focus Distance and F-stop remain constant at 7m and f/2.8, respectively.
+
+[OBJECT:SliderSequence]
+	[PARAMLITERAL:id]
+	focal
+	[/PARAMLITERAL]
+	[PARAMLITERAL:max]
+	2
+	[/PARAMLITERAL]
+	[PARAMLITERAL:ext]
+	jpg
+	[/PARAMLITERAL]
+	[PARAM:image_list]
+	![focal_0.png](focal_0.png)
+	![focal_1.png](focal_1.png)
+	![focal_2.png](focal_2.png)
+	[/PARAM]
+[/OBJECT]
+
+[REGION:note]
+Note that the FOV of the camera is controllable from the Cine Camera Actor through the filmback and focal length, since this is a physically accurate camera. However, from the Camera Actor you can directly control the FOV angle in degrees since it does not have a Focal Length setting. The focal length used for DOF is automatically derived by the renderer from the FOV angle, and the filmback setting.
+[/REGION]
+
+#### Focus Distance
+
+The **Focus Distance** is the distance from the center of the camera lens to the subject of the shot that is put into focus, creating a focal plane. The closer the camera is to the subject, the more the background will end up out-of-focus. 
+
+![image_12.png](image_12.png)
+
+1. Points in the Scene
+
+    1. Red = Out of Focus
+
+    2. Blue = Perfectly in Focus
+
+    3. Green = Marginally in Focus
+
+2. Focus Distance to Subject (Blue Dot)
+
+3. Lens of the Camera
+
+This diagram demonstrates the different object of the scene (1) that make up the scene; background, subject(s), and foreground. The distance from the Lens (3) to the focal subject (3) - in this case the blue dot - is the Focus Distance. Areas in the background and foreground are blurred depending on the F-stop used and the Focal Length of the camera. 
+
+We can visualize the DOF Layers by using the **Show** dropdown in the Viewport to select **Advanced > Depth of Field Layers**. In this instance, Green represents the foreground objects, Black is the area where objects are focused, and Blue the background objects.
+
+[OBJECT:ComparisonSlider]
+	[PARAM:before]
+	![Rendered Scene](image_13.png)
+	[/PARAM]
+	[PARAM:after]
+	![Depth of Field Layers Visualization](image_14.png)
+	[/PARAM]
+	[PARAMLITERAL:position]
+	bottom
+	[/PARAMLITERAL]
+[/OBJECT]
+
+To demonstrate how this works, drag the slider to change the Focus Distance between 4m, 7m, and 10m. Adjusting the Focus Distance adjusts the Focal Plane (purple dotted line) in the scene, indicating where the in focus region of the shot is. The blue reference point in the diagram uses the Character as the focal subject here. As camera and diagram shift focus, objects in the foreground and background will shift too. 
+
+Starting with 4m, the Police car in the foreground is in sharp focus and the character is out of focus in the background. The blue reference in the diagram indicates this with its focal point intersection in front of the filmback. At 7m, the character is completely in focus with background and foreground objects out of focus. At 10m, the construction barriers are in sharp focus and the character and Police car are out of focus in the foreground. The blue reference in the diagram indicates this since the intersection point extends beyond the filmback. 
+
+Through all these Focus Distance changes, the DOF effect increases or decreases depending on an object being in the foreground or background. Note that the Focus Distance is being changed here, the F-stop and Focal Length remain constant at f/1.4 and 75mm, respectively.
+
+[OBJECT:SliderSequence]
+	[PARAMLITERAL:id]
+	distance
+	[/PARAMLITERAL]
+	[PARAMLITERAL:max]
+	2
+	[/PARAMLITERAL]
+	[PARAMLITERAL:ext]
+	jpg
+	[/PARAMLITERAL]
+	[PARAM:image_list]
+	![distance_0.png](distance_0.png)
+	![distance_1.png](distance_1.png)
+	![distance_2.png](distance_2.png)
+	[/PARAM]
+[/OBJECT]
+
+### Procedural Diaphragm Simulation
+
+The job of Depth of Field is to apply focal importance to a subject(s) in your scene. Just as importantly, areas that are out-of-focus can be as important. Unreal Engine 4's depth of field simulates the lens diaphragm by enabling you to aesthetically control the shape of the Bokeh (out-of-focus areas) by adjusting the number and curvature of diaphragm blades. 
+
+![image_18.png](image_18.png)
+
+[REGION:note]
+For performance reasons, the diaphragm simulation is only supported on Engine Scalability settings of **Epic **and **Cinematic**. Lower quality settings will fall back to circle Bokeh shapes of equal area to ensure similar behavior with auto exposure functionality, in turn disabling the diaphragm blade simulation. 
+[/REGION]
+
+#### Number of Diaphragm Blades
+
+The **Diaphragm **(or Iris) is the mechanism made up of blades that can adjust to aesthetically inform the look of Bokeh shapes. Using a lower number of blades will result in a polygon-shaped diaphragm whereas a higher number of blades improves the curvature of the diaphragm to appear more natural with rounded out-of-focus Bokeh. 
+
+[OBJECT:SliderSequence]
+	[PARAMLITERAL:id]
+	blades
+	[/PARAMLITERAL]
+	[PARAMLITERAL:max]
+	3
+	[/PARAMLITERAL]
+	[PARAMLITERAL:ext]
+	jpg
+	[/PARAMLITERAL]
+	[PARAM:image_list]
+	![blades_0.png](blades_0.png)
+	![blades_1.png](blades_1.png)
+	![blades_2.png](blades_2.png)
+	![blades_3.png](blades_3.png)
+	[/PARAM]
+[/OBJECT]
+
+The number of diaphragm blades directly influences the shape of the Bokeh when the **Maximum Aperture** is increased (minimum F-stop number is decreased).
+
+[OBJECT:SliderSequence]
+	[PARAMLITERAL:id]
+	blades_zoom
+	[/PARAMLITERAL]
+	[PARAMLITERAL:max]
+	3
+	[/PARAMLITERAL]
+	[PARAMLITERAL:ext]
+	jpg
+	[/PARAMLITERAL]
+	[PARAM:image_list]
+	![blades_zoom_0.png](blades_zoom_0.png)
+	![blades_zoom_1.png](blades_zoom_1.png)
+	![blades_zoom_2.png](blades_zoom_2.png)
+	![blades_zoom_3.png](blades_zoom_3.png)
+	[/PARAM]
+[/OBJECT]
+
+#### Aperture and Max Aperture
+
+The **Aperture** is the opening of the lens controlled by the diameter of the **Diaphragm** which is measured in **f-stops**. The amount of light allowed to pass through the lens is controlled by the size of the aperture. Adjusting the size of the aperture sets the focal plane to adjust how in focus or out of focus areas outside of the focal region are. 
+
+The following chart shows the correlation between aperture size (f-stop number), max aperture (min f-stop number) and depth of field effects. 
+
+![image_27.png](image_27.png)
+
+A larger aperture (lower f-stop number) has a shallow focal region, causing the foreground and background to become more out-of-focus. A small aperture (higher f-stop number) has a wider focal region that includes more objects in focus in the foreground and background.
+
+[OBJECT:SliderSequence]
+	[PARAMLITERAL:id]
+	max_aperture
+	[/PARAMLITERAL]
+	[PARAMLITERAL:max]
+	5
+	[/PARAMLITERAL]
+	[PARAMLITERAL:caption]
+	Aperture F-stops: 1.4, 2.0, 2.8, 4.0, 5.6, 8.0
+	[/PARAMLITERAL]
+	[PARAMLITERAL:ext]
+	jpg
+	[/PARAMLITERAL]
+	[PARAM:image_list]
+	![max_aperture_0.png](max_aperture_0.png)
+	![max_aperture_1.png](max_aperture_1.png)
+	![max_aperture_2.png](max_aperture_2.png)
+	![max_aperture_3.png](max_aperture_3.png)
+	![max_aperture_4.png](max_aperture_4.png)
+	![max_aperture_5.png](max_aperture_5.png)
+	[/PARAM]
+[/OBJECT]
+
+While the Aperture can be as small as you'd like, it can only ever be as large as the maximum size of the lens. The **Maximum Aperture** (measured in f-stop) defines how large the lens can be. This also defines the curvature of the blades that make up the diaphragm by limiting the size that the aperture can open; limiting the size of the aperture shortens the effective focal region which limits how out-of-focus areas can be. 
+
+[OBJECT:SliderSequence]
+	[PARAMLITERAL:id]
+	max_aperture_blades
+	[/PARAMLITERAL]
+	[PARAMLITERAL:max]
+	3
+	[/PARAMLITERAL]
+	[PARAMLITERAL:caption]
+	5 Diaphragm Blades with Different Max Aperture Values
+	[/PARAMLITERAL]
+	[PARAMLITERAL:ext]
+	jpg
+	[/PARAMLITERAL]
+	[PARAM:image_list]
+	![max_aperture_blades_0.png](max_aperture_blades_0.png)
+	![max_aperture_blades_1.png](max_aperture_blades_1.png)
+	![max_aperture_blades_2.png](max_aperture_blades_2.png)
+	![max_aperture_blades_3.png](max_aperture_blades_3.png)
+	[/PARAM]
+[/OBJECT]
+
+In this example, the **Number of Diaphragm Blades** is set to 5. With a smaller Aperture, the blades of the diaphragm become more visible in the Bokeh shapes. Using a larger Aperture (f-stop = minimal f-stop) the Bokeh shapes will begin to become circular.
+
+In this example, a Cine Camera Actor is being used with following values:
+
+* **Min FStop**: 1.4
+* **Max FStop**: 4.0
+* **Diaphragm Blade Count**: 5
+* **Current Aperture**: 1.4 - 4.0
+
+With these values, the F-stop can only accept values between **1.4** and **4.0**. As the depth of field effect is widened with a larger F-stop (smaller aperture), the diaphragm blades become more visible in the Bokeh shapes.
+
+[REGION:note]
+Keep in mind that the aperture is not controlling the light intensity. This choice has been made so that you don't have to continually adjust the exposure at the same time like with a real camera. 
+[/REGION]
+
+### Available Settings
+
+#### Cinematic Camera
+
+The following settings are specific to the Cinematic Camera Actor. Note that you have access to the Camera and Depth of Field settings as well. 
+
+| Property | Description |
+| --- | --- |
+| Filmback Settings | |
+| Sensor Width | Sets the horizontal size of the filmback or digital sensor (measured in millimeters (mm)). |
+| Lens Settings | |
+| Min Focal Length | Sets the minimum focal length for the lens (measured in millimeters (mm)). |
+| Max Focal Length | Sets the maximum focal length for the lens (measured in millimeters (mm)). |
+| Min FStop | The maximum aperture for this lens. For example, 2.8 for an f/2.8 lens. This will also define the curvature of the blades of the diaphragm. |
+| Max FStop | The minimum aperture for this lens. For example, 8.0 for an f/8.0 lens. |
+| Diaphragm Blade Count | The number of blades that make up the diaphragm mechanism. |
+| Focus Settings | |
+| Focus Method | Select the type of focus method to use to control camera focus; None, Manual, or Tracking. |
+| Manual Focus Distance | Set the manually-controlled focus distance. This is only available when the Focus Method is set to Manual. |
+| Draw Debug Focus Plane | Enable to draw a translucent plane at the current focus depth. This enables a visual way to track where focus is for the shot. |
+| Debug Focus Plane Color | Set a color for Draw Debug Focus Plane when it is enabled. |
+| Smooth Focus Changes | Enable this to use interpolation to smooth out changes in focus distance. When disabled, the focus changes will be instantaneous. |
+| Focus Smoothing Interp Speed | Controls the interpolation speed when smoothing focus distance changes. This option is ignored if Smooth Focus Changes is not enabled. |
+| Focus Offset | This adds additional offset to focus depth which can be used to manually tweak focus if the chosen focus method needs adjustment. |
+| Current Focal Length | Controls the current focal length of the camera which controls the field of view (FoV) and zoom. |
+| Current Aperture | Sets the current aperture size based on an f-stop number. Note that this setting only accepts values that are within the Min FStop and Max FStop range. |
+| Current Focal Distance | Displays a read-only value that is used by the Focus Settings. |
+| Current Horizontal FOV | Displays a read-only value that is used by [REDACTED_PERSON] and Filmback Settings. |
+
+#### Post Process Volume and Camera Actor
+
+The following settings are available with the Cine Camera, Camera, and Post Process Volumes.
+
+| Property | Description |
+| --- | --- |
+| Camera Settings | |
+| Aperture (F-stop) | Defines the opening of the camera lens. Aperture is 1/f-stop with typical lens' going down to f/1.2 (larger opening). Larger numbers will reduce the DoF effect. |
+| Maximum Aperture (minimum f-stop) | Defines the maximum opening of the camera lens to control the curvature of the diaphragm. Set this to 0 to get straight blades. |
+| Number of Diaphragm Blades | Defines the number of blades of the diaphragm within the lens. Values between 4 and 16 can be used. |
+| Depth of Field Settings | |
+| Focal Distance | The distance in which the depth of field effect should be sharp. This value is measured in Unreal Units (cm). |
+| Depth Blur km for 50% | Defines the distance at which the pixel is is blurred with a radius half of Depth blur Radius. This is particularly useful to emulate cheaply atmospheric scattering. |
+| Depth Blur Radius | Radius in pixels at 1080p resolution to apply according to distance from camera to emulate atmospheric scattering. |
+
+### Useful Console Variables for Optimization
+
+The challenge of DOF is it may varies according to your content and how important artistically it has been used. For that purpose, the DOF implementation offers a variety of console variables ready to be customised per platforms, for you to control the budget you would like to allocate to your project on weaker platforms according to your priorities.
+
+Below are some console variables that can be useful for binding Cinematic DoF performance. Note that these are a good starting point and additional console variables can be found under r.DOF.*.
+
+* r.DOF.Kernel.MaxBackgroundRadius - The maximum size of the background blurring radius in horizontal screen space. 
+* r.DOF.Kernel.MaxForegroundRadius - The maximum size of the foreground blurring radius in horizontal screen space. 
+* r.DOF.Scatter.MaxSpriteRatio - The maximum ratio of scattered pixel quads as sprites. Useful to control DOF's scattering upperbound. 
+
+## Bokeh DOF method (Legacy)
+
+[REGION:warning]
+This depth of field method will be deprecated in a future release.
+[/REGION]
+
+[EXCERPT:DepthOfField_Bokeh]
+**Bokeh** DOF refers to the name of the shape that can be seen in photos or movies when objects are out of focus. This method renders a textured quad for each pixel using a user-specified texture to define a shape with which to reproduce the effect of a camera lens. 
+
+The implementation used only uses half-resolution to perform this intensive effect. It attempts to save render performance in areas where the effect is not noticeable by using Adaptive Depth of Field. Bokeh DOF is more costly than other DOF methods available in UE4, making this a prime candidate for cinematics and showcases where appealing visuals often outweigh any performance concerns in those situations. 
+[/EXCERPT:DepthOfField_Bokeh]
+
+[OBJECT:ComparisonSlider]
+	[PARAM:before]
+	![Depth of Field Disabled](image_38.png)
+	[/PARAM]
+	[PARAM:after]
+	![Bokeh Depth of Field](image_39.png)
+	[/PARAM]
+	[PARAMLITERAL:position]
+	bottom
+	[/PARAMLITERAL]
+[/OBJECT]
+
+### Adaptive Depth of Field for Bokeh DOF
+
+When using Bokeh DOF, it renders at quarter-resolution (half-resolution in each direction) by default for performance reasons. The downsampling that happens is barely noticeable, in most cases, and is perfectly acceptable; however, it can lead to artifacts and unwanted results in some situations. **Adaptive Depth of Field** resolves these types of artifacts where possible.
+
+[OBJECT:ComparisonSlider]
+	[PARAM:before]
+	![Adaptive DOF without Downsampling](image_40.png)
+	[/PARAM]
+	[PARAM:after]
+	![Adaptive DOF with Downsampling](image_41.png)
+	[/PARAM]
+	[PARAMLITERAL:position]
+	bottom
+	[/PARAMLITERAL]
+[/OBJECT]
+
+The characters being blurred in the background appear blocky when using only the downsampled DOF technique. Also, some artifacts appear around the horns of the character in the foreground. Adaptive DOF eliminates these types of artifacts, and the background characters have a much smoother appearance. 
+
+Adaptive DOF can be visualized by using the **Adaptive Depth of Field** show flag in the Level Viewport under **Show > Visualize**. When enabled, an overlay is shown where the downsampled effect is being used (green) and where the full resolution effect is being used (red). The regular scene color is shown where no blurring is being applied. 
+
+![image_42.png](image_42.png)
+
+In general, you want the Adaptive DOF to display all green. The more red shown by the visualization, the more expensive the scene is to render each frame.
+
+[OBJECT:ComparisonSlider]
+	[PARAM:before]
+	![Adaptive DOF mostly optimized](image_43.png)
+	[/PARAM]
+	[PARAM:after]
+	![Adaptive DOF unoptimized](image_44.png)
+	[/PARAM]
+	[PARAMLITERAL:position]
+	bottom
+	[/PARAMLITERAL]
+[/OBJECT]
+
+### Available Settings
+
+The following settings are available with the Cine Camera, Camera, and Post Process Volumes, located under the **Lens** tab in the **Depth of Field** section.
+
+| Property | Description |
+| --- | --- |
+| Focal Distance | The distance in which the depth of field effect should be sharp. This value is measured in Unreal Units (cm). |
+| Near Transition Region | The distance in Unreal Units from the focal region on the side nearer to the camera over which the scene transitions from focused to blurred when using Bokeh or Gaussian DOF. |
+| Far Transition Region | The distance in Unreal units from the focal region on the side farthest from the camera over which the scene transitions from focused to blurred when using Bokeh or Gaussian DOF. |
+| Scale | An overall scaling factor for Bokeh-based blurring. |
+| Max Bokeh Size | The maximum size (in percent of the view width) of the blur for the Bokeh-based DOF. Note that performance scales with the size. |
+| Shape Texture | A texture that defines the shape of the Bokeh when objects get out of focus. Note that these cannot be blended with other Post Process Volumes. |
+| Occlusion | Controls how much the blurred geometry extends beyond its usual silhouette and its opacity. A value of 0.18 gives a fairly natural occlusion result. A value of 0.4 may be necessary to solve layer color leaking issues. Very small values (less than 0.18) generally result in negating the blur effect, though can work well when objects are very near to the camera. |
+| Color Threshold | Controls the threshold where adaptive DOF switches to using full-resolution based on color. Smaller values cause more of the scene to use full-resolution processing. |
+| Size Threshold | Controls the threshold where adaptive DOF switches to using full-resolution based on size. Larger values cause more of the scene to use full-resolution processing. |
+
+## Resources
+
+* [Diaphragm (Optics)](http://en.wikipedia.org/wiki/Diaphragm_(optics))
+
+* [Aperture](http://en.wikipedia.org/wiki/Aperture)
+
+* [F-Number](http://en.wikipedia.org/wiki/F-number)
+
+* [Understanding the Factors That Affect Depth of Field](http://photography.tutsplus.com/articles/understanding-the-factors-that-affect-depth-of-field--photo-6844)
diff --git a/Engine/Documentation/Source/Engine/Rendering/PostProcessEffects/DepthOfField/DepthOfField.INT.udn b/Engine/Documentation/Source/Engine/Rendering/PostProcessEffects/DepthOfField/DepthOfField.INT.udn
--- a/Engine/Documentation/Source/Engine/Rendering/PostProcessEffects/DepthOfField/DepthOfField.INT.udn
+++ b/Engine/Documentation/Source/Engine/Rendering/PostProcessEffects/DepthOfField/DepthOfField.INT.udn
@@ -4,215 +4,81 @@ Title: Depth of Field
 Description: Simulating focus by blurring the scene depending on depth.
 Type: Overview
 SkillLevel:Beginner
-Version:4.16
+Version:4.20
 Parent:Engine/Rendering/PostProcessEffects
 Order:5
 Tags:Post Process
-Topic-image:DoF_topic.png
+Topic-image:Engine\Rendering\PostProcessEffects/DepthOfField/Cinematic/DOF_Cine_Topic.png
 
-
-
-
-[REGION:fullwidth]
-![Depth of Field](dof.png)
-[/REGION]
+![image_0.png](image_0.png)
 
 [EXCERPT:DepthOfField_Desc]
-**Depth of Field** (DoF) applies a blur to the scene based on distance in front of or behind a focal point. This is to simulate what happens in real-world cameras. The effect can be used to draw the viewer's attention and to
-make the rendering appear more like a photograph or like a movie.
+**Depth of Field **(DOF) applies a blur to the scene based on distance in front of or behind a focal point, similarly to what happens in real-world cameras. The effect can be used to draw the viewer’s attention on specific subject of shot based on dept, and adding an aesthetic to make the rendering appear more like a photograph or like a movie. 
 [/EXCERPT:DepthOfField_Desc]
 
-## Depth of Field Types
-
-There are three methods available for performing depth of field effects in Unreal Engine 4: Bokeh, Gaussian and Circle.
-
-[REGION:imagetable]
-| ![Bokeh DoF](bokeh.png)(w:300) | ![Gaussian DoF](gaussian.png)(w:300) | ![Circle DoF](Circle.png)(w:300) |
-| :----: | :----: | :----: |
-| Bokeh DoF | Gaussian DoF | Circle DoF |
-[/REGION]
-
-
-### Gaussian Depth of Field
 
-[EXCERPT:DepthOfField_Gaussian]
-Gaussian DoF blurs the scene using a standard Guassian blur. The Guassian method is fast and generally well-suited for use in-game, where performance is critical.
-[/EXCERPT:DepthOfField_Gaussian]
+# Depth of Field Types
 
-### Bokeh Depth of Field
+There are several methods of performing depth of field effects in Unreal Engine 4 that you can use. These have been broken into two categories:
 
-[EXCERPT:DepthOfField_Bokeh]
-Bokeh is the name of the shape that can be seen in photos or movies when objects are out of focus. Bokeh DoF is a method where a textured quad is rendered for each pixel using a texture to define the shape to reproduce the effect produced by camera lenses, which gives a filmic look to the scene. 
-The actual implementation does this performance intensive effect only in half resolution. It also tries to save render performance in areas where this is not that noticeable. See [Adaptive DOF](#adaptivedof) for more information.
+* **Cinematic **methods generally provide a cinematic or filmic look to depth of field effects. The adjustments for these methods also align more with common camera options available in photography and cinematography. These methods are also considered too expensive for use with Mobile platforms but work well for Desktop and Console platforms.
 
-The Bokeh method is far more costly performance-wise, making cinematics and showcases prime candidates for the Bokeh method as the appealing visuals often outweigh any performance concerns in those situations.
-[/EXCERPT:DepthOfField_Bokeh]
+* **Mobile** methods provide optimized and lower-cost depth of field options considered acceptable for Mobile platforms. 
 
-### Circle Depth of Field
+[DIR(output:"topic" parent:"Engine/Rendering/PostProcessEffects/DepthOfField" end:"1" org:"hierarchy")]
 
-[EXCERPT:DepthOfField_Circle]
-The Circle DoF is the newest addition and matches real world cameras more closely. Simialar to the Bokeh DoF you can see circular shaped Bokeh with sharp HDR content. Very large Bokeh are a weakness of this
-algorithm, they appear noisy and not as clear as the BokehDOF can do them.
-The method has good performance (quite better than Bokeh DoF) and good quality, especially in the transition areas.
+# Depth of Field Implementation
 
-The defaults of CircleDOF are intentionally weak to be more pleasant. To get the maximum effect you can tweak settings. Use a low Aperture for a large Bokeh, get close to an object and change the Field of View to be lower.
-Play with the focal distance to get some scene content in out of focus, ideally in front and behind the focal plane. Using showflag VisualizeDepthOfField might help to find the right settings.
+Depth of Field is broken up into three layers; Near, Far, and the Focal Region. Each of these is processed separately and then later composited together to achieve the final effect. Objects in the Near and Far layers (objects not in the Focal Region) are always fully blurred. These layers are blended with the non-blurred scene. 
 
-[/EXCERPT:DepthOfField_Circle]
+* Objects within the Focal Region (Black) use the non-blurred scene layer.
 
-## Implementation
+* Objects in the Near (Green) or Far (Blue) layers but outside any transition regions are fully blended to the blurred layer.
 
-Depth of Field is broken up into three layers - Near, Far, and the Focal Region - each of which are processed separately and then later composited together to achieve the final effect. Objects in the near and far layers - that is, objects not in the focal region - are always fully blurred. Then, these layers are blended with the non-blurred scene. 
+* Objects within the transition area are blended linearly between the non-blurred scene layer and their blurred layer based on their position within the transition region.
 
-* Objects within the focal region use the non-blurred scene layer.
-* Objects in the near or far layers, but outside any transition regions, are fully blended to the blurred layer. 
-* Objects within a transition region are blended linearly between the non-blurred scene layer and their blurred layer based on their position within the transition region.
+## Visualizing Depth of Field
 
+These layers, including transition regions, can be visualized using the **Depth of Field layers** showflag in the Level Viewport under **Show > Visualize**. 
 
-These layers, including the transition regions, can be visualized using the **Depth of Field Layers** show flag ( under **Show > Visualize** ). The near layer is shown in green, the far layer in blue, and the focal region is black.
 [OBJECT:ComparisonSlider]
 	[PARAM:before]
-	![Scene](dof_layers.png)(w:700)
+	![Scene View with Cinematic DoF](image_1.png)(w:998) 
 	[/PARAM]
 	[PARAM:after]
-	![Layer Visualization](dof_layers_visualize.png)(w:700)
+    ![Depth of Field Layers Vis with Cinematic DoF Enabled](image_2.png)(w:998) 
 	[/PARAM]
+	[PARAMLITERAL:position]
+	bottom
+	[/PARAMLITERAL]
 [/OBJECT]
 
-We further improved this visualization to also include a lot of interesting information:
-
-[REGION:imagetable]
-| ![Scene](NewVisualize.png) |
-| :----: |
-| VisualizeDOF |
-[/REGION]
-
-When using the CircleDOF you also can see a crosshair attached to the mouse cursor showing the computed circle of confusion (BokehSize) together with the depth and the pixel size in numbers.
-Note that the circle size might not always match to the actual effect as the implementation is not well suited / intended for large Bokeh shapes. We might want to improve on that in the future.
-
-## Properties
-
-| Property | Description |
-| -------- | ----------- |
-| **Method** | Selects the method, Bokeh, Gaussian or Circular, used for blurring the scene. |
-| **Depth Blur km for 50%** | todo |
-| **Depth Blur Radius** | Depth blur radius in pixels at 1920x |
-| **Aperture** | todo |
-| **Focal Distance** | The distance in Unreal Units from the camera that acts as the center of the region where the scene is perfectly in focus and no blurring occurs. [INCLUDE:#FocalDistance] |
-| **Focal Region** | The distance in Unreal Units beyond the focal distance where the scene is perfectly in focus and no blurring occurs. [INCLUDE:#FocalRegion] |
-| **Near Transition Range** | The distance in Unreal Units from the focal region on the side nearer to the camera over which the scene transitions from focused to blurred when using Gaussian-based DoF. [INCLUDE:#NearTransitionRange] |
-| **Far Transition Range** | The distance in Unreal Units from the focal region on the side farther from the camera over which the scene transitions from focused to blurred when using Gaussian-based DoF. [INCLUDE:#FarTransitionRange] |
-| **Scale** | An overall scaling factor for Bokeh-based blurring. [INCLUDE:#Scale] |
-| **Max Bokeh Size** | The maximum size (in percent of the view width) of the blur for Bokeh-based DoF (**Note:** performance cost scales with size*size). An approximate representation of the function used to calculate the size of the Bokeh at any distance is shown below: [INCLUDE:#MaxBokehSize] |
-| **Near Blur Size** | The maximum size (in percent of the view width) of the near blur for Gaussian-based DoF. (**Note:** performance cost scales with size) [INCLUDE:#NearBlurSize] |
-| **Far Blur Size** | The maximum size (in percent of the view width) of the far blur for Gaussian-based DoF. (**Note:** performance cost scales with size) [INCLUDE:#FarBlurSize] |
-| **Shape** | A texture that defines the shape of the Bokeh when objects get out of focus (cannot be blended). |
-| **Occlusion** | Controls how much the blurred geometry extends beyond its usual silhouette and its opacity. A value of 0.18 gives a fairly natural occlusion result. A value of 0.4 may be necessary to solve layer color leaking issues. Very small values (less than 0.18) generally result in negating the blur effect, though can work well when objects are very near to the camera. [INCLUDE:#Occlusion] |
-| **Color Threshold** | Controls the threshold where adaptive DoF switches to using full-resolution based on color. Smaller values cause more of the scene to use full-resolution processing. |
-| **Size Threshold** | Controls the threshold where adaptive DoF switches to using full-resolution based on size. Larger values cause more of the scene to use full-resolution processing. |
-| **Sky Distance** | Artificial distance to allow the skybox to be in focus (e.g. 200000), `<=0` to switch the feature off, only for GaussianDOF, can cost performance. |
-
-[COMMENT:none]
-
-[EXCERPT:FocalDistance]
-[REGION:imagetable]
-| ![DoF Focal Distance - Near](dof_focaldist_0.png)(w:240) | ![DoF Focal Distance - Mid](dof_focaldist_1.png)(w:240) | ![DoF Focal Distance - Far](dof_focaldist_2.png)(w:240) |
-| --- | --- | --- |
-| Focal Distance = 400.0 | Focal Distance = 1000.0 | Focal Distance = 2500.0 |
-[/REGION]
-[/EXCERPT:FocalDistance]
-
-[EXCERPT:FocalRegion]
-[REGION:imagetable]
-![Focal Region](dof_focalregion.png)
-[/REGION]
-[/EXCERPT:FocalRegion]
-
-[EXCERPT:NearTransitionRange]
-[REGION:imagetable]
-![Gaussian Near Layer](dof_gaussian_near.png)
-[/REGION]
-[/EXCERPT:NearTransitionRange]
-
-[EXCERPT:FarTransitionRange]
-[REGION:imagetable]
-![Gaussian Far Layer](dof_gaussian_far.png)
-[/REGION]
-[/EXCERPT:FarTransitionRange]
-
-[EXCERPT:MaxBokehSize]
-[REGION:imagetable]
-![Bokeh CoC](dof_bokeh_coc.png) 
-| ![DoF Bokeh Size - 2.5](dof_size_0.png)(w:240) | ![DoF Bokeh Size - 5.0](dof_size_1.png)(w:240) | ![DoF Bokeh Size - 10.0](dof_size_2.png)(w:240) |
-| --- | --- | --- |
-| Size = 2.5 | Size = 7.5 | Size = 15.0 |
-[/REGION]
-[/EXCERPT:MaxBokehSize]
-
-[EXCERPT:NearBlurSize]
-[REGION:imagetable]
-| ![DoF Near Blur Size - 2.5](dof_nearsize_0.png)(w:240) | ![DoF Near Blur Size - 5.0](dof_nearsize_1.png)(w:240) | ![DoF Near Blur Size - 10.0](dof_nearsize_2.png)(w:240) |
-| --- | --- | --- |
-| Size = 0.0 | Size = 0.25 | Size = 2.5 |
-[/REGION]
-[/EXCERPT:NearBlurSize]
-
-[EXCERPT:FarBlurSize]
-[REGION:imagetable]
-| ![DoF Far Blur Size - 2.5](dof_farsize_0.png)(w:240) | ![DoF Far Blur Size - 5.0](dof_farsize_1.png)(w:240) | ![DoF Far Blur Size - 10.0](dof_farsize_2.png)(w:240) |
-| --- | --- | --- |
-| Size = 0.0 | Size = 0.25 | Size = 2.5 |
-[/REGION]
-[/EXCERPT:FarBlurSize]
-
-[EXCERPT:Occlusion]
-[REGION:imagetable]
-| ![DoF Occlusion - 0.05](dof_occlusion_0.png)(w:240) | ![DoF Occlusion - 0.18](dof_occlusion_1.png)(w:240) | ![DoF Occlusion - 0.4](dof_occlusion_2.png)(w:240)
-| --- | --- | --- |
-| Occlusion = 0.05 | Occlusion = 0.18 | Occlusion = 0.4 |
-[/REGION]
-[/EXCERPT:Occlusion]
-
-[EXCERPT:Scale]
-[REGION:imagetable]
-| ![DoF Scale - 0.0](dof_scale_0.png)(w:240) | ![DoF Scale - 0.375](dof_scale_1.png)(w:240) | ![DoF Scale - 1.0](dof_scale_2.png)(w:240) |
-| --- | --- | --- |
-| Scale = 0.0 | Scale = 0.375 | Scale = 1.0 |
-[/REGION] 
-[/EXCERPT:Scale]
-
-[/COMMENT]
-
-## Adaptive DoF
-
-Bokeh depth of field is rendered at quarter-resolution (half-res in each direction) by default for performance reasons. In most cases, this downsampling is barely noticeable and is perfectly acceptable; however, it can lead to artifacting and unwanted results in some situations.
+Visualizing the **Depth of Field Layers** also includes useful information relevant to the DOF method being used, such as values that are currently set or when moving the mouse around the scene, the distance from the camera to the Actor is displayed next to the mouse cursor. 
 
-[OBJECT:ComparisonSlider]
-	[PARAM:before]
-	![Adaptive Resolution - Off](dof_adaptive_0.png)(w:700)
-	[/PARAM]
-	[PARAM:after]
-	![Adaptive Reoslution - On](dof_adaptive_1.png)(w:700)
-	[/PARAM]
-[/OBJECT]
+![image_3.png](image_3.png)
 
+## Using DOF in the Editor
 
-You can see that the characters being blurred in the background appear blocky when using only the downsampled DoF technique. Also, there is some artifacting around the horns of the character in the foreground. The artifacting is eliminated by using the adaptive DoF, and the background characters have a much smoother appearance.
+Using Depth of Field in Unreal Editor can be done in a few different ways; by placing a [Post Process Volume](https://docs.unrealengine.com/en-US/Engine/Rendering/PostProcessEffects), using a [Camera Actor](https://docs.unrealengine.com/en-us/Engine/Actors/CameraActors), or a [Cine Camera Actor](https://docs.unrealengine.com/en-us/Engine/Sequencer/HowTo/CineCameraActors). Each of these has access to the same [Depth of FIeld properties](http://www.linktopropertiesinpage.com) with the exception being that Cine Camera Actor has some additional industry standard settings for cameras and lenses. 
 
-The adaptive DoF can be visualized using the **Adaptive Depth of Field** show flag (under **Show > Visualize** ). When enabled, an overlay is displayed showing where the downsampled effect is being used (in green) and where the full-resolution effect is being used (in red). The regular scene color is shown where no blurring is being applied.
+The majority of settings used can be accessed under the **Lens** tab in the **Camera** and **Depth of Field** sections. When selecting a DOF method, the properties that affect that method will be editable with others being grayed out (even if they are enabled with a checkbox next to them). 
 
-[REGION:fullwidth]
-![Adaptive DoF Visualization](dof_adaptive_visualize.png)
-[/REGION]
+![image_4.png](image_4.png)
 
-In general, you want this to display all green. The more red shown by the visualization, the more expensive the scene is to render each frame.
+When using a [Cine Camera Actor](https://docs.unrealengine.com/en-us/Engine/Sequencer/HowTo/CineCameraActors), replacing properties that affect depth of field can be found under the **Current Camera Settings** in the **Lens Settings** section. 
 
-[OBJECT:ComparisonSlider]
-	[PARAM:before]
-	![Adaptive Resolution - Optimized](dof_adaptive_good.png)(w:700)
-	[/PARAM]
-	[PARAM:after]
-	![Adaptive Resolution - Expensive](dof_adaptive_bad.png)(w:700)
-	[/PARAM]
-[/OBJECT]
+![image_5.png](image_5.png)
+
+If you’re using a Camera or Cine Camera Actor, you can inhabit them using [Actor Piloting ](https://docs.unrealengine.com/en-us/Engine/UI/LevelEditor/Viewports/ActorLocking)in the Level Viewport by selecting the **Perspective** and choosing from one of the **Placed Cameras** in the scene. 
+
+![image_6.png](image_6.png)
+
+
+The Level Viewport will snap to the camera’s view and indicate that you are piloting and viewing what that camera sees.
+
+![image_7.png](image_7.png)
+
+Any properties changed in the Camera or Post Process Volume (if the camera is within it) will be shown in the Level Viewport. 
+
+As an example to get something similar to the shot above, the key is to use a low **Aperture (F-stop)** to get a large Bokeh, move the camera or viewport close to an object and change the **Field of View **(FoV) to be lower. Then, adjust the **Focal Distance** to get some scene content to be out-of-focus in front of and behind the focal region (also called a Focal Plane). 
 
diff --git a/Engine/Documentation/Source/Engine/Rendering/PostProcessEffects/DepthOfField/Mobile/MobileDepthofField.INT.udn b/Engine/Documentation/Source/Engine/Rendering/PostProcessEffects/DepthOfField/Mobile/MobileDepthofField.INT.udn
new file mode 100644
--- /dev/null
+++ b/Engine/Documentation/Source/Engine/Rendering/PostProcessEffects/DepthOfField/Mobile/MobileDepthofField.INT.udn
@@ -0,0 +1,102 @@
+title:Mobile Depth of Field
+description:The following depth of field method provides an optimized and low-cost depth of field effect for mobile platforms using the Forward rendering path. 
+crumbs:
+availability:Public
+version:4.20
+parent:Engine/Rendering/PostProcessEffects/DepthOfField
+order:2
+type:overview
+tags:Post Process
+tags:Depth of Field
+topic-image:DOF_Gaussian_Topic1.png
+
+The following depth of field method provides an optimized and low-cost depth of field effect for mobile platforms using the Forward rendering path. 
+
+## Gaussian
+
+[REGION:warning]
+The Deferred renderer will no longer support Gaussian Depth of Field method and will only be supported for the Mobile rendering path.
+[/REGION]
+
+The **Gaussian** depth of field (DOF) method blurs the scene using a standard [Gaussian blur](https://en.wikipedia.org/wiki/Gaussian_blur) (or smoothing) function. Gaussian DOF blur foreground and background uses a fixed size Gaussian blur kernel, which is very fast on lower-end hardware, such as mobile. This enables performance to be maintained with lower overhead costs where performance is critical. 
+
+[OBJECT:ComparisonSlider]
+	[PARAM:before]
+	![No Depth of Field](image_0.png)
+	[/PARAM]
+	[PARAM:after]
+	![Gaussian Depth of Field](image_1.png)
+	[/PARAM]
+	[PARAMLITERAL:position]
+	bottom
+	[/PARAMLITERAL]
+[/OBJECT]
+
+## Visualizing Depth of Field
+
+These layers, including transition regions, can be visualized using the **Depth of Field layers** show flag in the Level Viewport under **Show > Visualize**. 
+
+[OBJECT:ComparisonSlider]
+	[PARAM:before]
+	![Final Scene](image_2.png)
+	[/PARAM]
+	[PARAM:after]
+	![Depth of Field Visualization](image_3.png)
+	[/PARAM]
+	[PARAMLITERAL:position]
+	bottom
+	[/PARAMLITERAL]
+[/OBJECT]
+
+
+Visualizing the **Depth of Field Layers** also includes useful information relevant to the DOF method selected by including the properties and their currently set values.
+
+![image_4.png](image_4.png)
+
+## Focal Distance
+
+The **Focal Distance **represents the focal region and the angle of view that will be captured. A longer focal length has a shallow depth of field causing objects to be more blurred outside of the focal region while shorter focal lengths have a wider depth of field with sharper focus and fewer objects out of focus. The f-stop number can stay the same and changing the lens size adjusts how shallow or wide the depth of field effect is.
+
+[REGION:lightbox]
+[![image_5.png](image_5.png)](image_5.png)
+[/REGION]
+
+Once you’ve set the **Focal Distance** and the **Focal Region**, you can use the **Near Transition** and **Far Transition** to adjust the are distance between the focal region and the fully blurred scene. In addition, you can even adjust the **Near Blur Size** and the **Far Blur Size** to further tweak the look of Guassian DOF. 
+
+![image_6.png](image_6.png)
+
+[REGION:caption]
+Guassian DOF with adjusted values for Near/Far Transition and Blur Size to produce a nice depth of field effect. 
+[/REGION]
+
+In this example, the following values have been set to produce a soft depth of field effect for near and far regions. 
+
+* Focal Distance: **215**
+
+* Focal Region: **2500**
+
+* Near Transition: **500**
+
+* Far Transition: **400**
+
+* Near Blur Size: **2.0**
+
+* Far Blur Size: **2.5**
+
+## Available Settings
+
+
+| Property | Description |
+| --- | --- |
+| High Quality Gaussian DOF on Mobile | Enables HQ Gaussian on high-end mobile platforms. |
+| Focal Distance | The distance in which the depth of field effect should be sharp. This value is measured in Unreal Units (cm). |
+| Near Transition Region | The distance in Unreal Units from the focal region on the side nearer to the camera over which the scene transitions from focused to blurred when using Bokeh or Gaussian DOF. |
+| Far Transition Region | The distance in Unreal units from the focal region on the side farthest from the camera over which the scene transitions from focused to blurred when using Bokeh or Gaussian DOF.  |
+| Focal Region | An artificial region where all content is in focus, starting after the Focal Distance (measured in Unreal Units (cm)). |
+| Near Blur Size | The maximum size (in percent of the view width) of the near blur for Gaussian-based DOF. Note that performance cost scales with size. |
+| Far Blur Size | The maximum size (in percent of the view width) of the far blur for Gaussian-based DOF. Note that performance cost scales with size.  |
+| Sky Distance | Artificial distance to allow the skybox to be in focus (for example, 200000 units), where values less than 0 will disable the feature. Note that this feature can cost performance.  |
+| Vignette Size | Artificial circular mask to (near) blur content outside the radius in diameter in percent of screen width. This costs performance if the mask is used. |
+
+
+
diff --git a/Engine/Documentation/Source/Engine/Rendering/PostProcessEffects/UsingLUTs/UsingLUTs.INT.udn b/Engine/Documentation/Source/Engine/Rendering/PostProcessEffects/UsingLUTs/UsingLUTs.INT.udn
--- a/Engine/Documentation/Source/Engine/Rendering/PostProcessEffects/UsingLUTs/UsingLUTs.INT.udn
+++ b/Engine/Documentation/Source/Engine/Rendering/PostProcessEffects/UsingLUTs/UsingLUTs.INT.udn
@@ -9,8 +9,7 @@ Parent:Engine/Rendering/PostProcessEffects
 Related:Engine/Rendering/PostProcessEffects/ColorGrading
 Order:1000
 Tags:Post Process
-Topic-image:ColorGrading_Topic.png
-Social-image:ColorGrading_Overview.png
+
 
 
 ![](Engine/Rendering/PostProcessEffects/ColorGrading/ColorGrading_Overview.png)
diff --git a/Engine/Documentation/Source/Engine/Rendering/RenderTargets/BlueprintRenderTargets/HowTo/HeightFieldPainter/3/HFP3.INT.udn b/Engine/Documentation/Source/Engine/Rendering/RenderTargets/BlueprintRenderTargets/HowTo/HeightFieldPainter/3/HFP3.INT.udn
--- a/Engine/Documentation/Source/Engine/Rendering/RenderTargets/BlueprintRenderTargets/HowTo/HeightFieldPainter/3/HFP3.INT.udn
+++ b/Engine/Documentation/Source/Engine/Rendering/RenderTargets/BlueprintRenderTargets/HowTo/HeightFieldPainter/3/HFP3.INT.udn
@@ -137,7 +137,7 @@ Now that the required Blueprints have been setup, in the next section we will go
 [EXCERPT:m2]
 
 [REGION:note]
-You have to enable **Support UV from Hit Result** feature in the editor to use this method. See [](Engine/Physics/Tracing/Overview#EnablingUVCoordinatesfromTrace) for information on how to do this.
+You have to enable **Support UV from Hit Result** feature in the editor to use this method. See [getting uv coordinates from trace](Engine/Physics/Tracing/Overview#gettinguvcoordinatesfromtrace) for information on how to do this.
 [/REGION]
 
 1. Click the small (white) plus sign on the **Functions** tab to create a new function and name it **TraceFromCamera**.
@@ -936,7 +936,7 @@ Now that the required Blueprints have been setup, in the next section we will go
 		blueprint
 	[/PARAMLITERAL]
 	[PARAM:image]
-		![](Engine\Rendering\RenderTargets\BlueprintRenderTargets\HowTo\HeightFieldPainter\HFP_VP_Setup_UV_01.png)
+		![](Engine/Rendering/RenderTargets/BlueprintRenderTargets/HowTo/HeightFieldPainter/HFP_BP_Setup_01.png)
 	[/PARAM]
 	[PARAM:snippet]
 		Begin Object Class=K2Node_FunctionEntry Name="K2Node_FunctionEntry_48"
@@ -1165,7 +1165,7 @@ Now that the required Blueprints have been setup, in the next section we will go
 		blueprint
 	[/PARAMLITERAL]
 	[PARAM:image]
-		![](Engine\Rendering\RenderTargets\BlueprintRenderTargets\HowTo\HeightFieldPainter\HFP_VP_Setup_UV_02.png)
+		![](Engine\Rendering\RenderTargets\BlueprintRenderTargets\HowTo\HeightFieldPainter\HFP_BP_Setup_02.png)
 	[/PARAM]
 	[PARAM:snippet]
 		Begin Object Class=K2Node_CustomEvent Name="K2Node_CustomEvent_0"
diff --git a/Engine/Documentation/Source/Engine/Rendering/nDisplay/nDisplay.INT.udn b/Engine/Documentation/Source/Engine/Rendering/nDisplay/nDisplay.INT.udn
new file mode 100644
--- /dev/null
+++ b/Engine/Documentation/Source/Engine/Rendering/nDisplay/nDisplay.INT.udn
@@ -0,0 +1,509 @@
+Availability: Public
+Crumbs:
+Title: Rendering to Multiple Displays with nDisplay
+Description: Describes how to use the nDisplay system to render your Unreal Engine Project simultaneously on multiple displays.
+Version: 4.20
+Parent: Engine/Rendering
+Order: 0
+Tags: Rendering 
+Tags: Graphics 
+Tags: Experimental 
+Tags: Landing 
+Tags: Page
+
+
+![](ndisplay_banner.png)
+
+Interactive content isn't limited to being displayed on a single screen, or even a single dual-screen device like a VR headset. An increasing number of visualization systems aim to immerse the viewer more effectively in the game environment by rendering real-time content through multiple simultaneous displays. These systems may be made up of multiple adjacent physical screens, such as a [Powerwall](https://en.wikipedia.org/wiki/Powerwall) display; or they may use multiple projectors to project the 3D environment onto physical surfaces like domes, tilted walls, or curved screens, such as in a [Cave](https://en.wikipedia.org/wiki/Cave_automatic_virtual_environment) virtual environment.
+
+The Unreal Engine supports these usage scenarios through a system called **nDisplay**. This system addresses some of the most important challenges in rendering 3D content simultaneously to multiple displays:
+
+* It eases the process of deploying and launching multiple instances of your Project across different computers in the network, each rendering to one or more display devices.
+* It manages all the calculations involved in computing the viewing frustum for each screen in real time, based on the spatial layout of your display hardware.
+* It ensures that the content being shown on the various screens remains *exactly* in sync, with deterministic content across all instances of the Engine.
+* It offers passive and active stereoscopic rendering.
+* It can be driven by input from VR tracking systems, so that the viewpoint in the displays accurately follows the point of view of a moving viewer in real life.
+* It is flexible enough to support any number of screens in any relative orientation, and can be easily reused across any number of Projects.
+
+## nDisplay System Overview
+
+Every nDisplay setup has a single *master* computer, and any number of additional computers.  
+
+* Each computer in the network runs one or more instances of your Project's packaged executable file. Each of those Unreal Engine instances is responsible for rendering a single segment of the same 3D scene to a single screen or display.
+* The master node is also responsible for accepting input from a VRPN device, and replicating that input to all other connected computers.
+
+![nDisplay network overview](image5.png "nDisplay network overview")
+
+To make this possible, nDisplay adds several components to the usual Unreal system architecture:
+
+* A network provisioning and management application, called **nDisplayLauncher**. You run this application on a single computer in your network to automatically deploy and launch your Project on all the computers in your network. 
+* A separate listener application, called **nDisplayListener**, that runs on each computer. This application listens for incoming requests from the nDisplayLauncher, and processes those requests on the local computer.
+* A shared configuration file that contains all the settings nDisplay needs to start up the right number of instances on the right computers, each rendering the right point of view on the game's 3D world to produce the illusion of a seamless rendering across all display screens or projectors. See [About the nDisplay Configuration File](#aboutthendisplayconfigurationfile) below.
+
+## Getting Started
+
+This section describes how to get up and running with nDisplay for the first time. Before you begin:
+
+* Make sure you have your physical equipment — screens, projectors, etc. — set up and working correctly.
+* Make sure that the Windows account you intend to use on the master computer has administrative rights on *all* computers you intend to use in the nDisplay network.
+* Make sure that all computers you want to use in the nDisplay network can receive TCP/IP communications over ports 41000, 41001, and 41002. (You can use different ports instead; see [Changing Communication Ports](#changingcommunicationports) below.) 
+
+### Step 1. Set up Your Project for nDisplay
+
+The easiest way to get a Project set up to use nDisplay is to create it from the nDisplay Template Project:
+
+![Create an nDisplay Project from the Template](ndisplay-create-from-template.png "Create an nDisplay Project from the Template")
+
+This automatically enables the nDisplay plugin, adds some sample configuration files to your Project, and provides a default Level that is already configured with default settings.
+
+If you have an existing Project that you want to use with nDisplay, you can do the same configuration by hand. See [Add nDisplay to an Existing Project](#addingndisplaytoanexistingproject) below.
+
+### Step 2. Set up the Configuration File
+
+You need to tell nDisplay about the different computers you want to use in your network, the size and resolution of the screens or projectors those computers will render to, the spatial relationships between those screens in the 3D space of your Level, and more. To do this, you create a configuration file that expresses all of this information in a series of settings.
+
+Creating this configuration file is likely to be the trickiest part of your nDisplay setup, so you should approach it with care. For details, see [About the nDisplay Configuration File](#aboutthendisplayconfigurationfile) below.
+
+Typically, once your configuration file is set up, you only have to modify it when the topology of your network changes: for example, when you need to change the computers you are rendering to, or if you change the physical arrangement of your screens and displays in the real world.
+
+Save your configuration file inside your Project's Content folder. You'll use it in the next step.  
+
+### Step 3. Package and Deploy
+
+Every time you change the content in your Project, you'll need to package your game and deploy it to all of the computers that you've identified in your configuration file.
+
+1. In the Unreal Editor, cook and package your game for Windows. For details, see [Build Operations: Cook, Package, Deploy & Run](https://docs.unrealengine.com/en-us/Engine/Deployment/BuildOperations).  
+   Later, when you deploy the application to the slave computers, nDisplay will copy it to the same location on each of the target computers. So it's a good idea to package your project now to a location that exists on all machines.
+1. Find the*Engine\\Binaries\\DotNET\\nDisplayListener.exe *application under your Unreal Engine installation folder. Copy this application to the folder that contains the *.exe* file you've packaged for your game.
+1. Copy your nDisplay configuration file to the same folder.  
+   You should have your packaged Project *.exe* file, the *nDisplayListener.exe *application, and your configuration file all side-by-side in the same folder.
+1. Run the *Engine\\Binaries\\DotNET\\nDisplayLauncher.exe *application under your Unreal Engine installation folder.
+   [REGION:lightbox]
+   [![Deploy in the nDisplay Launcher](ndisplay_launcher_config-1.png "Deploy in the nDisplay Launcher")](ndisplay_launcher_config-1.png) *Click for full image.*
+   [/REGION]
+1. Add your packaged Project *.exe* file to the Applications list.  
+   Click **Add** under the **Applications** list, then browse to and select the *.exe* file you packaged for your Project.
+1. Specify your configuration file.  
+   Click **Add** to the right of the **Config Files** control, then browse to and select your configuration file.
+1. Click **Deploy application**.  
+   nDisplay copies the entire contents of the folder that contains your Project's executable file to the same location on each other computer that you've identified in your configuration file.
+    
+
+### Step 4. Launch all Nodes
+
+Once you have your Project deployed successfully to all the computers you've identified in your configuration file, you can use the **nDisplayLauncher** application to start the Project on all computers simultaneously.
+
+1. If you don't already have the **nDisplayLauncher** application running, start it and set the Application and Config File as instructed in the previous section.  
+   [REGION:lightbox]
+   [![Run with the nDisplay Launcher](ndisplay_launcher_launch.png "Run with the nDisplay Launcher")](ndisplay_launcher_launch.png) *Click for full image.*
+   [/REGION]
+1. Click **Start listeners**. This runs the nDisplayListener application on each computer you've identified in your nDisplay configuration file.
+1. When nDisplay confirms that all listeners have been started on all computers, click **Run** to launch all instances.
+
+When you're done:
+
+1. Click **Kill** to automatically shut down all instances of the Unreal Engine on all computers, or simply shut down the instance of Unreal Engine that is running on the master computer.
+1. Click **Stop listeners** to shut down the nDisplayListener application on all computers. 
+
+## About the nDisplay Configuration File
+
+[REGION:tip]
+The best way to get started understanding the nDisplay configuration file, and creating your own, is to start from the example configurations provided by the nDisplay plugin. If you've created your Project from the nDisplay template, you'll find these files in your Project folder, under *Content/ConfigExamples*. If not, you can find these files in the Unreal Engine installation folder, under *Templates/TP_nDisplayBP/Content/ConfigExamples*.
+[/REGION]
+
+The structure of the nDisplay configuration file is directly tied to the different types of components it uses to render the visualization.
+
+* Each different type of component that you configure has its own line in the file, and is identified by a string ID that you assign. You use these string IDs when one configuration section needs to refer to another.
+* Many of the components that you configure in this file have defined positions (and often rotations) in virtual 3D space. Each object's position and rotation is relative to that object's *parent*. By default, the parent of all objects is the VR space origin: an arbitrary point in 3D world space where VR space is considered to start. You can also configure specific named transforms in 3D space, called scene_nodes, which can act as parents for one or more components. This can help simplify the spatial layout of your screens, cameras and other components.
+* All parameters expect values in meters and degrees, unless otherwise indicated.
+
+### Camera Configurations
+
+nDisplay renders the scene from a single viewpoint at a time. Each of these potential viewpoints is defined by a **camera **configuration line.  
+
+You can switch between these viewpoints at runtime. Each viewpoint can also be driven by a tracking device.
+
+#### Example configuration:
+
+	[camera] id=camera_static loc="X=0.Y=0,Z=1.7" tracker_id=VRPNTracking tracker_ch=0
+
+#### Parameters:
+
+| Parameter | Description |
+| --- | --- |
+| **id** | A unique name for this camera configuration. |
+| **loc** | The location of this camera in VR space, relative to its parent. |
+| **tracker_id** | The name of the **input** configuration that defines the VR device you want to drive the position of the camera over time. Optional. If you omit this parameter, the camera's position will be static in VR space. |
+| **tracker_ch** | When you provide a **tracker_id**, this parameter specifies the channel of that device that nDisplay will read tracking data from. |
+| **parent** | The name of a **scene_node** configuration that you want to act as the parent for this object. This parameter is optional. If you specify a parent, the values you set in the loc parameter will be relative to the position of that parent. If you omit the parent, the values you set in the loc parameter will be relative to the VR root. |
+
+### Screen Configurations
+
+Each different output display renders the scene from the current camera's position, using a frustum that is defined by a rectangle of with a defined size and placement in the 3D VR space. Each of these rectangles is defined by a **screen **configuration. Usually, each of these projection screens has the same dimensions in VR space as the physical screen that you'll use to render it.  
+
+The pivot point of a screen is always in its exact midpoint.
+
+#### Example configuration:
+
+This definition describes a screen that is 3 meters by 3 meters screen, directly in front of its parent. Because the pivot point of the screen is at the center of the rectangle defined by the size parameter, we add a 1.5 meter offset on the Z axis to move the screen upward by half its height.
+
+	[screen] id=screen_front loc="X=1.5.Y=0,Z=1.5" rot="P=0,Y=0,R=0" size="X=3,Y=3" parent=screens
+
+To define a screen on the left side of the viewer, we move it to the left (negative values on the Y axis), and rotate it around its local Y axis (yaw), 
+
+	[screen] id=screen_left loc="X=0,y=-1.5,Z=1.5" rot="P=0,Y=-90,R=0" size="X=3,Y=3" parent=screens
+
+#### Parameters:
+
+| **Parameter** | **Description** |
+| --- | --- |
+| **id** | A unique name for this screen configuration. |
+| **loc** | The location of the center of this screen in VR space, relative to its parent. |
+| **rot** | The pitch (P), yaw (Y) and roll (R) angles of the screen's facing direction, in degrees. |
+| **size** | The total size of the screen rectangle along its local X and Y axes, in meters. |
+| **parent** | The name of a **scene_node configuration** that you want to act as the parent for this object. This parameter is optional. If you specify a parent, the values you set in the loc and rot parameters will be relative to the position of that parent. If you omit the parent, the values you set in the loc and rot parameters will be relative to the VR root. |
+
+
+### Cluster Node Configurations
+
+Each of the projection screens defined above is typically rendered by a different instance of your Unreal Engine application. For each of these instances, you need to define a **cluster_node **configuration. Each cluster node refers to the projection screen configuration that it should render. 
+
+The cluster_node configuration also defines the hostname or IP address of the computer that will run that application instance. You may set up a different physical computer for each cluster_node configuration, or you may have multiple cluster_node configurations that run on the same host.
+
+#### Example configuration:
+
+This example configures a master node (one per network):
+
+	[cluster_node] id=node_front addr=192.168.0.1 screen=screen_front viewport=vp_front sound=true port_cs=41001 port_ss=41002 master=true
+
+This example shows a non-master cluster node:
+
+	[cluster_node] id=node_left addr=192.168.0.2 screen=screen_left viewport=vp sound=false
+
+#### Parameters:
+
+| **Parameter** | **Description** |
+| --- | --- |
+| **id** | A unique name for this cluster node configuration. |
+| **addr** | The IP address of the computer that will run this instance of Unreal Engine. This must be an IPv4 address. IPv6 is not supported. |
+| **screen** | The name of the **screen** configuration that defines the projection screen that this instance of Unreal should be in charge of rendering. |
+| **viewport** | The name of the **viewport** configuration that defines the position of the rendered frame within the application window run by this instance of Unreal Engine. |
+| **sound** | Determines whether this instance of Unreal Engine plays sound. Optional; default value is `false`. |
+| **Winx WinY** | Specifies the position of the top left corner of the application window on the desktop, in pixels of screen space. |
+| **ResX ResY**	| Specifies the size of the application window in pixels of screen space. Make sure that the viewport definition you use for this cluster node does not exceed the size of this window. |
+| **port_cs port_ss** | port_ss	Cluster synchronization ports that the master node uses to communicate with other nodes in the cluster. Optional; default values are `14001` and `14002`. |
+| **master** | Determines whether this instance of Unreal Engine is the master node of the cluster. Only one cluster_node section can have this parameter set to true. Optional; default value is `false`. |
+
+### Viewport Configurations
+
+Each **cluster_node** configuration described above refers to a **viewport **configuration, which defines the rectangular area of the game window where the rendered frame is mapped. 
+
+You'll typically need a different viewport configuration for each display device that you have with a different set of resolution settings. For example, if all your display screens are identical, you only need one viewport configuration, which you can refer to from all of your cluster node configurations.
+
+Usually, a viewport starts at the upper left corner of the application window, and its width and height are set so that they fill the parent window. However, in some cases you may need to offset the viewport within its parent application window. For example, you might want to do this if you need to set up two projectors that partially overlap.  
+
+#### Example configuration:
+
+	[viewport] id=vp_left X=0 Y=0 width=1920 height=1080
+
+#### Parameters:
+
+| **Parameter** | **Description** |
+| --- | --- |
+| **id** | 	A unique name for this viewport configuration. |
+| **X Y** | The coordinates of the top left corner of the viewport, in pixels, within the screen space of the main application window. |
+| **width height** | The width and height of the rendered frame, in pixels. This should not be larger than the size of the game window that is set by the size parameter of any **cluster_node** configuration using this viewport. |
+
+### Scene Node Configurations
+
+In your configuration file, you can define a hierarchy of scene nodes, each of which represents a transform in 3D space. Anything that you set up in the configuration file that requires a position and rotation in 3D space, such as a camera or a projection screen, can use one of these *scene_node *configurations as its parent. This can help you to define the spatial relationships between all of the different components of the visualization system.
+
+Like cameras, scene nodes can also be driven by VR tracking devices.
+
+#### Example configuration:
+
+The following lines define a hierarchy of two nodes, where the child node has an offset of 2 meters in front of its parent.
+
+	[scene_node] id=vr_space_root loc="X=0.Y=0,Z=0" rot="P=0,Y=0,R=0"
+	[scene_node] id=walls_front_group loc="X=2.Y=0,Z=0" rot="P=0,Y=0,R=0" parent= vr_space_root
+
+The following line shows a scene node that is configured to be driven by a VR tracking device:
+
+	[scene_node] id=cave_wand loc="X=0, Y=0,Z=1" tracker_id=CaveTracking tracker_ch=1
+
+#### Parameters:
+
+| **Parameter** | **Description** |
+| --- | --- |
+| **id** | A unique name for this scene node configuration. |
+| **loc** | The location of this scene node in VR space, relative to its parent. |
+| **rot** | The pitch (P), yaw (Y) and roll (R) angles of the scene node's facing direction, in degrees. |
+| **parent** | The name of another **scene_node** configuration that you want to act as the parent for this scene node. This parameter is optional. If you specify a parent, the values you set in the **loc** and **rot** parameters will be relative to the position of that parent. If you omit the parent, the values you set in the **loc** and **rot** parameters will be relative to the VR root. |
+| **tracker_id** | The name of the **input** configuration that defines the VR device you want to drive the position of the scene node over time. Optional. If you omit this parameter, the scene node's position and rotation will be static in VR space with respect to its parent. |
+| **tracker_ch** | When you provide a **tracker_id**, this parameter specifies the channel of that device that nDisplay will read tracking data from. |
+
+### Input Configurations
+
+Each **camera** and each **scene_node** may optionally be driven by a VR tracking device. To do this, you define an **input** section for each VR device, and refer to it in the **camera** or **scene_node** configuration.
+
+#### Example configuration:
+
+	[input] id=CaveTracking type=tracker [REDACTED_EMAIL] loc="X=-1.5,Y=0,Z=3.4" rot="P-0,Y=0,R=0" front=X right=Y up=-Z
+
+#### Parameters:
+
+| Parameter | Description |
+| --- | --- |
+| **id** | A unique name for this input device configuration. |
+| **type** | [INCLUDE:#Excerpt_2] |
+| **addr** | [INCLUDE:#Excerpt_3] |
+| **remap** | [INCLUDE:#Excerpt_4] |
+
+[COMMENT:none]
+
+[EXCERPT:Excerpt_2]
+The type of this VRPN input device:
+* `tracker` for a tracking device.
+* `analog` for a device that produces axis data.
+* `button` for a device that produces Boolean button data.
+[/EXCERPT]
+
+[EXCERPT:Excerpt_3]
+The type of this VRPN input device:	The address of the VRPN server that handles this particular device. The value must match the following format:
+`DEVICENAME@SERVER_ADDRESS`
+`where:`
+* `DEVICENAME` is the VRPN name for this device.
+* `SERVER_ADDRESS` is the IPv4 address of the VRPN server.
+[/EXCERPT]
+
+[EXCERPT:Excerpt_4]
+This parameter remaps channel numbers from the tracking device so that their input can be read from different channels in Unreal. For example, some controllers use channel 0 for head tracking and channel 1 for the controller; others use channel 1 for the head and channel 5 for the controller. Switching between these two devices might require you to change your configuration and your runtime code to match. To avoid having to make any other changes in this case, you can use this parameter to remap the channels.
+For example, the value:
+`remap[0:1, 5:2]`
+forwards the data from channel 1 in the tracking device to channel 0 in Unreal, and forwards the data from channel 2 in the tracking device to channel 5 in Unreal.
+[/EXCERPT]
+
+[/COMMENT]
+
+Devices where `type=tracker` also accept the following additional parameters:
+
+| **Parameter** | **Description** |
+| --- | --- |
+| **loc rot** | Similar to other configuration sections, the **loc** and **rot** parameters specify position and rotation offsets in local space for this input device. However, for an input device, you typically use these offsets to adjust the root position of a tracking device in VR space to match the location you expect it to be in your scene node hierarchy. |
+| **front right up** | [INCLUDE:#Excerpt_5] |
+
+
+[COMMENT:none]
+[EXCERPT:Excerpt_5]
+These parameters match each local axis of the tracker in Unreal (front, right, and up) with the corresponding axis in the tracker's coordinate system. Unreal uses a right-handed, Z-up coordinate system. If your tracker uses a different coordinate system, you can use these parameters to map the tracker's coordinate system to Unreal's.
+For example, the following line maps the Y axis of the tracker to the front (X) axis in Unreal; the X axis of the tracker to the right (Y) axis in Unreal, and the negative Z axis of the tracker to the up (Z) axis in Unreal:
+`front=Y right=X up=-Z`
+[/EXCERPT]
+[/COMMENT]
+
+### General Configuration
+
+The **general** configuration line contains parameters that control the overall operation of the nDisplay cluster.
+
+#### Example configuration:
+
+	[general] swap_sync_policy=1
+
+#### Parameters:
+
+| **Parameter** | **Description** |
+| --- | --- |
+| **swap_sync_policy** | [INCLUDE:#Excerpt_6] |
+
+[COMMENT:none]
+[EXCERPT:Excerpt_6]
+Determines how output is synchronized over the network.
+* 0: no synchronization.
+* 1: Software swap synchronization
+* 2: NV swap lock (only for NVIDIA cards rendering with OpenGL)
+[/EXCERPT]
+[/COMMENT]
+
+### Stereo Configuration
+
+The **stereo** configuration line sets optional global parameters for stereoscopic rendering.
+
+#### Example configuration:
+
+\[stereo\] eye\_swap=True eye\_dist=0.064
+
+#### Parameters:
+
+| **Parameter** | **Description** |
+| --- | --- |
+| **eye_swap** | Determines whether or not the images generated for the left and right eye are swapped. Optional. Default value is false. |
+| **eye_dist** | The inter-ocular distance to use for offsetting the images generated for the left and right eyes, in meters. |
+
+### Structure
+
+The references between the different configuration types outlined above imply the following structure, where red arrows are mandatory references and green arrows are optional references:
+
+![Configuration reference structure](ndisplay-configuration-references.png "Configuration reference structure")
+
+### Example
+
+To take a specific example, open the *wall_flat_3x2.cfg* sample file. This file defines six projection screens, each of which is to be rendered by a separate physical computer.
+
+It also defines several scene_nodes, which taken together create the following hierarchy:
+
+![nDisplay example scene hierarchy](ndisplay-scene-hierarchy.png "nDisplay example scene hierarchy")
+
+The relative positions and rotations of the nodes in this hierarchy lay out the arrangement of the camera and the six screens in VR space so that the six projection screens are side-by-side, at a distance of 1 meter from the camera. 
+
+[REGION:lightbox]
+[![nDisplay example layout](ndisplay-example-layout.png)](ndisplay-example-layout.png)
+[/REGION]
+
+[REGION:caption]
+Click image for full size.
+[/REGION]
+
+Note that the configuration implies a small space in between each adjacent pair of projection screens, to account for the edges of the monitors rendering the scene.
+
+## Blueprint API
+
+You can control the behavior of the nDisplay system in your game's runtime logic using its Blueprint API.
+
+To get to the functions exposed in this API for managing clusters, input and rendering:
+
+1.  Create a new **Display Cluster > Get Plugin API** node in your Blueprint.
+2.  Drag from the **Out API** pin, and look under the **Display Cluster** category:  
+    ![nDisplay Get Plugin API](ndisplay-get-plugin-api.png "nDisplay Get Plugin API")
+
+## Actor Replication
+
+All inputs to the nDisplay system are handled only by the master node. Without any replication, only the master node would see changes in the scene. Therefore, the master node needs to be able to replicate changes to all other parts of the nDisplay network.
+
+To accomplish this, nDisplay offers two different kinds of Components that you can attach to your Actors:
+
+*   The **DisplayClusterSceneComponentSyncParent** Component tracks changes in the 3D transforms of its parent Component, and pushes those changes to the other cluster nodes in the network.  
+    The default DisplayClusterPawn used by the nDisplay system uses this Component.  
+    
+*   The **DisplayClusterSceneComponentSyncParent **Component tracks changes to the 3D transforms of its child components, and pushes those changes to the other cluster nodes in the network.
+
+For example, in the Actor shown below, the **DisplayClusterSceneComponentSyncParent_DefaultSceneRoot** Component tracks and replicates changes to the 3D transforms of its parent Actor as the Actor moves around the Level. The **DisplayClusterSceneComponentSyncThis** Component tracks and synchronizes movements of its child Cube component as it moves relative to the scene graph root.
+
+![DisplayClusterSceneComponentSyncParent ](ndisplay-actor-replication.png "DisplayClusterSceneComponentSyncParent ")
+
+If you have other Actors in your scene that can be affected during gameplay, you must use one of these two Components to replicate those changes to all nodes. To do this:
+
+1.  Select the Actor you need to replicate in the Level viewport or the **World Outliner** panel.
+2.  In the **Details** panel, click **\+ Add Component**. Search for either **DisplayClusterSceneComponentSyncParent** or **DisplayClusterSceneComponentSyncThis**, and select it from the list.  
+    ![Add an nDisplay sync Component](ndisplay-actor-replication-selection.png "Add an nDisplay sync Component")
+
+[REGION:note]
+These components do not carry out a full replication. Only the transforms of the parent Actor or of child Components are sent to the cluster.
+[/REGION]
+
+## Using VRPN Inputs
+
+To use a [VRPN](https://en.wikipedia.org/wiki/VRPN) input device with nDisplay:
+
+1. Install a VRPN server on your network.  
+   This version of nDisplay requires VRPN **version 33**.
+1. In the server’s *vrpn.cfg* file, which you'll find located next to the server's executable file, enable your input device and give itone a name.
+1. In your nDisplay configuration file, add an entry to set up your VRPN input device, adjust its coordinate system, and attach it to a camera.
+
+The following example shows one way to set up the ART DTrack tracking system:
+
+* In Vrpn.cfg, located next to the vrpn server exe, add the following line:
+		vrpn_Tracker_DTrack DTrack  5000	
+  This makes VRPN receive the DTrack inputs from port 5000, and maps them to the VRPN device named DTrack. (Make sure DTrack is configured to output its tracking data on port 5000.)
+* In the nDisplay config file, add the following lines:
+		[input] id=CaveTracking type=tracker [REDACTED_EMAIL] loc="X=1.32,Y=0,Z=0.93735" rot="P=0,Y=0,R=0" front=Z right=-X up=Y
+		[camera] id=camera_dynamic loc="X=0,Y=0,Z=0" tracker_id=CaveTracking tracker_ch=0
+  The first line creates an nDisplay input device named CaveTracking, which fetches data from VRPN address [REDACTED_EMAIL]. You have to adapt your coordinate system here to fit with your tracking system and offset. The second line tells nDisplay to fetch camera positions from the CaveTracking input on channel 0.
+
+To manage the state of your VRPN devices from your gameplay code, and to detect input events such as button presses:
+
+* In C++, use the `IDisplayClusterInputManager` class.
+* In Blueprints, use the nodes under **DisplayCluster > Input**. Make sure that the ID you set in the nodes matches the ID value you have set in your configuration file. For example:  
+  ![](ndisplay_vrpn_api.png "ndisplay_vrpn_api.png")
+
+## Changing Communication Ports
+
+The nDisplay system communicates between hosts over three TCP/IP ports: 14000, 14001, and 14002. You need to make sure you have these hosts open on all computers.
+
+If you want to change the port numbers yourself, you can do so in the following places.
+
+* **Runtime synchronization ports**. The master node uses two ports to synchronize data with the other nodes in the cluster. To set these two ports, include the **port_cs** and **port_ss** configuration parameters in your configuration file, on the **cluster_node **line that defines your master node. For example:
+		[cluster_node] id=node_front addr=192.168.0.1 screen=screen_front viewport=vp_front port_cs=42001 port_ss=42002 master=true
+
+* **nDisplay Launcher and nDisplay listener ports**. The nDisplay Launcher and nDisplay Listener both need to be configured to use the same communication port. You can specify this on the command line when you start up these applications.  
+  When you start the nDisplay Launcher, use the `listener_port` argument. For example:
+	nDisplayLauncher.exe listener_port=15003
+
+  In addition, you'll have to start the nDisplayListener application on each host yourself, with the `port` argument. For example:
+	nDisplayListener.exe port=15003
+
+## Adding nDisplay to an Existing Project
+
+To set up an existing Project to use nDisplay:
+
+1.  Enable the nDisplay plugin.  
+    In the Unreal Editor, choose **Edit > Plugins** from the main menu. Search for "nDisplay", and check the **Enabled **checkbox.  
+	[REGION:lightbox]
+    [![Enable the nDisplay Plugin](ndisplay-enable-plugin.png "Enable the nDisplay Plugin")](ndisplay-enable-plugin.png)  
+	*Click for full image.*
+    [/REGION]
+2.  Enable nDisplay for your Project.  
+    Choose **Edit > Project Settings** from the main menu, and find the **Plugins > nDisplay** section. Check the **Enabled** checkbox.  
+    [REGION:lightbox]
+    [![nDisplay Project Settings](ndisplay-project-settings.png "nDisplay Project Settings")](ndisplay-project-settings.png)  
+	*Click for full image.*
+    [/REGION]
+3.  Still in the **Project Settings** window, go to the **Project > Description** section, and check the **Settings > Use Borderless Window** checkbox.  
+    [REGION:lightbox]
+    [![Use Borderless Window](ndisplay-project-settings-borderless.png "Use Borderless Window")](ndisplay-project-settings-borderless.png)  
+	*Click for full image.*
+    [/REGION]
+4.  Restart the Unreal Editor, reopen your Project, and open your Project's default Level.
+5.  In the **World Settings** panel, set the **Game Mode > GameMode Override** setting to **DisplayClusterGameModeDefault**.  
+    [REGION:lightbox]
+    [![nDisplay World Settings](ndisplay-world-settings.png "nDisplay World Settings")](ndisplay-world-settings.png)  
+	*Click for full image.*
+    [/REGION]    
+6.  Add a new **DisplayClusterSettings** Actor to your Level.  
+    You can find this Actor in the **Modes** panel, on the **All Classes** list.
+7.  Continue on with the rest of the setup instructions under [Getting Started](#gettingstarted) above.
+
+## nDisplay Launcher UI Reference
+
+This section describes all of the settings and options available in the user interface of the nDisplay Launcher.
+
+[REGION:lightbox]
+[![nDisplay Launcher UI](ndisplay-launcher-ui.png)](ndisplay-launcher-ui.png) 
+[/REGION]
+
+| **Control** | **Description** |
+| --- | --- |
+| **Render API** | Specifies the rendering API to use the next time you click Run. |
+| **Render mode** | [INCLUDE:#excerpt_9] |
+| **Use All Available Cores** | Forces each Unreal instance to use all available processors on its host. When this option is selected, the nDisplay Launcher adds the USEALLAVAILABLECORES option to the command line it uses to launch each instance. |
+| **No Texture Streaming** | Disables texture streaming for each Unreal instance. Highest quality textures are always loaded. When this option is selected, the nDisplay Launcher adds the NOTEXTURESTREAMING option to the command line it uses to launch each instance. |
+| **Custom command line arguments** | If you want the nDisplay Launcher to pass any additional arguments on the command line it uses to launch each Unreal instance, include them here. For details, see the [Command-Line Arguments]("Programming/Basics/CommandLineArguments") reference. |
+| **Applications** | Lists all packaged Unreal applications that you can run with the nDisplay Launcher. Use the Add and Delete buttons to edit the list. For more, see [Step 3. Package and Deploy]("Engine/Rendering/nDisplay#step3packageanddeploy") above. |
+| **Config Files** | Lists all configuration files that you've set up for the nDisplay Launcher. Use the Add and Delete buttons to edit the list. For more, see [Step 3. Package and Deploy]("Engine/Rendering/nDisplay#step3packageanddeploy") above. |
+| **Run** | Launches the application that you've selected in the **Applications** list on all hosts that are configured in the file that you've selected in the **Config Files** list. |
+| **Kill** | Shuts down all instances of the Unreal Engine that the nDisplay Launcher started up. |
+| **Start listeners** | Starts an instance of the nDisplay Listener application on each host configured in the file that you've selected in the **Config Files** list.  |
+| **Stop listeners** | Shuts down all instances of the nDisplay Listener application that this nDisplay Launcher started up. |
+| **Deploy application** | Copies the application that you've selected in the **Applications** list to each host configured in the file that you've selected in the **Config Files** list. |
+| **Copy** | Copies to the clipboard all messages listed in the log window to the left. |
+| **Clean** | Clears all messages from the log window to the left. |
+[COMMENT:none]
+[EXCERPT:Excerpt_9]
+Specifies the type of output nDisplay produces:
+* Monoscopic: no stereo rendering
+* Frame sequential: active quad buffer stereo
+* Side-by-side: passive horizontally aligned stereo
+* Top-bottom: passive vertically aligned stereo
+
+Monoscopic does not require any specific hardware features, but frame sequential does. Make sure your display device, GPU and diver settings are compatible with the render mode you choose.
+[/EXCERPT]
+[/COMMENT]
\ No newline at end of file
diff --git a/Engine/Documentation/Source/Engine/Sequencer/HowTo/ControlAnimInstances/ControlAnimInstances.INT.udn b/Engine/Documentation/Source/Engine/Sequencer/HowTo/ControlAnimInstances/ControlAnimInstances.INT.udn
new file mode 100644
--- /dev/null
+++ b/Engine/Documentation/Source/Engine/Sequencer/HowTo/ControlAnimInstances/ControlAnimInstances.INT.udn
@@ -0,0 +1,102 @@
+Availability: Public
+Title: Controlling Anim Instances with Sequencer
+Crumbs: %ROOT%, Engine, Engine/Sequencer, Engine/Sequencer/HowTo
+Description: Shows how to animate variables on anim instances through possessables
+Version: 4.20
+SkillLevel: Advanced
+parent: Engine/Sequencer/HowTo
+order: 1
+tags: Sequencer
+type: how-to
+
+Within Sequencer, it is possible to animate variables on Anim Instances through Possessables, enabling you direct control of Anim Blueprint variables, functions and other content. By adding a Skeletal Mesh Component Track, you can gain access to the Anim Instance Track where any variables that are exposed to cinematics will be shown and available to keyframe. 
+
+In this how-to, we blend between multiple animation poses that are defined within an Animation Blueprint by keyframing those changes inside Sequencer. 
+
+[REGION:banner]
+![image alt text](image_0.png)
+[/REGION]
+
+## Steps
+
+For this guide we are using a new **Blueprint Third Person** template project.
+
+1. In the **Content/Mannequin/Character/Mesh **folder, right-click the **SK_Mannequin **then select **Anim Blueprint **under **Create **and give it any name. 
+
+	![image alt text](image_1.png)(w:624 convert:false)
+
+1. Right-click in the **Content Browser** and under **Blueprints**, select **Enumeration **and call it **MoveType**. 
+
+	![image alt text](image_2.png)(w:540 convert:false)
+
+1. Create three enumerators called, **Idle, Walk **and **Run **by clicking the **New **button. 
+
+	![image alt text](image_3.png)(w:624 convert:false)
+
+1. Open the **Animation Blueprint **created in step 1 and create a variable of the **MoveType **type called **MoveType **and enable **Expose to Cinematics**. 
+
+	![image alt text](image_4.png)(w:434 convert:false)
+
+1. On the **Anim Graph**, add the **ThirdPersonIdle, ThirdPersonWalk **and **ThirdPersonRun **animations and a **Blend Poses by Move Type **node.
+
+	![image alt text](image_5.png)(w:624 convert:false)
+
+1. Right-click on the **Blend Poses **node then add a pin for **Idle, Walk **and **Run**.
+
+	![image alt text](image_6.png)(w:317 convert:false) 
+
+1. Add the **Move Type **variable to the graph and connect each of the nodes to the **Final Animation Pose **as shown below.
+
+	![image alt text](image_7.png)(w:624 convert:false)
+
+1. Drag the **Animation Blueprint **into the Level, then create a new **Level Sequence **(give it any name) and add the Anim Blueprint to the Sequence.
+
+	![image alt text](image_8.png)(w:624 convert:false)
+
+1. Click the **+ Track **button on the Anim Blueprint and add an **SkeletalMeshComponent0 **track.
+
+	![image alt text](image_9.png)(w:624 convert:false) 
+
+1. Click the **+ Track **button on the SkeletalMeshComponent and add the **Anim Instance **track.
+
+	![image alt text](image_10.png)(w:624 convert:false)
+
+1. Click the **+ Track **button on the Anim Instance and add the **Move Type **property.
+
+	![image alt text](image_11.png)(w:624 convert:false)
+
+1. Scrub the Timeline to frame **45** and change the **Move Type **drop-down to **Walk, **adding a key.
+
+	![image alt text](image_12.png)(w:624 convert:false)
+
+1. Scrub the Timeline to frame **90 **and change the **Move Type **drop-down to **Run**, adding another key. 
+
+	![image alt text](image_13.png)(w:624 convert:false)
+
+1. Add a key for **Move Type **at frame **120 **set to **Walk **and another key at frame **150 **set to **Idle. **
+
+	![image alt text](image_14.png)(w:624 convert:false)
+
+1. Set the Level Sequence to **Auto Play **in the **Details **panel, then click the **Play **or **Simulate **button to play/simulate in the Editor. 
+
+## End Result
+
+When you Play or Simulate, the Level Sequence will play back and change the character's state based on the key framed property **Move Type **defined in the Sequence. Animating variable properties is useful in the event that you have characters that have Animation Blueprints that drive their animation logic, however you want to control what animation the character enters through a Sequence. 
+
+[OBJECT:EmbeddedVideo]
+    [PARAMLITERAL:videoid]
+    7l9vk4o-diE  <!-- Insert the youtube video link here -->
+    [/PARAMLITERAL]
+    [PARAMLITERAL:width]
+    78
+    [/PARAMLITERAL]
+    [PARAMLITERAL:height]
+    56.25
+    [/PARAMLITERAL]
+    [PARAMLITERAL:units]
+    %
+    [/PARAMLITERAL]
+    [PARAMLITERAL:type]
+    simple
+    [/PARAMLITERAL]
+[/OBJECT]
