Copying //UE4/Dev-Documentation to //UE4/Samples-Main (Source: //UE4/Dev-Documentation @ 4321946)

diff --git a/Engine/Documentation/Source/Engine/Sequencer/HowTo/SequenceRecorder/SequenceRecorder.INT.udn b/Engine/Documentation/Source/Engine/Sequencer/HowTo/SequenceRecorder/SequenceRecorder.INT.udn
--- a/Engine/Documentation/Source/Engine/Sequencer/HowTo/SequenceRecorder/SequenceRecorder.INT.udn
+++ b/Engine/Documentation/Source/Engine/Sequencer/HowTo/SequenceRecorder/SequenceRecorder.INT.udn
@@ -12,103 +12,242 @@ Related: Engine/Sequencer/ReferenceEditor
 tags:Sequencer
 type:how-to
 
-The **Sequence Recorder** allows you to capture specified Actors during level editing or gameplay that can be saved as a new **Level Sequence**, and can be edited inside of **Sequencer**. 
-This is useful for quickly capturing content for scenes, as you can take a playable character, perform some actions during gameplay while recording with the Sequence Recorder, then take that data into Sequencer, creating a cinematic around it.
+# Sequence Recorder
 
-In this example, we will use the Sequence Recorder to record our playable character's movement, which we can then edit.
+The **Sequence Recorder** enables you to record Actor(s) from gameplay or level editing into Skeletal Animation, Transform, Event and Audio Tracks inside a Level Sequence. You can then edit the Level Sequence inside the **Sequencer Editor** to create a cinematic based around your recording.
 
-## Steps
+## Opening the Sequence Recorder
 
-[REGION:note]
-For this How-to, we are using the **Blueprint Third Person Template** project with **Starter Content Enabled**. 
-[/REGION]
+You can open the **Sequence Recorder** from the main **Menu Bar** under the **Window** tab:
 
-1.	From the **Menu Bar** under **Window**, select **Sequence Recorder**. 
+![image alt text](image_0.png)(w:270 h:300 convert:false)
 
-	![](RecorderStep1.png)
+After selecting **Sequence Recorder**, the **Sequence Recorder** window will open:
 
-	The Sequence Recorder window will automatically open. There are some options under **Sequence Recording** which will determine how (and where) the new Level Sequence asset will be saved. You can choose to record Actors that are spawned (such as particle effects, other characters, etc.), determine the naming convention and save location, and change other settings like the duration to record or delay before starting to record.
-	
-1.	From the Main Editor Toolbar, click the **Play** button to start a Play in Editor (PIE) session. 
+![image alt text](image_1.png)(w:586 h:709 convert:false)
 
-1.	Press **Shift+F1** to gain mouse control once inside the game. 
+## Tracking Actors for Recording
 
-1.	On the **Sequence Recorder** window, click the **Add** button. 
-	
-	![](SequenceRecorderAdd.png)
+Across the top of the Sequence Recorder are buttons used to control recording: Record, Add, Add Player, Remove and Remove All.
 
-1.	Click on the new recording (which will say **None**), then for **Actor to Record** click the drop-down and select **ThirdPersonCharacter**. 
+![image alt text](image_2.png)(w:322 h:64 convert:false)
 
-	![](ActorToRecord.png)
 
-	This is where we specify which Actor to target before starting the recording process, updating the UI as shown above.
+| Button | Description |
+| --- | --- |
+| Record | Initiates recording (affects all recordings).  |
+| Add | Add a new recording for selected actors. If nothing is selected, an empty recording will be added. |
+| Add Player | Add a new recording for the current player. |
+| Remove | Remove selected recording. |
+| Remove All | Remove all recordings. |
 
-1.	Optionally, you can choose to record audio and set the audio gain levels along with your recorded clip. 
 
-	![](RecordAudio.png)
 
-	Audio recording requires an attached microphone and will start recording when the sequence starts recording. 
+After clicking the **Add** button, both the **Record** option as well as the **Remove** option will become accessible. This will also add an**Actor Recording** field to the Actor Tracking window with the Actor to track listed as **None** (shown below).
 
-1.	Click the **Record** button. After 4 seconds (which is the **Record Delay** option under the Sequence Recording section), the recording process will start. 
+![image alt text](image_3.png)(w:516 h:222 convert:false)
 
-	![](RecordStart.png)
+Selecting the Actor Recording will populate the Sequence Recorder with additional options related to that recording.
 
-	[REGION:note]
-	When clicking the **Record** button, all Actors in the list that are set to be tracked will also be recorded. 
-	[/REGION]
+![image alt text](image_4.png)(w:524 h:709 convert:false)
 
-1.	Use the **WASD** and **Spacebar** keys to move the character around and jump, then press **ESC** when done.
+With these options you can define what Actor to track during the recording along with its animation settings. You can add multiple Actor Tracking fields to the Tracking Window, with each **Actor to Record** set to track different Actors.
 
-	[OBJECT:EmbeddedVideo]
+[REGION:note]Please see the [Recording to Sequencer](https://docs.unrealengine.com/en-US/Engine/Sequencer/HowTo/SequenceRecorder) page for a step-by-step guide to recording Actors with the Sequence Recorder.[/REGION]
+
+## Sequence Recorder Options
+
+Once you have the **Sequence Recorder** open, there are several properties you can adjust to determine how the **Sequence Recorder**records its Actors. When first opening the **Sequence Recorder**, you will only be presented with the default **Sequence Recording**options. After adding Actor(s) to track for recording, the **Actor Recording** and **Animation Recording** sections become available on a per Actor basis.
+
+### Recording Groups
+
+**Sequence Recorder** supports a profile system that is stored in the Persistent Level. **Recording Profiles** enable you to store which Actors you wish to record and their settings, as well as the output path to store the recorded data in. **Sequence Recorder** also supports recording multiple takes for each of the selected Actors. 
+
+The following options are available within the **Recording Groups **section of the **Sequence Recorder**:
+
+![image alt text](image_5.png)(w:502 h:257 convert:false)
+
+
+| Option | Description |
+| --- | --- |
+| Sequence Name | The base name of the Sequence to record to. This name will also be used to auto-generate any assets created by this recording. |
+| Sequence Recording Base Path | Base path for this recording. Sub-assets will be created in sub-directories as specified.  |
+| Specify Target Level Sequence | Whether we should specify the target Level Sequence or auto-create it.  |
+| Target Level Sequence | The Level Sequence to record into.  |
+| Duplicate Target Level | Whether we should duplicate the target Level Sequence and record into the duplicate.  |
+|  |  |
+
+
+
+The video below demonstrates an example of setting up **Recording Groups** with **Sequence Recorder**. 
+
+[OBJECT:EmbeddedVideo]
+	[PARAMLITERAL:videoid]
+	QGO_jCrjTEI
+	[/PARAMLITERAL]
 	[PARAMLITERAL:width]
-	640
+	78
 	[/PARAMLITERAL]
 	[PARAMLITERAL:height]
-	360
+	56.25
 	[/PARAMLITERAL]
-	[PARAMLITERAL:videoid]
-	9mZI75qBsx8
+	[PARAMLITERAL:units]
+	%
+	[/PARAMLITERAL]
+	[PARAMLITERAL:type]
+	simple
 	[/PARAMLITERAL]
-	[/OBJECT] 
+[/OBJECT]
 
-	Inside the **Content Browser**, a new folder will be created containing assets related to the recorded sequence. 
+### Sequence Recording
 
- 
+These options enable you to define how the Sequence is recorded along with what Actors are recorded:
 
-## End Result
+![image alt text](image_6.png)(w:540 h:452 convert:false)
 
-You can open the **RecordedSequence** asset and begin editing it as you would a normal Level Sequence. Below is our recorded sequence, to which we could add cameras, and a [Camera Cuts](Engine/Sequencer/HowTo/TracksCameraCut) track to provide multiple angles, music, effects, or anything else we'd like to add. We could even take this sequence and embed it in other sequences as part of a [Shots Track](Engine/Sequencer/HowTo/TracksShot). 
 
-[OBJECT:EmbeddedVideo]
-[PARAMLITERAL:width]
-640
-[/PARAMLITERAL]
-[PARAMLITERAL:height]
-360
-[/PARAMLITERAL]
-[PARAMLITERAL:videoid]
-1nFX_bpWfy0
-[/PARAMLITERAL]
-[/OBJECT]
+| Option | Description |
+| --- | --- |
+| Create Level Sequence | Whether to create a Level Sequence when recording or not. Actors and animations will be inserted into this Sequence if enabled. |
+| Immersive Mode | Whether to maximize the viewport when recording. |
+| Sequence Length | Defines the maximum length of the Sequence to record (setting to 0 recording will continue until Stop is invoked by pressing Esc or the Stop button). |
+| Recording Delay | Delay that will be used before starting the recording. |
+| Animation Sub Directory | The name of the sub-directory animations will be placed in. Leave this empty to place into the same directory as the Sequence base path.  |
+| Record Nearby Spawned Actors | Whether or not to record nearby spawned Actors as part of this recording. |
+| Nearby Actor Recording Proximity | Defines the proximity to currently recorded Actors to record newly spawned Actors. |
+| Record World Settings Actor | Whether to record the World Settings Actor in the Sequence (some games use this to attach world sound effects). |
+| Reduce Keys | Whether to remove keyframes within a tolerance from the recorded tracks. |
+| Auto Save Asset | Whether to auto-save asset when recording is completed. Defaults to false. |
+| Actor Filter | Filter to check spawned Actors against to see if they should be recorded. You can specify which Actor classes to record from the expandable option. |
+| Level Sequence Actors to Trigger | Sequence Actors to trigger playback on when recording starts (for example, this can be used for recording synchronized Sequences). |
+| Default Animation Settings | The default settings applied to animation recording. Expand to see options for recording in World Space, Removing Root Animation, whether or not to Auto Save the asset when recording is complete, and the ability to adjust Sample Rate and Length. |
+| Classes and Properties to Record | The properties to record for specified classes. Component classes specified here will be recorded. If an Actor does not contain one of these classes it will be ignored. Expanded to define the classes and associated properties. |
+| Record Sequencer Spawned Actors | Whether to record Actors that are spawned by Sequencer itself (this is usually disabled as results can be unexpected). |
+| Per Actor Settings | Settings applied to Actors of a specified class. For each class you can define whether to record to a Possessable (level-owned) or Spawnable (sequence-owned). |
+
+
+
+### Actor Recording
+
+This section enables you to define the settings for Actors that are recorded: 
+
+![image alt text](image_7.png)(w:536 h:308 convert:false)
+
+
+| Option | Description |
+| --- | --- |
+| Record Visibility | Whether to record the Actors visibility or not. |
+| Record Transforms | Whether to record Actor transforms. This can be useful if you want the Actors to end up in specific locations after the sequence. By default we rely on animations to provide transforms, but this can be changed using the Record in World Space animation setting. |
+| Active | Whether this actor is active and to be recorded when the Record button is pressed.  |
+| Create Level Sequence | Whether to create a Level Sequence for this Actor recording. |
+| Target Level Sequence | The Level Sequence to record into.  |
+| Target Name | Optional target name to record to. If not specified, the Actor label will be used.  |
+| Take Number | Specify the Take Number for the new recording. |
+| Record to Possessable | Whether to record to Possessable (level owned) or Spawnable (sequence-owned) actors. Defaults to the global setting.  |
+| Actor to Record | This option is used to define which Actor you want to record. |
+
+
+
+### Animation Recording
+
+These settings can be used to define how Animation is recorded: 
+
+![image alt text](image_8.png)(w:536 h:165 convert:false)
+
+
+| Option | Description |
+| --- | --- |
+| Specify Target Animation | Whether we should specify the target animation or auto-create it. |
+| Target Animation | Enabled when Specify Target Animation is enabled and is used to specify the animation you want to record to. |
+| Animation Settings | See Default Animation Settings under Sequence Recording. These are applied to the Actor to Record. |
+
+
+
+## Tips and Tricks
+
+This section contains some workflow tips and tricks you can use when working with the **Sequence Recorder**.
+
+### Continuous Recording
 
-In addition to recording gameplay, you can record your actions during level editing by assigning an Actor to Record. 
+In the **Sequence Recorder** window, you can set the **Sequence Length** to **0** to perform continuous recording.
+
+![image alt text](image_9.png)(w:530 h:778 convert:false)
+
+[REGION:tip]When set to 0, recording will not stop unless you manually choose to stop recording by hitting the **Esc** key.[/REGION]
+
+### Right-click Recording
+
+If you have an active **Level Sequence** open, rather than clicking the **Add** button inside the **Sequence Recorder**, you can right-click on an Actor in your Level and initiate recording of the Actor from the context menu (or press **Alt + R**). This will add the Actor to the**Actor Tracking** window of the **Sequence Recorder** and automatically start recording after the set **Recording Delay** expires.
+
+![image alt text](image_10.png)(w:602 h:508 convert:false)
+
+[REGION:note]Please note that this will initiate Editor Recording and not [Gameplay Recording](https://docs.unrealengine.com/en-US/Engine/Sequencer/HowTo/SequenceRecorder) with the **Sequence Recorder**.[/REGION]
+
+### Subscene Recording Gameplay
+
+Using the **Sequence Recorder** to record gameplay will (by default) automatically create a new Level Sequence to record the content to. If you already have a Level Sequence that you want to add recorded content to, you can do so with the use of a [Subscene Track](https://docs.unrealengine.com/en-US/Engine/Sequencer/HowTo/TracksSub)and record the current Player Pawn.
 
 [OBJECT:EmbeddedVideo]
-[PARAMLITERAL:width]
-640
-[/PARAMLITERAL]
-[PARAMLITERAL:height]
-360
-[/PARAMLITERAL]
-[PARAMLITERAL:videoid]
-fMzumqLk7IQ
-[/PARAMLITERAL]
+	[PARAMLITERAL:videoid]
+	Wab8bB69TIo
+	[/PARAMLITERAL]
+	[PARAMLITERAL:width]
+	78
+	[/PARAMLITERAL]
+	[PARAMLITERAL:height]
+	56.25
+	[/PARAMLITERAL]
+	[PARAMLITERAL:units]
+	%
+	[/PARAMLITERAL]
+	[PARAMLITERAL:type]
+	simple
+	[/PARAMLITERAL]
 [/OBJECT]
 
-Above, we have placed a cube in our level and instructed the Sequence Recorder to record the cube. 
-We then moved the cube around and the Sequence Recorder captured the movements we entered through keyframes in a newly created Level Sequence. 
-When we play back our Level Sequence, a new Cube Actor is created (as a spawnable) in the Level Sequence, which is why a second Cube appears when the sequence is active.
+Above we create a new Level Sequence along with a Subscene Track, then during gameplay we hit **Shift + F1** to gain mouse control and add a primed Recorded Sequence based on our current Player Pawn. We can then press the **Record** button in the Viewport while playing our game to start recording and hit **Esc** to end recording.
+
+[REGION:note]Recording in this fashion is the same as manually adding the **Actor to Record** through the **Sequence Recorder** window and starting to Record from within the **Sequence Recorder**. One thing to keep in mind is that the settings defined inside the **Sequence Recorder** will be used when using this method of Subscene Recording.[/REGION]
+
+### Recording Class Properties
+
+If you want to record a change in a specific property value for an Actor during your Sequence Recording, you can use the **Classes and Properties to Record** setting to specify which Class and its associated property (or properties) you want to track for recording.
+
+Below, we select the **CineCameraComponent **Class, clicking the **... **button opens the **Choose Properties to Be Recorded **menu (right) where you can enable each of the highlighted checkboxes to track and record the property. 
+
+![image alt text](image_11.png)(w:624 h:380 convert:false)
+
+If that property changes at all during recording it will automatically be keyframed reflecting the changes during the recording. If the Actor you are tracking does not contain one of the defined classes it will just be ignored. 
+
+By default, **Sequence Recorder** includes the following classes to track for recording:
+
+* [SkeletalMeshComponent](https://docs.unrealengine.com/en-us/Engine/Components/SkeletalMesh)
+
+* [StaticMeshComponent](https://docs.unrealengine.com/en-us/Engine/Components/StaticMesh)
+
+* [ParticleSystemComponent](https://docs.unrealengine.com/en-us/Engine/Components/Rendering#particlesystemcomponent)
+
+* [LightComponent](https://docs.unrealengine.com/en-us/Engine/Components/Lights)
+
+* [CameraComponent](https://docs.unrealengine.com/en-us/Engine/Components/Camera)
+
+* [CineCameraComponent](https://docs.unrealengine.com/en-us/Engine/Sequencer/HowTo/CineCameraActors)
+
+You can use these classes and track properties within them during your recording or you can delete them and use your own classes.
+
+### Recording to Possessable
+
+By default when you elect to perform a recording the Actors you have tagged for tracking will automatically have a **Spawnable** version of that Actor created and added to your recording. While the Spawnable version will only exist while the Level Sequence is open, there may be times when you want to "write" directly to an Actor in your level instead of creating a new Spawnable. You can do so by checking the **Record to Possessable** option under the **Actor Recording** section in the Sequence Recorder.
+
+![image alt text](image_12.png)(w:508 h:309 convert:false)
+
+Checking this option will no longer create a Spawnable but instead will be using a Possessable Actor to apply the recording to.
+
+[REGION:note]Please see the [Possessables vs. Spawnables](https://docs.unrealengine.com/latest/INT/Engine/Sequencer/Overview/index.html#possessablesvsstopspawnables) section in the Sequencer Overview for more information.[/REGION]
+
+### Immersive Mode
+
+When recording your gameplay, you may want a larger Viewport to "play" your game while recording. The **Immersive Mode** option inside the **Sequence Recorder** window will automatically maximize your Viewport for you while recording is active. Once you stop recording, the Viewport will return back to its normal size.
+
+![image alt text](image_13.png)(w:508 h:381 convert:false)
 
-[REGION:note]
-Only properties that can be keyframed, can be captured and recorded when Level Editing recording. 
-[/REGION]
diff --git a/Engine/Documentation/Source/Engine/Sequencer/HowTo/WorkingWithTracks/WorkingWithTracks.INT.udn b/Engine/Documentation/Source/Engine/Sequencer/HowTo/WorkingWithTracks/WorkingWithTracks.INT.udn
new file mode 100644
--- /dev/null
+++ b/Engine/Documentation/Source/Engine/Sequencer/HowTo/WorkingWithTracks/WorkingWithTracks.INT.udn
@@ -0,0 +1,176 @@
+Availability: Public
+Crumbs: %ROOT%, Engine, Engine/Sequencer, Engine/Sequencer/HowTo
+title: Working with Tracks in Sequencer
+description:
+type: how-to
+SkillLevel: intermediate
+Version: 4.20
+Parent: Engine/Sequencer/HowTo
+Order: 2
+Tags: Sequencer
+Tags: how-to
+
+# Working with Tracks in Sequencer
+
+Whenever you create a Level Sequence, regardless of the type of cinematic you wish to create, you will be adding and working with Tracks inside the Tracks window. Sequencer offers several features that you can use to help organize and visualize those Tracks and the content within them. 
+
+In this how-to, we take a look at several different features of the Track system inside Sequencer and how to use those features.
+
+## Steps
+
+[REGION:note]For this guide we are using the **Blueprint Thirdperson** template with **Starter Content** enabled.[/REGION]
+
+1.  Inside your project in the **Content/Mannequin/Character/Mesh **folder, drag the **SK_Mannequin **into the Level. 
+
+	![image alt text](image_0.png)(w:624 h:250 convert:false)
+
+1.  From the Main Toolbar, click the **Cinematics **button and **Add Level Sequence **then give it any name and save location. 
+
+	![image alt text](image_1.png)(w:420 h:162 convert:false)
+
+1.  From the **World Outliner** window, drag the **SK_Mannequin **Actor into the **Tracks **window to add it to Sequencer. 
+
+	![image alt text](image_2.png)(w:624 h:353 convert:false)
+
+	This is an alternate method of adding Actors to Sequencer. In addition to the clicking the **+ Track **button, you can drag items from the **World Outliner **or even the **Content Browser**. 
+
+1.  Inside the **Content/StarterContent/Audio **folder, drag any of the **Sound Wave **assets into Sequencer. 
+
+	[OBJECT:EmbeddedVideo]
+		[PARAMLITERAL:videoid]
+		RlCorqp9wdg
+		[/PARAMLITERAL]
+		[PARAMLITERAL:width]
+		78
+		[/PARAMLITERAL]
+		[PARAMLITERAL:height]
+		56.25
+		[/PARAMLITERAL]
+		[PARAMLITERAL:units]
+		%
+		[/PARAMLITERAL]
+		[PARAMLITERAL:type]
+		simple
+		[/PARAMLITERAL]
+	[/OBJECT]
+
+	When dragging an asset into Sequencer, it will highlight the drop location to show you where it will be placed and pulsate yellow briefly when added to show you your newly added section. You can also drag any section onto any row within the Timeline or move the Timeline Marker to your desired location, then drop in an asset to create the section at that point (as shown in the video above). 
+
+1.  Add a **Folder **from the **+ Track **menu, rename it with **F2**, then drag-and-drop tracks into it or re-order it in the **Tracks List **window. 
+
+	[OBJECT:EmbeddedVideo]
+		[PARAMLITERAL:videoid]
+		LieJ347Qa9U
+		[/PARAMLITERAL]
+		[PARAMLITERAL:width]
+		78
+		[/PARAMLITERAL]
+		[PARAMLITERAL:height]
+		56.25
+		[/PARAMLITERAL]
+		[PARAMLITERAL:units]
+		%
+		[/PARAMLITERAL]
+		[PARAMLITERAL:type]
+		simple
+		[/PARAMLITERAL]
+	[/OBJECT]
+
+	You can group several Tracks together inside of **Folders **to help organize the Tracks window. Tracks, Actors, and Folders that are added to a Sequence are highlighted after creating them. If a folder is already currently selected, new items will be placed inside the folder. If another Track is currently selected, the new item will be placed as a sibling of that item. 
+
+	Tracks, Actors and Folders can also be reordered in Sequencer (as well as UMG Animations) and you can sort them (by type and then alphabetically) from the **General Options **menu. 
+
+1.  For the **SK_Mannequin **track, add any animations to the **Animation **track then scrub the track with the animation selected. 
+	
+	[OBJECT:EmbeddedVideo]
+		[PARAMLITERAL:videoid]
+		TXZnzMkgFvM
+		[/PARAMLITERAL]
+		[PARAMLITERAL:width]
+		78
+		[/PARAMLITERAL]
+		[PARAMLITERAL:height]
+		56.25
+		[/PARAMLITERAL]
+		[PARAMLITERAL:units]
+		%
+		[/PARAMLITERAL]
+		[PARAMLITERAL:type]
+		simple
+		[/PARAMLITERAL]
+	[/OBJECT]
+
+	When you scrub an animation inside Sequencer, the frame number for the source animation is displayed for your selected sections. You can use this to see where the start of an animation that is being blending starts or where the blended animation ends. 
+
+1.  Loop any of your animations, then right-click on it and select **Auto Size **under **Edit **from the context menu. 
+
+	![image alt text](image_3.png)(w:624 h:350 convert:false)
+
+	You can use the Auto Size feature to quickly resize a section to its source duration (for example, resizing Skeletal Mesh Animation to the bounds of their animation as we've done above). 
+
+1.  Add some movement for the character as shown below, then mask out the **Z **axis by deleting the channel or disabling it from the context menu. 
+
+	[OBJECT:EmbeddedVideo]
+		[PARAMLITERAL:videoid]
+		0pg9DgIrpGo
+		[/PARAMLITERAL]
+		[PARAMLITERAL:width]
+		78
+		[/PARAMLITERAL]
+		[PARAMLITERAL:height]
+		56.25
+		[/PARAMLITERAL]
+		[PARAMLITERAL:units]
+		%
+		[/PARAMLITERAL]
+		[PARAMLITERAL:type]
+		simple
+		[/PARAMLITERAL]
+	[/OBJECT]
+
+	You can also perform channel masking in UMG animations (below we mask out the Translation and Scale of an animation). 
+
+	[OBJECT:EmbeddedVideo]
+		[PARAMLITERAL:videoid]
+		jrTfOw0GhLo
+		[/PARAMLITERAL]
+		[PARAMLITERAL:width]
+		78
+		[/PARAMLITERAL]
+		[PARAMLITERAL:height]
+		56.25
+		[/PARAMLITERAL]
+		[PARAMLITERAL:units]
+		%
+		[/PARAMLITERAL]
+		[PARAMLITERAL:type]
+		simple
+		[/PARAMLITERAL]
+	[/OBJECT]
+
+1. Move the timeline marker to any location where a transform key is not present, then right-click any transform key and select **Rekey **to move it to the markers location.
+
+	[OBJECT:EmbeddedVideo]
+		[PARAMLITERAL:videoid]
+		nGglzSUID9I
+		[/PARAMLITERAL]
+		[PARAMLITERAL:width]
+		78
+		[/PARAMLITERAL]
+		[PARAMLITERAL:height]
+		56.25
+		[/PARAMLITERAL]
+		[PARAMLITERAL:units]
+		%
+		[/PARAMLITERAL]
+		[PARAMLITERAL:type]
+		simple
+		[/PARAMLITERAL]
+	[/OBJECT]
+
+	Using Rekey from the context menu is useful for quickly adjusting keyframes and assigning them to the current time. In the example above, we move our character with the Transform track however the animation continues while our movement has stopped. We use Rekey to move the Transform value to the end of the sequence matching the length of our animation. 
+
+## End Result
+
+In this guide you have learned about some of the different ways you can work with Tracks in Sequencer including how to organize those tracks with Folders, resizing sections, how to rekey a keyframe, channel masking, and other tips when adding Tracks to a Sequence. While you may not use all of these tips at once, its good to know that they are there and can be used situationally with your existing Sequencer workflow. 
+
diff --git a/Engine/Documentation/Source/Engine/Sequencer/Workflow/TimeRefactorNotes/TimeRefactorNotes.INT.udn b/Engine/Documentation/Source/Engine/Sequencer/Workflow/TimeRefactorNotes/TimeRefactorNotes.INT.udn
new file mode 100644
--- /dev/null
+++ b/Engine/Documentation/Source/Engine/Sequencer/Workflow/TimeRefactorNotes/TimeRefactorNotes.INT.udn
@@ -0,0 +1,160 @@
+Availability:Public
+Title:Sequencer Time Refactor Technical Notes
+Crumbs:%ROOT%, Engine
+Description:Technical information for the Time Refactor performed on the Sequencer tool for 4.20. 
+Navigation:topic
+Parent:Engine\Sequencer\Workflow
+order:1
+Version:4.20
+tags:Sequencer
+
+[TOC(start:2 end:2)]
+
+[REGION:warning]
+For the 4.20 release of UE4 substantial changes have been made to the way time is represented within Sequencer. This has been done to better support filmic pipelines and contexts where frame-accuracy is of huge importance. This page outlines the high-level conceptual changes, and details the breaking C++ code changes for upgrade purposes.
+[/REGION]
+
+## What has changed?
+
+Since its inception, **Sequencer** has represented time as a **number of seconds**, using a **floating point number**. This has made it very hard to support robust frame-accuracy, and has led to keyframe positioning errors, errors during rendering, unstable camera cut frames, ambiguous evaluation problems and more. As of the engine 4.20 release, the switch has been made to using integer ticks to represent time which resolves many, if not all of these problems.
+
+The way in which these ticks are interpreted is defined at a per-sequence level, through the sequence’s **Tick Resolution **which defines how many ticks there are per second. Using this mechanism, the user is in control of how granular sequencer data can be. For example, for real-time applications where frame-accuracy is less of a concern one might choose a very high resolution, whereas a Master Sequence for a rendered cinematic might choose the actual desired framerate for the project.
+
+This approach enables great flexibility for Sequencer in different contexts without compromising on robustness or performance. By default, content created before 4.20 will be upgraded with a tick resolution of **60,000 fps** which covers all common frame rates except **23.976 (NTSC)**. New sequences will default to a tick resolution of **24,000,** which covers all common whole number rates, as well as **23.976 (NTSC).**
+
+These cvars are useful for customizing what defaults to use for any given project:
+
+* MovieScene.LegacyConversionFrameRate (**Default: 60000fps**)
+
+    * Specifies default frame resolution for UMovieScene data saved before 4.20. 
+
+    * Has no effect on assets (re)saved from 4.20+.
+
+* LevelSequence.DefaultTickResolution (**Default: 24000fps**)
+
+    * Specifies a default tick resolution for newly created level sequences.
+
+* LevelSequence.DefaultDisplayRate (**Default: 30fps**)
+
+    * Specifies a default display frame rate for newly created level sequences; also defines frame locked frame rate where sequences are set to be frame locked.
+
+Bear in mind that Sequencer still supports evaluation of keyframe data using a floating-point subframe, which means we can still perform smooth interpolation at runtime even on sequences that have a very low tick resolution. Only section boundaries and the keyframes themselves are constrained by the tick resolution.
+
+All timing information for a given sequence now resides on the **Display Rate** drop-down on the Sequencer toolbar where the tooltip describes the current sequence’s tick resolution:
+
+![image alt text](image_0.png)(w:624 h:314 convert:false)
+
+This menu allows you to change the sequence’s Display Frame Rate, Time Display mode, and Evaluation Type. Common Display Rates that are compatible with the current sequence’s tick resolution are shown at the top level. It is possible to display using a frame rate that is incompatible with the tick resolution, but this is not advisable.
+
+Once created, it is possible to change the tick resolution of the sequence but bear in mind that converting a sequence to a tick resolution that is not a multiple of its current resolution may result in section boundaries and keys being rounded to the new resolution which can be found in the **Advanced Options **menu.
+
+![image alt text](image_1.png)(w:624 h:416 convert:false)
+
+## C++ Upgrade Notes
+
+In order to support the new time format, a large reworking of many internal Sequencer and MovieScene classes and data structures was required. Given the size and scope of the changes, it has not been possible to provide a graceful deprecation path for some APIs. 
+
+This section will explain in detail the major changes and provide an example code migration from 4.19 to 4.20.
+
+### Changes and Additions
+
+* **TimeManagement **module: provides common time-related data structures and manipulation and handling utilities:
+
+
+| Code | Description |
+| --- | --- |
+| FFrameNumber (32 bits) | A typesafe 32-bit integer primarily used to represent a context-free tick or a frame number. Does not support conversion from any other numeric type other than an int32. |
+| FFrameTime (64 bits) | A time representation comprising a context-free frame number and a floating point sub-frame. Primarily used during evaluation. |
+| FFrameRate (64 bits) | A fractional frame rate stored as an integral numerator and denominator (frames / seconds). |
+| FQualifiedFrameTime (128 bits) | A composition of FFrameTime and FFrameRate that allows for convenient conversion to/from seconds values and other frame rates. |
+| FTimecode (20 bytes) | A timecode representation comprising integers for hours, minutes, seconds and frames, including a boolean flag for drop frame/non-drop timecode |
+
+
+
+* **MovieScene Data**: UMovieScene now contains a tick resolution and display rate, and bolsters the previous ‘Force Fixed Frame Interval’ evaluation with an evaluation type enum. Notable changes to UMovieScene’s API include:
+
+
+| Code | Description |
+| --- | --- |
+| UMovieScene::GetEvaluationType() | Retrieves an enumeration specifying how to evaluate this sequence: [INCLUDE:#MovieSceneDataInsetTable] |
+| UMovieScene::GetTickResolution() | Retrieves the tick resolution that all FFrameNumbers within this movie scene should be interpreted with (unless otherwise stated) |
+| UMovieScene::GetDisplayRate() | Retrieves the frame rate that user-facing UI, Blueprint nodes and Sequence Players use when interfacing with this sequence. Also defines t.maxfps and the locked framerate when EvaluationType is set to FrameLocked. |
+
+<!---
+[EXCERPT:MovieSceneDataInsetTable]
+| Code | Description |
+|--- | --- |
+| WithSubFrames (default) | Evaluate using sub-frame interpolation |
+| FrameLocked | Lock to the DisplayRate of the sequence, only evaluate round frame numbers, no subframes, set t.maxfps during evaluation. |
+[/EXCERPT]
+--->
+
+
+All movie scene data relating to time has been converted to using FFrameNumber, based on the tick resolution of the owning UMovieScene. This includes: 
+
+* UMovieSceneSection ranges. Section ranges are now exposed exclusively as TRange<FFrameNumber>. This range now also implies infinite upper or lower bounds, and removes the need for a separate IsInfinite flag. Sections that support infinite ranges can now set bSupportsInfiniteRange to true. 
+
+* UMovieScene playback and selection ranges. Both inclusive and exclusive upper and lower bounds are now correctly handled.
+
+* Pre/Post roll amounts
+
+* UMovieSceneSubSection start offsets (in the inner sequence’s tick resolution)
+
+* **Curve Data**: All existing curve classes (for example FRichCurve, FIntegralCurve, FNameCurve ) are replaced internally with new channel types that operate on FFrameNumber times in the owning UMovieScene’s tick resolution, rather than floats. Backwards compatibility with floating point time representations is not possible due to the disparity in time types. Automatic upgrade from their legacy counterparts is supported through SerializeFromMismatchedTag, using the legacy upgrade tick resolution as outlined above.
+
+    * FRichCurve becomes FMovieSceneFloatChannel
+
+    * FIntegralCurve becomes FMovieSceneIntegerChannel
+
+    * FStringCurve becomes FMovieSceneStringChannel
+
+    * FNameCurve is no longer used internally
+
+    * FMovieSceneByteChannel (including enums) and FMovieSceneBoolChannel have been added.
+
+Custom curve data types that were previously using TCurveInterface, should be trivial to migrate to the new FMovieSceneXYZChannel model as long as they provide a TMovieSceneChannelData<> GetData() method.
+
+### Movie Scene Channels
+
+Given the necessary breaking changes to curve (now referred to as channel) data, the opportunity was also taken to unify several editor codepaths and reduce duplication for sections that support keyframes. A new type, FMovieSceneChannelProxy, has been added which affords editor and runtime code a common language for interacting with and manipulating keyframes. 
+
+To this end, IKeyframeSection<> has been completely removed and is no longer necessary. Sections that support keyframes now just need to populate their channel proxy correctly in order to support generic keyframe interaction. This change will also make it much simpler to compose new section types with keyframes in the future since the common key-manipulation interface now exists at the channel level, rather than the section level.
+
+There are several components to this new system, with a slightly different approach to editor customization:
+
+#### FMovieSceneChannelProxy
+
+Stored on UMovieSceneSection: derived types should populate this structure with all of its channels as in the example below. Channels are stored by their base FMovieSceneChannel* in buckets by derived type. With this in mind, any reallocation of channels should be immediately followed by a re-creation of the channel proxy; doing so will invalidate any pointers and handles to the channels stored in the old proxy.
+
+All interaction with channels is through either the FMovieSceneChannel interface directly, or ISequencerChannelInterface, depending on context. The latter is registered per-type through the sequencer module (ISequencerModule::RegisterChannelEditor).
+
+[REGION: warning]Failure to do so will result in a failed assertion at runtime when trying to edit such channels. A templated helper is provided through TSequencerChannelInterface which allows single-concept overloading for any given channel type, resolved through argument-dependent lookup (ADL). [/REGION]
+
+This allows customization of specific behavior without having to re-implement the entire interface if the defaults are suitable for most channels. It also means that core sequencer code can automatically populate UI for channel data without having to manually define ISequencerSection interfaces and manually defining the channel layout in the editor as well as in the runtime. This streamlines the creation of custom section types to a degree, particularly when composing sections from already supported intrinsic channel data.
+
+Default implementation functions for ISequencerChannelInterface are defined in the Sequencer namespace, but overloads should be added either to your channel’s namespace or the global namespace.
+
+[REGION: note]It is recommended that custom channels follow the pattern of storing times and values in parallel arrays, and provide a TMovieSceneChannelData<T> GetData() method for interacting with the keys. The majority of FMovieSceneChannel’s interface directly maps to functions callable on TMovieSceneChannelData.[/REGION]
+
+Using this proxy approach, all data required for constructing UI and interacting with the channel’s keys is readily available to Sequencer UI code. 
+
+#### FMovieSceneChannel
+
+Provides an interface through which all common channel data can be interacted with. All channels added to the channel proxy must implement this type.
+
+#### ISequencerChannelInterface
+
+An interface to all overloads relating to UI interaction and manipulation required by sequencer for a given channel. Must be registered through the ISequencerModule class for each channel type (normally in your editor module’s StartupModule method).
+
+#### TMovieSceneChannelTraits
+
+Specifies compile-time traits for channel types such as extended editor data, and whether the channel supports default values. Many sequencer UI utilities use the function overloads specified in MovieSceneChannelTraits.h to interact with each concrete channel type. If the default templated overloads are incompatible with your channel type, you should overload the necessary functions for the specific channel.
+
+### Other Utilities
+
+Several helper utilities have been added to make working with movie scene ranges easier. Now that we use a discrete time base, it’s easier to rationalize about section boundaries. Internally ranges are represented as TRange<FFrameNumber>, which supports inclusive, exclusive and open bounds. In order to make dealing with these ranges simpler for discrete time bases, you can use MovieScene::DiscreteInclusiveLower, DiscreteExclusiveUpper and DiscreteSize to consistently deal with the various boundary conditions.
+
+### Migration Guide
+
+Due to the nature of these changes, there will be C++ changes that require migration if you have written your own tracks, sections or curves in C++. Please see the prepared [example migration]([REDACTED_PROFILE_URL]/custom-sequencer-example/commits/all) project for a track, section and channel type to assist in your own upgrade.
+
diff --git a/Engine/Documentation/Source/Engine/UI/ColorPicker/ColorPicker.INT.udn b/Engine/Documentation/Source/Engine/UI/ColorPicker/ColorPicker.INT.udn
--- a/Engine/Documentation/Source/Engine/UI/ColorPicker/ColorPicker.INT.udn
+++ b/Engine/Documentation/Source/Engine/UI/ColorPicker/ColorPicker.INT.udn
@@ -1,22 +1,19 @@
 Availability: Public
-Title:Color Picker
-Crumbs:%ROOT%, Engine, Engine/UI
-Description:Tool for interactively choosing colors to assign to color properties of actors.
+Title: Color Picker
+Crumbs: %ROOT%, Engine, Engine/UI
+Description: Tool for interactively choosing colors to assign to color properties of actors.
 Version: 4.9
-topic-image:Engine/UI/colorPicker_topic.png
+topic-image: Engine/UI/colorPicker_topic.png
 
-
-[EXCERPT:Intro]
-The **Color Picker** allows you to easily adjust a color value for a color property in the Unreal Editor.  You can adjust colors by RGBA (red, green, blue, alpha), HSV (hue, saturation, value), and Hex (ARGB). Alternatively you can select a color by clicking anywhere in the color wheel or by capturing the color under your mouse, anywhere on your monitors.
-[/EXCERPT:Intro]
+The **Color Picker** allows you to easily adjust a color value for a color property in the Unreal Editor. You can adjust colors by RGBA (red, green, blue, alpha), HSV (hue, saturation, value), and Hex (ARGB). Alternatively you can select a color by clicking anywhere in the color wheel or by capturing the color under your mouse, anywhere on your monitors.
 
 ![ColorPicker.jpg](ColorPicker.png)
 
 | Item | Description |
-| ---- | ----------- |
-| ![Themes & sRGB](ThemeSRGB.png) | Allows for [Theme](#Themes) selection and **sRGB** previews (see note below).|
-| ![Color Wheel](ColorWheel.png) | A Color wheel that allows a color to be selected graphically by clicking and dragging with the mouse. Additional adjustments can be made by using the two vertical bars. Instead of Color Wheel, a Color Spectrum may be used if the option is toggled on.  |
-| ![Color Swatches](ColorSwatches.png) | Displays the color currently assigned (Old) to the property and the color currently selected in the Color Picker (New). Colors can be dragged from either the _Old_ or _New_ sections to the **Themes** bar to save them for reuse later on. |
+| --- | --- |
+| ![Themes & sRGB](ThemeSRGB.png) | Allows for [Theme](Engine/UI/ColorPicker#themes) selection and **sRGB** previews (see note below). |
+| ![Color Wheel](ColorWheel.png) | A Color wheel that allows a color to be selected graphically by clicking and dragging with the mouse. Additional adjustments can be made by using the two vertical bars. Instead of Color Wheel, a Color Spectrum may be used if the option is toggled on. |
+| ![Color Swatches](ColorSwatches.png) | Displays the color currently assigned (Old) to the property and the color currently selected in the Color Picker (New). Colors can be dragged from either the *Old* or *New* sections to the **Themes** bar to save them for reuse later on. |
 | ![Picker Options](ColorOptions.png) | Toggles between the Color Wheel or Color Spectrum (left image). The **Eye Dropper** (right image) can be used for selecting a colored pixel from any window inside or outside of Unreal Engine. |
 | ![Red Channel](RedChannel.png) | Drag or type value to set the **Red** channel. |
 | ![Green Channel](GreenChannel.png) | Drag or type value to set the **Green** channel. |
@@ -32,11 +29,9 @@ When working with the Color Picker, it does not assume a straight inverse gamma
 [REGION:note]
 `FColor` and `FLinearColor` can also default to the sRGB space using a UPROPERTY meta tag of **sRGB=true** (or **false**), you can force a specific default checking of the sRGB checkbox.
 
-Several places in [Unreal Motion Graphics (UMG)](Engine/UMG) now force this default since sRGB is the space UMG will be rendered in (on PC). Additional work will be needed on Mac where the final rendering space is gamma 2.2. 
+Several places in [Unreal Motion Graphics (UMG)](Engine/UMG) now force this default since sRGB is the space UMG will be rendered in (on PC). Additional work will be needed on Mac where the final rendering space is gamma 2.2.
 [/REGION]
 
-
-
 ## Themes
 
 **Themes** are collections of colors that can be easily reused. These can be used to make commonly used colors available or to restrict the colors used by artists and designers to a specific palette.
@@ -45,45 +40,57 @@ Several places in [Unreal Motion Graphics (UMG)](Engine/UMG) now force this defa
 
 ### Creating Themes
 
-New themes can be added through the **Theme Menu** accessible via the ![Theme Menu](button_menu.png) button.
-
-![Theme Menu](theme_menu.png)
-
-Clicking the ![New Theme](button_new.png) button will add an empty theme to the list.
+New themes can be added by clicking the **Theme Menu** button.  
+![](ThemeMenuButton.png)
 
-![Theme Menu](theme_menu_new.png)
+Click the **New Theme** button to add an empty theme to the list.  
+![](newtheme_2.png)
 
-Themes can also be created by duplicating an existing theme using the ![Duplicate](button_duplicate.png) button.
-
-![Duplicated Theme](theme_duplicate.png)
+Click the **Duplicate** button to create a new theme that copies an existing theme.  
+![](theme_duplicate.png)
 
 ### Renaming Themes
 
-Themes can be renamed through the **Theme Menu** by clicking the ![Rename](button_rename.png) button.
+Themes can be renamed by clicking the **Rename** button.  
+![](theme_rename.png)
 
-[REGION:imagetable]
-| ![Theme Menu](theme_menu_rename_0.png)(w:300) | ![Theme Menu](theme_menu_rename_1.png)(w:300) | ![Theme Menu](theme_menu_rename_2.png)(w:300) |
-| ------ | ------ | ------ |
-| Edit box for theme name | Enter new name for theme | Press ![Accept](button_accept.png) to set the name |
-[/REGION]
+This displays a text field, with the default text "New Theme".  
+![](theme_rename_editbox.png)
+
+Type in a new name for the theme. Then click the **Accept** button.  
+![](theme_rename_newname.png)
 
 ### Adding and Removing Colors
 
 Colors can be added, removed, or reordered within the theme using drag and drop.
 
-[REGION:imagetable]
-| ![Add color](theme_add_color.png)(w:300 convert:false) | ![Reorder Colors](theme_edit.png)(w:300 convert:false) | ![Remove Color](theme_edit_color_delete.png)(w:300 convert:false) |
-| ------ | ------ | ------ |
-| Drag from main swatch to add a color | Drag color within theme to reorder | Drag color to trashcan to delete |
-[/REGION]
+Drag colors from the main swatch to add them to a theme.  
+![](theme_add_color.png)
+
+Reorder colors within a theme by dragging them right or left.  
+![](theme_edit.png)
+
+Drag colors to the **Trashcan** icon to delete them from the theme.  
+![](theme_edit_color_delete.png)
 
 Colors can also be reordered and removed from the **Theme Menu**.
 
-[REGION:imagetable]
-| ![Reorder Colors](theme_edit_menu.png)(convert:false) | ![Remove Color](theme_menu_color_delete.png)(convert:false) |
-| ------ | ------ |
-| Drag color within theme to reorder | Drag color to trashcan to delete |
-[/REGION]
+To edit the theme, drag colors left or right.  
+![](theme_menu_edit.png)
+
+To delete a color from the theme, drag it to the **Trashcan** icon.  
+![](theme_menu_delete.png)
+
+### Labeling Saved Colors in a Theme
+
+You can label colors saved in your **Theme Bar** by right-clicking a saved color, and typing a name for the color in the **Color Label** window.  
+![](ColorLabel-2.png)
+
+When you hover the mouse cursor over the color that has been labeled, you will see the name listed in the tooltip properties for the color.  
+![](ColorLabel-3.png)
+
+Colors can also be labeled from the **Theme Menu**. Right-click the color, then type a name for the color into the **Color Label** window.  
+![](ColorLabel-4.png)
 
 ### Using Themes
 
@@ -91,10 +98,10 @@ To change themes, select the theme you want to be active from the **Theme Menu**
 
 ![Select Theme](theme_select.png)
 
-To use a color from the current theme, **Double-click** the color to be applied. This will update the currently selected color in the Color Picker. 
+To use a color from the current theme, double-click the color to be applied. This will update the currently selected color in the Color Picker.
 
 ![Use Color](theme_use_color.png)
 
 Hovering the cursor over a color in the theme will display the color information.
 
-![Color Info](theme_color_info.png)
+![Color Info](theme_color_info.png)
\ No newline at end of file
diff --git a/Engine/Documentation/Source/Engine/UI/EditorSettings/EditorPreferences.INT.udn b/Engine/Documentation/Source/Engine/UI/EditorSettings/EditorPreferences.INT.udn
--- a/Engine/Documentation/Source/Engine/UI/EditorSettings/EditorPreferences.INT.udn
+++ b/Engine/Documentation/Source/Engine/UI/EditorSettings/EditorPreferences.INT.udn
@@ -37,21 +37,22 @@ The **Editor Preferences** window can be opened from the **Edit** menu.
 | Option | Description |
 | ------ | ----------- |
 | **Highlight Objects Under Mouse Cursor** | If enabled, objects will be highlighted using a transparent overlay in viewports while hovering over them with the mouse cursor. 
+
 [REGION:none]
-[OBJECT:EmbeddedVideo]
-[PARAMLITERAL:width]
-420
-[/PARAMLITERAL]
-[PARAMLITERAL:height]
-200
-[/PARAMLITERAL]
-[PARAMLITERAL:videoid]
-yvwZUtmW12I
-[/PARAMLITERAL]
-[/OBJECT]
+	[OBJECT:EmbeddedVideo]
+	[PARAMLITERAL:width]
+	420
+	[/PARAMLITERAL]
+	[PARAMLITERAL:height]
+	200
+	[/PARAMLITERAL]
+	[PARAMLITERAL:videoid]
+	yvwZUtmW12I
+	[/PARAMLITERAL]
+	[/OBJECT]
 [/REGION] 
 
-|
+[COMMENT:none]
 | **Highlight Selected Objects with Brackets** | [INCLUDE:Shared/Types/ULevelEditorViewportSettings#bhighlightwithbrackets] |
 | **Use Selection Outline** | [INCLUDE:Shared/Types/ULevelEditorViewportSettings#buseselectionoutline] |
 | **Selection Highlight Intensity** | [INCLUDE:Shared/Types/ULevelEditorViewportSettings#selectionhighlightintensity] |
@@ -59,7 +60,7 @@ yvwZUtmW12I
 | **Use Absolute Translation** | If enabled, translation transformations are treated as absolute. When disabled, translations are considered relative to the prior location. |
 | **Preview Selected Cameras** | If enabled, selecting a camera Actor will display a live _picture in picture_ preview from the camera's perspective within the current editor viewport.  This can be used to easily tweak camera positioning, post-processing and other settings without having to possess the camera itself. [REGION:warning]This feature may reduce application performance when enabled.[/REGION] |
 | **Background Drop Distance** | The distance in Unreal Units from the camera to place Actors which are dropped on nothing in the viewport. |
-
+[/COMMENT]
 ### Editor Look and Feel
 
 | Option | Description |
@@ -142,7 +143,7 @@ yvwZUtmW12I
 
 [EXCERPT:SelectionBrackets]
 [REGION:imagetable]
-| ![](selection_brackets.png)(w:355) |![](selection_highlight.png)(w:355) |
+| ![](Shared/Types/ULevelEditorViewportSettings/selection_brackets.png)(w:355) |![](Shared/Types/ULevelEditorViewportSettings/selection_highlight.png)(w:355) |
 | :--------------------------------: | :-------------------------------: |
 | **Enabled** | **Disabled** |
 [/REGION]
@@ -150,7 +151,7 @@ yvwZUtmW12I
 
 [EXCERPT:SelectionOutline]
 [REGION:imagetable]
-| ![](selection_outline.png)(w:355) |![](selection_highlight.png)(w:355) |
+| ![](Shared/Types/ULevelEditorViewportSettings/selection_outline.png)(w:355) |![](Shared/Types/ULevelEditorViewportSettings/selection_highlight.png)(w:355) |
 | :--------------------------------: | :-------------------------------: |
 | **Enabled** | **Disabled** |
 [/REGION]
@@ -158,7 +159,7 @@ yvwZUtmW12I
 
 [EXCERPT:SelectionHighlight]
 [REGION:imagetable]
-| ![](selection_highlight.png)(w:355) |![](selection_highlight_full.png)(w:355) |
+| ![](Shared/Types/ULevelEditorViewportSettings/selection_highlight.png)(w:355) |![](Shared/Types/ULevelEditorViewportSettings/selection_highlight_full.png)(w:355) |
 | :--------------------------------: | :-------------------------------: |
 | **0.2** | **1.0** |
 [/REGION]
diff --git a/Engine/Documentation/Source/Engine/UMG/UserGuide/UMGSafeZones/UMGSafeZones.INT.udn b/Engine/Documentation/Source/Engine/UMG/UserGuide/UMGSafeZones/UMGSafeZones.INT.udn
new file mode 100644
--- /dev/null
+++ b/Engine/Documentation/Source/Engine/UMG/UserGuide/UMGSafeZones/UMGSafeZones.INT.udn
@@ -0,0 +1,80 @@
+Availability: Public
+Crumbs: 
+Title: UMG Safe Zones
+Description: This page describes Safe Zones, which are designed to keep the UI from displaying somewhere the player can't see.
+Version: 4.20
+Parent: Engine/
+Order: 5 
+Tags: UMG UI Designer
+Tags: UMG
+
+
+**Safe Zone** widgets are an essential part of developing a game user interface (UI) that can run on a lot of different devices. Safe zones are designed to keep the UI from displaying somewhere that it can technically use but the player can’t see, such as the edge of a TV display or under the notch and home bar of the iPhoneX. The UMG Designer enables you to test the resolution (or rotation) of the device with your UI and applied Safe Zone widgets.
+
+When you add a **Safe Zone** widget to the **Designer** , the widget will scale any other widgets that are made children of it from the **Hierarchy** panel. These child widgets will scale and adjust anytime there is an “unsafe” zone.
+
+[OBJECT:ComparisonSlider]
+ [PARAM:before]
+ ![Without Safe Zone widget](WithoutSafeZone_opt-1.png) 
+ [/PARAM]
+ [PARAM:after]
+ ![With Safe Zone widget](WithSafeZone_opt.png) 
+ [/PARAM]
+[/OBJECT]
+
+In this example, a 1080p display has a **Uniform Safe Zone** area of 0.9 (red) for testing and debugging purposes. When the widgets are parented to a Safe Zone widget, the child widgets will scale to the safe zone area. This prevents widgets from being cut off at the screen edge deemed “unsafe”. This can be seen in the example, with the “My Menu” title text.
+
+## Setting and Testing Safe Zone Resolutions
+
+In [UMG](Engine/UMG) (or for [Play-in-Editor](Engine/UI/LevelEditor/InEditorTesting) settings), the screen sizes chosen are now linked with [Device Profiles](Platforms/DeviceProfiles) which also takes into account the [Mobile Content Scale Factor](Platforms/Mobile/Performance#mobilecontentscalefactor), meaning that the final resolution and DPI scale will change based on the selected device.
+
+To test the screen resolution of a device, use the UMG **Designer** viewport to select the **Screen Size** dropdown and select from the listed devices.  
+![ScreenSizeSelection.png](ScreenSizeSelection_opt.png)
+
+If the device supports screen orientation, like a phone or tablet, use the **Portrait/Landscape** button to swap between the two view modes. This option will be grayed out when the device selected doesn’t support orientation shifting.  
+  
+![](OrientationButton.png)
+
+Here is an example of the two view modes: ![](iPhoneX_PortraitAndLandscape.png)
+
+[REGION:caption]
+*Left: Landscape, Right: Portrait*
+[/REGION]
+
+When a device is selected, relevant information and marked “unsafe” zones will be displayed in the **Designer** graph.  
+![](DeviceDetails.png)
+
+[REGION:caption]
+1.  Marked “Unsafe” Zones
+2.  Device Details for Mobile Content Scale Factor, Device Name or Uniform Safe Zone, and Screen Size
+3.  DPI Scale
+[/REGION]
+
+For non-uniform safe zones, you can use the **Flip** button to simulate device rotation for Landscape mode.  
+![](DeviceFlipButton.png)
+
+[REGION:note]
+This option is only available on some devices and only for Landscape view mode.
+[/REGION]
+
+When previewing TVs and Monitors from the list, if you’ve set a [debug title safe zone](Platforms/TVSafeZoneDebugging) that is less than 1, the **Uniform Safe Zone** will display that size. The debug safe zone area is represented by the red area around the canvas widgets.  
+![](UniformSafeZone.png)
+
+[REGION:note]
+The **Debug Safe Zone** mode is on by default in UMG which enables the red “unsafe” zone area to be represented.
+[/REGION]
+
+For some devices, there are now custom “unsafe” zones represented within the Designer graph. These are used to represent hardware or software specific areas of the device that take up screen real estate, such as the iPhoneX where these zones represent the notch and home bar sections of the screen.  
+![](iPhoneXUnsafeZones.png)
+
+In this example of the iPhoneX, the Safe Zone widget is parented to the section of the widget hierarchy where the menu and buttons are contained. When areas are “unsafe,” the encapsulated widgets scale to avoid these areas (see example above).
+
+![](SafeZoneHierarchyPanel.png)
+
+## Testing UI in PIE and Standalone Game
+
+When testing your UI in the Editor, you can set the screen size based on a common device screen size when using PIE with a **New Editor Window** or when using **Standalone Game** without deploying content to a device. Use the **Editor Preferences** to set the device under **Level > Play > Game Viewport Settings** to choose a common device resolution.  
+![](PIESettings.png)
+
+Use the **Landscape/Portrait Orientation** button to swap which orientation should be used.  
+![](PIESettings_SwapOrientation.png)
\ No newline at end of file
diff --git a/Engine/Documentation/Source/Enterprise/Automotive/AutomotiveMaterialPack/AutomotiveMaterialPack.INT.udn b/Engine/Documentation/Source/Enterprise/Automotive/AutomotiveMaterialPack/AutomotiveMaterialPack.INT.udn
--- a/Engine/Documentation/Source/Enterprise/Automotive/AutomotiveMaterialPack/AutomotiveMaterialPack.INT.udn
+++ b/Engine/Documentation/Source/Enterprise/Automotive/AutomotiveMaterialPack/AutomotiveMaterialPack.INT.udn
@@ -8,7 +8,7 @@ tags:Studio
 tags:Automotive
 tags:Materials
 type:overview
-order:1
+order:100
 Related:Engine/Rendering/Materials/PhysicallyBased/
 Related:Engine/Rendering/Materials/MaterialProperties/LightingModels
 Related:Engine/Rendering/Materials/HowTo/ClearCoatDualNormal
diff --git a/Engine/Documentation/Source/Enterprise/Enterprise.INT.udn b/Engine/Documentation/Source/Enterprise/Enterprise.INT.udn
--- a/Engine/Documentation/Source/Enterprise/Enterprise.INT.udn
+++ b/Engine/Documentation/Source/Enterprise/Enterprise.INT.udn
@@ -13,7 +13,11 @@ topic-icon:%ROOT%/enterprise_icon.png
 ![](Enterprise_Hero.png)
 
 Unreal Engine and its real-time interaction and rendering capabilities are well-suited for a huge range of projects across all Enterprise applications, including automotive, aviation, architecture, consumer electronics and complex data visualization.
- 
+
+[REGION:note]
+We have provided early access to some content in the form of PDF documents in order to get you up and running. Check back soon for updates and an improved experience!
+[/REGION]
+
 ## Getting Started
 
 [DIR(output:"topic" org:"hierarchy" parent:"Enterprise" tags:"Getting Started" end:"1")]
diff --git a/Engine/Documentation/Source/Gameplay/ClassCreation/CodeAndBlueprints/CodeAndBlueprints.INT.udn b/Engine/Documentation/Source/Gameplay/ClassCreation/CodeAndBlueprints/CodeAndBlueprints.INT.udn
--- a/Engine/Documentation/Source/Gameplay/ClassCreation/CodeAndBlueprints/CodeAndBlueprints.INT.udn
+++ b/Engine/Documentation/Source/Gameplay/ClassCreation/CodeAndBlueprints/CodeAndBlueprints.INT.udn
@@ -194,7 +194,7 @@ executes whenever the C++ function `OnOverlap` would be called. In the LightSwit
 The **OnOverlap** event node is connected to a **Set Light Color** node. Whenever the event executes, it will cause the light color of the PointLightComponent to be set to a random color. This overrides
 the `OnOverlap_Implementation` function in the source file, which toggles the visibility of the PointLightComponent.
 
-For more on events and working with graphs, see the [](Engine/Blueprints\UserGuide\Events), [](Engine/Blueprints\UserGuide\EventGraph), and [Graph Panel](Engine/Blueprints\Editor\Modes\GraphPanel) documentation.
+For more on events and working with graphs, see the [](Engine/Blueprints\UserGuide\Events), [](Engine/Blueprints\UserGuide\EventGraph), and [Blueprint Class UI](Engine/Blueprints/Editor/UIBreakdowns/ClassBPUI) documentation.
 
 ![](Both_EventGraph_2.png)
 
diff --git a/Engine/Documentation/Source/Gameplay/ForceFeedback/ForceFeedback.INT.udn b/Engine/Documentation/Source/Gameplay/ForceFeedback/ForceFeedback.INT.udn
--- a/Engine/Documentation/Source/Gameplay/ForceFeedback/ForceFeedback.INT.udn
+++ b/Engine/Documentation/Source/Gameplay/ForceFeedback/ForceFeedback.INT.udn
@@ -10,7 +10,7 @@ Related: Engine/Blueprints/UserGuide/Timelines/KeysAndCurves
 version: 4.15
 type:overview
 tags:Gameplay
-topic-image: force_feedback_topic
+topic-image: force_feedback_topic.png
 
 [toc(start:1 end:2)]
 
diff --git a/Engine/Documentation/Source/Platforms/MR/MR.INT.udn b/Engine/Documentation/Source/Platforms/MR/MR.INT.udn
new file mode 100644
--- /dev/null
+++ b/Engine/Documentation/Source/Platforms/MR/MR.INT.udn
@@ -0,0 +1,25 @@
+Availability: Public
+Crumbs: %ROOT%, Platforms
+Title: Mixed Reality Development
+Description: Mixed Reality Development
+Platform: Mixed Reality
+Version: 4.20
+Parent: Platforms
+Order: 0
+Tags: MR
+Tags: Early Access
+Tags: Landing Page
+
+
+The Mixed Reality Capture (MRC) framework makes it easy for users to composite real players into a virtual play space.
+
+## Getting Started
+
+%Platforms/MR/MROverview:topic%
+%Platforms/MR/MRQuickStart:topic%
+%Platforms/MR/MRHowToCaptureCalibrationTool:topic%
+
+## Reference
+
+%Platforms/MR/MRSupportedVideoDevices:topic%
+%Platforms/MR/MRTroubleshooting:topic%
diff --git a/Engine/Documentation/Source/Platforms/MR/MRHowToCaptureCalibrationTool/HowToCaptureCalibrationTool.INT.udn b/Engine/Documentation/Source/Platforms/MR/MRHowToCaptureCalibrationTool/HowToCaptureCalibrationTool.INT.udn
new file mode 100644
--- /dev/null
+++ b/Engine/Documentation/Source/Platforms/MR/MRHowToCaptureCalibrationTool/HowToCaptureCalibrationTool.INT.udn
@@ -0,0 +1,385 @@
+Availability: Public
+Crumbs: %ROOT%, Platforms, Platforms/MR
+Title: How To Use The Mixed Reality Capture Calibration Tool
+Description: Product documentation including reference and guides for Unreal Engine 4
+Version: 4.20
+Parent: Platforms/MR
+Tags: GettingStarted
+Tags: MR
+Tags: landingpage
+Tags: Early Access
+topic-image:AdjustingAlignmentStep.png
+reference-image:AdjustingAlignment2.gif
+reference-image:MR_AdjustSample.gif
+reference-image:MR_CompletingCalibration.gif
+reference-image:MR_GarbageMattingControls.gif
+
+The goal of calibration is to have the virtual and physical cameras match up in the virtual world. 
+
+![](AdjustingAlignmentStep.png "AdjustingAlignmentStep.png")  
+At runtime, the game needs to know a few things, like: 
+
+*   The type of camera and lens that you’re using.
+*   Where your camera is in relation to your virtual scene. 
+
+And to composite your scene, the game needs to know things like:
+
+*   The color of your chroma backdrop (usually green).
+*   How aggressively to cut out the chroma color.
+*   Areas you want fully cut out from the camera frame. 
+
+The game needs to know these things so the virtual camera can emulate the physical one. This is where the calibration tool comes into play. The calibration tool lets you configure all of these settings at once, and produces a calibration settings file with that information. The settings file can then be reused across different games.
+
+## Pre-Calibration Setup
+
+Before you can start a Mixed Reality capture, you will need the right equipment. Here is quick overview of the equipment you will need, and some tips on how set it up. This include the basic green screen and camera setup for the capture process, as well as the equipment and software needed for calibration.
+
+#### Green Screen and Camera
+
+![](MR_GreenScreenSetup.png "MR_GreenScreenSetup.png")
+
+1.  **Video Camera**  
+    UE4 only supports a very specific set of video capture devices; see [Supported Video Devices](https://docs.unrealengine.com/en-us/Platforms/MR/MRSupportedVideoDevices) for the current list of the supported devices. Using a device from the list, connect it to your PC for streaming.
+2.  **Chroma backdrop**  
+    For chroma keying, we generally use a green screen. When setting up the green screen you will want to make sure it’s taut, minimizing any wrinkles, especially behind your subject. If you are setting up lighting, you will want to make sure not to cast shadows directly behind the subject. You want a smooth, flat color. The more shades of green that there are across the backdrop, the harder it is to chroma key. It helps if the subject is far away from the backdrop as possible. If you plan on filming the subject’s feet, consider a green screen for the floor as well.
+3.  **Camera Mount**  
+    For initial setup (calibration), your camera needs to be static. If you are using a webcam this may be as simple as attaching it to your desk/monitor. Another option is mount the camera to a tripod.
+4.  **Multi-Mount + Tracker (Optional)**  
+    If you plan on moving the camera around while filming, you will want attach a tracking device to the camera, such as an HTC Vive Tracker. In addition, connect the camera and tracker together with a [multi mount](https://www.bhphotovideo.com/c/product/1062513-REG/desmond_d3d_1_stereo_camera_bracket.html), so that they stay firmly locked in place.
+
+#### Calibration Specific Equipment
+
+There are a few additional items required specifically for the calibration process (as opposed to the capture process).
+
+![](MR_RequiredSetupItems.jpg "MR_RequiredSetupItems.jpg")
+
+*   **HTC Vive or Oculus Rift**  
+    The calibration tool requires one of these two VR systems to work. The tool uses the controllers’ tracking to evaluate the location of the camera. If this is the first time you’ve set up an HTC Vive or Oculus system, you will need to complete their setup process first for calibration to work properly.
+*   **Printed Checkerboard Pattern**  
+    Print and attach a printed checkerboard pattern to a flat rigid surface (like a piece of cardboard). You can find an example of this checkerboard pattern included with the calibration tool download.
+    [REGION:tip]
+    When taping down the checkerboard printout, make sure not to cover the checkerboard pattern itself with tape (even clear tape can cause specular reflections that make it hard for the camera to read).
+    [/REGION]
+
+## Steps
+
+Once the equipment is set up, you can download the calibration tool from [here](http://epic.gm/mrccal). Once you've downloaded the file, unzip it and launch MRCalibration.exe.
+
+The calibration process is divided into distinct steps. The individual steps are discussed further below.
+
+1.  Device/Camera/Tracker Selection
+2.  Lens Calibration
+3.  Alignment Calibration
+4.  Alignment Adjustment
+5.  Compositing Calibration
+6.  Garbage Matting
+
+During calibration, at the end of each step your progress is saved. So you can exit and return to the tool if needed. Once the calibration process is complete, the tool generates a settings file that you will then use to start a MRC session in other projects.
+
+[REGION:tip]
+Your calibration progress is recorded in the setting file. If you want to restart the calibration process, you will need to delete the settings file. The settings file is named *MrcCalibration.sav*, and can be found under the tool’s */Saved/SaveGames/* directory.
+[/REGION]
+
+#### Shared Controls
+
+While each step has its own distinct controls, there are some controls that are used consistently through the tool:
+
+*   **Enter** : Next/Confirm/Accept/Submit
+*   **End** : Skip (only applicable when step requirements have been met - some steps can’t be skipped)
+*   **P/Thumbstick** : Preview
+*   **M** : Mirror the video feed
+*   **R** : Reset altered settings
+*   **Esc** : Exit
+
+[REGION:tip]
+The tool is rough around the edges, and usability can be improved in places. The tool is functional, but we plan on improving the usability of it in later releases. We welcome your feedback!
+[/REGION]
+
+## Calibration Process
+
+### 1. Device/Camera/Tracker Selection
+
+If you have multiple video capture devices connected to your PC, you have to specify which one you want to use. Use the **Up/Down** keys to cycle through them.
+
+![](MR_DeviceSelection.png "MR_DeviceSelection.png")
+
+Additionally, each camera may have multiple formats and resolutions available to it. We try to auto-select the best format and resolution, but you can use **Left/Right** to change the selection.
+
+[REGION:tip]
+The resolution of the camera feed does not affect the output resolution of the Mixed Reality capture. The output resolution of the capture is controlled by the game project you will be capturing from. The resolution of the camera feed only controls how clear the feed is within the scene.
+[/REGION]
+
+#### Device Selection Controls
+
+*   **Up/Down** : Select video capture source
+*   **Left/Right** : Select video capture format
+*   **Tab** : Select camera tracker
+
+#### Tracker Selection
+
+If you’re planning on moving your camera while filming, you will need to attach it to a tracking device.
+
+Before continuing to the next step, you can use the **Tab** key to cycle through the available tracking attachments. For the HTC Vive, the first tracker will be named “Special_1” in the attachments list.
+
+After you have made the appropriate selections, press **Enter** and continue to the next step.
+
+### 2. Lens Calibration
+
+Different camera lenses bend and distort the picture in different ways (think of a fisheye lens). You don’t want that distortion to be carried over into the captured scene, so use this lens calibration step to figure out how to undistort the image. ![](MR_LensCalibration.png "MR_LensCalibration.png")
+
+During this step, you will need a printout of the checkerboard image that you downloaded along with the calibration tool. In this step, the calibration tool continuously takes screenshots using the checkerboard pattern as a baseline.
+
+Hold the printout at different locations in front of the camera. Use different angles and depths, especially around the edges of your frame, so that the calibration tool gets a wide collection of samples. The text at the top of the screen will change to let you know when the calibration tool has collected enough samples.
+
+[REGION:tip]
+If your camera is set to auto-focus, you might find it helpful to disable auto-focus during this step.
+[/REGION]
+
+#### Reprojection Error
+
+After the calibration tool has collected enough samples, the tool will show you a preview of the undistorted feed. If the feed looks acceptable, press **Enter** to continue. If needed, you can add more samples, or press **R** to reset and start over.
+
+[REGION:tip]
+A good way to tell if the undistortion process is working is to find elements in frame that should be straight (wall corners, doorways, etc.). Lens distortion will often warp straight edges (especially towards the edge of the frame).
+[/REGION]
+
+![](MR_ReprojectionError.png "MR_ReprojectionError.png")
+
+The reported “reprojection error” lets you know how accurate the process thinks it is. The lower the reprojection error value, the better. You’re in great shape if the value is under one!
+
+#### Angle of View
+
+Probably the most important use of this step is to approximate the camera’s angle-of-view (or FOV, ‘field-of-view’).
+
+The FOV determines how much of the world the camera can see, and it is important that the FOV of our virtual camera matches this as closely as possible.
+
+If you already know the FOV of your camera (in degrees), you can use the *mrc.FovOverride* console variable to set it.  
+
+[REGION:tip]
+Most USB camera manufacturers list the diagonal-FOV (DFOV) for their devices. However, we are interested in the horizontal-FOV (in degrees). You can compute the horizontal-FOV from the diagonal-FOV with the below formula, using the height (h) & width (w) of the resolution you selected in the previous step.
+[/REGION]
+
+![](MR_FOVFormula.png "MR_FOVFormula.png")
+
+There are also ways to calculate the horizontal-FOV for DSLR lenses. 
+
+If you do not know your camera’s specific FOV offhand, don’t worry, this step will approximate it for you.  
+
+[REGION:tip]
+If you have a camera with an adjustable lens, you will want to be very careful not to adjust the zoom after you’ve completed calibration. Adjusting the zoom will change the physical camera’s field-of-view (FOV), but not the virtual one. The virtual camera will use the FOV that was used during calibration. If you adjust the camera's zoom after calibration, you will need to repeat the calibration process.
+[/REGION]
+
+### 3. Alignment Calibration
+
+![](AdjustingAlignment2.png "AdjustingAlignment2.png")
+
+This step requires you to rerun the Align and Adjust phases until complete.
+
+1.  Align the controller with an on-screen model.
+2.  Adjust the sampled alignment.
+3.  Repeat.
+
+[REGION:tip]
+It is important to cover your headset’s sensor, so that we can get accurate tracking data from your controllers.
+[/REGION]
+
+For the HTC Vive, make sure that the controllers are on before you launch the tool. Otherwise, the calibration tool may not display models to align with (the tool needs to recognize the type of controllers that you are using).  
+
+#### Phase 1 - Align w/ Model
+
+On screen, there will be a pink/fuschia model to align your controller with. Hold the controller up to the screen and pull the trigger when the controller and the model are lined up.
+
+Once you’ve pulled the trigger, the screen will freeze frame, and switch into the next phase.
+
+#### Phase 2 - Adjust Sample
+
+Don’t expect to be wholly accurate when you first pull the trigger (shaky hands, etc.). In this phase, you adjust the last sample so it is more exactly aligned in the frame.
+
+#### Phase 3 - Repeat
+
+In total, there are 11 points for you to align with. You don’t have to go through all of them, if at any point you think the alignment is satisfactory.
+
+[REGION:asyncgif]
+![image alt text](AdjustingAlignment2.gif)(w:640)
+[/REGION]
+
+#### Adjustment Controls
+
+*   **Up/Down** : Move the model vertically
+*   **Left/Right** : Move the model horizontally
+*   **Num +/-** : Zoom the model in/out
+*   **Alt + Up/Down** : Rotate the model vertically (pitch)
+*   **Alt + Left/Right** : Rotate the model horizontally (yaw)
+*   **Alt + Num +/-** : Rotate the model left & right (roll)
+*   **H** : Hide the model & its outline
+*   **Alt + H** : Hide the model only (leaving its outline)
+*   **P** : Preview alignment
+*   **Enter** : Accept alignment
+
+#### Preview the Result
+
+As you take samples, the screen will show a series of color coded icons (a cross and target). How close each pair is gives you an idea how close the alignment is. You want all pairs to be as close as possible, though don’t worry if one is wildly off (it may have been an inaccurate sample).
+
+You can preview how the alignment is coming together by holding **P** (or the thumbstick). If you’re happy with the results and wish to skip the rest of the alignment points, press the **End** key.
+
+![](MR_PreviewResults.png "MR_PreviewResults.png")
+[REGION:tip]
+**Video Lag:**There will likely be some lag between the video feed and controller tracking. It can make it hard to tell if the two are properly aligned. To combat this, you can use the console variable, *mrc.TrackingLatencyOverride* to introduce delay so that they’re more in sync.
+[/REGION]
+
+### 4. Alignment Adjustment
+
+Up to this point, you’ve only adjusted the alignment one sample at a time. This single sample gives the calibration tool a good approximation of the camera’s position, but the calibration is likely better in some corners of the frame than others. 
+
+[REGION:asyncgif]
+![Click to play animation](MR_AdjustSample.gif)(w:640)
+[/REGION]
+
+In this step, there are five white boxes on-screen. Move one of your controllers to each of the boxes and pull the trigger in each.
+
+[REGION:tip]
+You want to perform this process at the same depth (distance away from the camera) that you plan on filming from.
+[/REGION]
+
+Each box will disappear when it’s taken a sample for that region. The model tracking with your hands needs to be fully inside the box for it to take a sample.
+
+#### Making Adjustments
+
+Once you’ve hit all 5 boxes, the screen will change to a split-screen like view.
+
+![](MR_MakingAdjustments.png "MR_MakingAdjustments.png")
+
+Like in the previous alignment process, move, rotate, and scale the models to match the image.
+
+#### Adjustment Controls
+
+*   **Up/Down** : Move the models vertically
+*   **Left/Right** : Move the models horizontally
+*   **Num +/-** : Zoom the models in/out
+*   **Alt + Up/Down** : Rotate the models vertically (pitch)
+*   **Alt + Left/Right** : Rotate the models horizontally (yaw)
+*   **Alt + Num +/-** : Rotate the models left & right (roll)
+*   **``</>``** : Adjust FOV up & down (careful)
+*   **R** : Reset adjustments
+*   **P** : Preview alignment
+*   **Enter** : Accept adjusted alignment
+
+You can see how adjusting/optimizing for one sample can throw off the others.
+
+[REGION:tip]
+Rotating in this step is a little hard to wrap your head around, as you are moving all the models all at once in unison. Think of them together as a single model that you are rotating about their collective center. Try to keep the middle sample screen center, and align it first, using it as an anchor point.
+[/REGION]
+
+### 5. Compositing Calibration
+
+In this step, you tweak specific compositing settings. This is the first step where you get to see the composited scene as you would in a Mixed Reality capture. You could exit now, and everything would work properly.
+
+![](MR_CompositingCalibratrion.png "MR_CompositingCalibratrion.png")
+
+Use the arrow keys to select and set specific values.
+
+#### Chroma Key Settings
+
+Most of the settings in this step pertain to the chroma keying process which is discussed in detail in [this UE blog post](https://www.unrealengine.com/en-US/blog/setting-up-a-chroma-key-material-in-ue4).
+
+  
+
+| Setting | Description |
+| --- | --- |
+| ChromaColor | The color of your chroma backdrop (generally green). |
+| Luminance Power | Used to help separate the backdrop color from shadows that may tint the visible color. |
+| Chroma Clip Threshold | Colors matching the chroma color up to this tolerance level will be completely cut out. The higher the value the more that is cut out. A value of zero means that the chroma color has to be an exact match for the pixel to be completely transparent.|
+| Chroma Opacity Strength | Scales the opacity of the remaining pixels. The higher the number, the less translucency there will be. The closer a color is to the ChromaColor, the higher you have to turn this up to keep it from going transparent. |
+| Despill Strength | Scales the despill intensity. A value of zero means no despill correction will be performed. |
+| Despill Cutoff Cap | Colors that differ from the chroma color beyond this tolerance level will be unaffected by despill correction. The higher the value, the more that will be color corrected. |
+| Despill Sharpness | Used to smooth out the despill gradient - defines the despill falloff curve. Scales the despill strength exponentially. The lower the value, the more subtle/smooth it will be. |
+| Faux Bounce Intensity | Used to scale the strength of the color used to replace the color removed as part of the despill process. A value of zero means that the faux bounce color won’t be applied. |
+| Faux Bounce Color | A color to replace the chroma color bleed that was removed as part of the despill process with. |
+
+#### Other Compositing Settings
+
+
+| Setting | Description |
+| --- | --- |
+| Tracking Latency | Depending on your video capture device, the video feed may be laggy compared to your controllers. With this setting, you can introduce a delay to the controllers, in order to better sync the two. |
+| Depth Offset | By default, the video aligns to the depth of the headset. This is used to determine what renders in front of and behind the subject. This lets you offset that forward or backward in the scene. | 
+
+Pressing **Enter** cycles from one setting to the next. Once you’re satisfied with the configuration, hit **End** to save. You can always start the tool again to tweak these values later.
+
+### 6\. Garbage Matting
+
+Garbage matting is the process of blocking out areas that should always be cut out of the video feed.
+
+If your green screen doesn’t cover your entire frame, you will need to complete this additional process to cut out those excess regions from the capture.
+
+[REGION:tip]
+This is the last step of the calibration process. If you don’t have this problem or a need for it, you can skip this step and exit.
+[/REGION]
+
+![](MR_GarbageMatting.png "MR_GarbageMatting.png")
+
+This is the only step in the tool that requires you to work in VR. Put your headset on, and the controls should be listed in front of you.
+
+#### Garbage Matting Controls
+
+*   **End/Esc** : Save & exit the tool
+
+| Left Controller | Right Controller |
+| --- | --- |
+| [INCLUDE:#excerpt_14] | [INCLUDE:#excerpt_15] |
+
+<!-->
+[EXCERPT:excerpt_15]
+
+*   **Grip** : Create matte model
+*   **Trigger** : Select/Deselect matte model
+*   **Thumbstick** : Next gizmo mode (rotate, translate, scale, etc.)
+
+w/ LGrip held
+
+*   **Trigger** : Deselect all
+*   **Thumbstick** : Prev. gizmo mode
+
+[/EXCERPT:excerpt_15]
+-->
+<!--
+[EXCERPT:excerpt_14]
+
+*   **Grip** : Alt (hold to change controls)
+*   **Trigger (Hold & Drag)** : Move gizmo
+*   **Thumbstick** : Undo
+
+w/ LGrip held
+
+*   **Trigger** : Reset gizmo position
+*   **Thumbstick** : Redo
+
+[/EXCERPT:excerpt_14]
+
+-->
+
+[REGION:asyncgif]
+![Click to play animation](MR_GarbageMattingControls.gif)(w:640)
+[/REGION]
+
+Similar to how our [VR world editor](https://docs.unrealengine.com/en-us/Engine/Editor/VR) works, you place and position 3D planes. Match the planes up with the real-world sections you want cut out of your capture.
+
+There’s a preview window in VR, so you can see the results as you modify the mattes. You may want to to move the camera around at this point (if it’s attached to a tracker), to see how it looks from different angles.  
+
+[REGION:tip]
+Getting the mattes aligned perfectly can be a difficult process. If you are having trouble, try standing where you want the mat to be, and bring it to yourself. Use the picture-in-picture to help guide your progress.
+[/REGION]
+
+Once you’ve matted out the areas you want, you can exit the app (**End** or **Esc**) \- all of your settings will be saved. If necessary, you can start the tool back up again to tweak the settings.
+
+## End Result
+
+Once you’ve finished the calibration process, you can exit the tool. In the tool’s *Saved/SaveGames/* directory you should find a *MrcCalibration.sav* file.
+
+[REGION:asyncgif]
+![Click to play animation](MR_CompletingCalibration.gif)(w:640)
+[/REGION]
+
+Once you have a *MrcCalibration.sav* file, copy it into your game’s */Saved/SaveGames/* folder. If your game doesn’t have a SaveGames directory yet, you will need to create it manually.  
+
+If you’ve gone through this process once, you don’t need to do it again (as long as your setup doesn’t change). You can reuse the same calibration settings file across different UE titles (titles that have enabled the MRC framework plugin).
\ No newline at end of file
diff --git a/Engine/Documentation/Source/Platforms/MR/MROverview/MROverview.INT.udn b/Engine/Documentation/Source/Platforms/MR/MROverview/MROverview.INT.udn
new file mode 100644
--- /dev/null
+++ b/Engine/Documentation/Source/Platforms/MR/MROverview/MROverview.INT.udn
@@ -0,0 +1,25 @@
+Availability: Public
+Crumbs:
+Title: Mixed Reality Capture Overview
+Description: Product documentation including reference and guides for Unreal Engine 4 
+Version: 4.20
+Parent: Platforms/MR
+Order:1
+Tags: GettingStarted
+Tags: MR
+Tags: landingpage
+Tags: Early Access
+topic-image:MR_RoboRecall.png
+reference-image:MR_RoboRecall.gif
+
+The Mixed Reality Capture (MRC) framework makes it easy for users to composite real players into a virtual play space so that observing XR experiences is more compelling.
+
+[REGION:asyncgif]
+![image alt text](MR_RoboRecall.gif)(w:653)
+[/REGION]
+
+Mixed Reality Capture support has three components: video input, calibration, and in-game compositing. We have a list of supported [webcams and HDMI capture devices](Platforms/MR/MRSupportedVideoDevices) that enable you to pull real world green screened video into the Unreal Engine from a variety of sources. If you have a separate tracking device, Mixed Reality Capture can match your camera location to the in-game camera to make shots more dynamic and interesting. Setup and calibration is done through a standalone calibration tool. Once you set up and calibrate your filming location, you can reuse those calibration settings across Unreal applications.
+
+While feature support is in early access, we’re looking forward to getting feedback as we continue to improve the framework. 
+
+Get started using Mixed Reality Capture by following the [Mixed Reality Quick Start](Platforms/MR/MRQuickStart).
\ No newline at end of file
diff --git a/Engine/Documentation/Source/Platforms/MR/MRQuickStart/MRQuickStart.INT.udn b/Engine/Documentation/Source/Platforms/MR/MRQuickStart/MRQuickStart.INT.udn
new file mode 100644
--- /dev/null
+++ b/Engine/Documentation/Source/Platforms/MR/MRQuickStart/MRQuickStart.INT.udn
@@ -0,0 +1,91 @@
+Availability: Public
+Crumbs: %ROOT%, Platforms, Platforms/MR
+Title: Mixed Reality Quick Start
+Description: Mixed Reality Quick Start  
+Version: 4.20
+Parent: Platforms/MR 
+Order: 2
+Tags: GettingStarted
+Tags: MR
+Tags: landingpage
+Tags: Early Access
+reference-image:AdjustingAlignmentStep.gif
+reference-image:MR_ApplyingCalibData.gif
+reference-image:Platforms/MR/MROverview/MR_RoboRecall.gif
+
+
+
+![](MR_RoboRecall.png "MR_RoboRecall.png")
+
+At the end of this tutorial, you will have setup and calibrated your system for Mixed Reality Capture (MRC).
+
+## 1. Setting up Green Screen and Camera
+
+Before you can take a capture you will need the right equipment for the job. Here is a quick overview of the equipment you will need, and some tips on how to set it up.  
+
+![](MR_GreenScreenSetup.png "MR_GreenScreenSetup.png")
+
+1.  **Video Camera**  
+    UE4 only supports a very specific set of video capture devices; see [Supported Video Devices](Platforms/MR/MRSupportedVideoDevices) for the current list of the supported devices. Using a device from the list, connect it to your PC for streaming.
+2.  **Chroma backdrop**  
+    For chroma keying, we generally use a green screen. When setting up the green screen you will want to make sure it’s taut, minimizing any wrinkles, especially behind your subject. If you are setting up lighting, you will want to make sure not to cast shadows directly behind the subject. You want a smooth, flat color. The more shades of green that there are across the backdrop, the harder it is to chroma key. It helps if the subject is far away from the backdrop as possible. If you plan on filming the subject’s feet, consider a green screen for the floor as well.
+3.  **Camera Mount**  
+    For initial setup (calibration), your camera needs to be static. If you are using a webcam this may be as simple as attaching it to your desk/monitor. Another option is mount the camera to a tripod.
+4.  **Multi-Mount + Tracker (Optional)**  
+    If you plan on moving the camera around while filming, you will want attach a tracking device to the camera, such as an HTC Vive Tracker. In addition, connect the camera and tracker together with a [multi mount](https://www.bhphotovideo.com/c/product/1062513-REG/desmond_d3d_1_stereo_camera_bracket.html), so that they stay firmly locked in place.
+
+## 2. Calibrating Physical and Virtual Cameras
+
+Once you have your physical space and equipment set up, you are ready to run through calibration. At runtime, the game needs to know where your camera is in relation to your virtual scene. The configuration is different for each setup, which is why you need to run through a calibration process. To simplify calibration, we’ve created a standalone calibration tool (MRCalibration.exe), which can be downloaded [here](http://epic.gm/mrccal). 
+
+[REGION:tip]
+The calibration process is complicated, so we've created a separate [How To Use The Mixed Reality Capture Calibration Tool](Platforms/MR/MRHowToCaptureCalibrationTool) topic that contains detailed instructions on using the calibration tool.
+[/REGION]
+
+[REGION:asyncgif]
+![image alt text](AdjustingAlignmentStep.gif)(w:653)
+[/REGION]
+
+When the calibration process is complete, the calibration tool will generate a calibration settings file (*MrcCalibration.sav*). The calibration settings file (*MrcCalibration.sav*) file can be found in the calibration tool’s */Saved/SaveGames/* folder. This settings file will be used by the game when starting a Mixed Reality capture.  As long as your physical setup stays the same, you don’t need to rerun the calibration process. You can use the same calibration settings file across multiple games.  
+
+[REGION:tip]
+If you have a camera with an adjustable lens, be careful not to adjust the zoom after you’ve completed calibration. Adjusting the zoom will change the physical camera’s field-of-view (FOV), but not the virtual one. The virtual camera will use the FOV that was used during calibration. If you adjust the camera's zoom after calibration, you will need to repeat the calibration process.
+[/REGION]
+
+## 3. Adding the Calibration File to Your Game Project
+
+In the last step, we covered how to use the MRC Calibration Tool to produce a *MrcCalibration.sav* file. Once you have a *MrcCalibration.sav* file, copy it into your game’s */Saved/SaveGames/* folder.If your game doesn’t have a SaveGames directory yet, you will need to create it manually.  
+
+If you have the Mixed Reality Capture Framework plugin enabled in your game, it checks for a calibration settings file on startup.  
+
+## 4. Enabling the Mixed Reality Capture Framework Plugin
+
+Integrating the Mixed Reality Capture (MRC) framework into your project is as easy as enabling the Mixed Reality Capture framework plugin. With the MRC plugin enabled and the calibration settings file in place, the game should start in MRC mode and display the composited scene to the spectator window.  
+By default, MRC broadcasting uses the [VR Spectator Screen](Platforms/VR/VRSpectatorScreen) to display the composited scene. If the game isn’t in VR mode, then the Mixed Reality capture will not be displayed. The MRC plugin is capturing, it's just not displaying.
+[REGION:tip]
+Mixed Reality captures are supported both in the editor (VR PIE) and in packaged projects, as long as the MRC plugin is enabled.
+[/REGION]
+
+[REGION:asyncgif]
+![image alt text](MR_ApplyingCalibData.gif)(w:653)
+[/REGION]
+
+#### Starting in MRC Mode
+
+When the MRC plugin starts up, it checks for a calibration settings file (*MrcCalibration.sav*) in the game's */Saved/SaveGames/* folder. If the MRC plugin does not find a calibration settings file, it will not start in MRC mode. Below is a diagram showing how the game determines whether to start in MRC mode or not.
+
+![](image3.png)(w:800 convert:false)
+
+## 5. Recording the Mixed Reality Capture
+
+The MRC framework does not record a video of the capture. The MRC framework just handles compositing the scene and outputting that to the display window. You must use 3rd party software, like [OBS](https://obsproject.com/) or [Nvidia ShadowPlay](https://www.nvidia.com/en-us/geforce/geforce-experience/shadowplay/) to record the capture.
+
+[REGION:asyncgif]
+![image alt text](Platforms/MR/MROverview/MR_RoboRecall.gif)(w:653)
+[/REGION]
+
+## 6. On your own!
+
+The hardest part of creating a Mixed Reality Capture is properly calibrating. Once you’ve done that
+
+(if your setup doesn’t change), you can keep reusing the calibration file. It’s as simple as copying a file and running!
\ No newline at end of file
diff --git a/Engine/Documentation/Source/Platforms/MR/MRSupportedVideoDevices/MRSupportedVideoDevices.INT.udn b/Engine/Documentation/Source/Platforms/MR/MRSupportedVideoDevices/MRSupportedVideoDevices.INT.udn
new file mode 100644
--- /dev/null
+++ b/Engine/Documentation/Source/Platforms/MR/MRSupportedVideoDevices/MRSupportedVideoDevices.INT.udn
@@ -0,0 +1,27 @@
+Availability: Public
+Crumbs: %ROOT%, Platforms, Platforms/MR
+Title: Supported Video Devices
+Description: Product documentation including reference and guides for Unreal Engine 4
+Version: 4.20
+Parent: Platforms/MR
+Order: 4
+Tags: Reference
+Tags: MR
+Tags: landingpage
+Tags: Early Access
+topic-image:MR_Camera.jpg
+
+We have tested a limited number of video capture devices with the Mixed Reality Capture Framework (MRC). This is not to say that other video capture devices won’t work, but below is a list of the devices that we’ve tested here at Epic Games.
+
+**Webcams**
+
+* [Logitech C920](https://www.logitech.com/en-us/product/hd-pro-webcam-c920)
+* [Logitech C930e](https://www.logitech.com/en-us/product/c930e-webcam)
+
+**USB Capture Dongles**
+
+* [Magewell USB Capture HDMI Gen 2](http://www.magewell.com/usb-capture-hdmi)¹
+
+¹ While untested, it is likely that other Magewell USB capture products are just as compatible with MRC, such as their [Plus model](http://www.magewell.com/usb-capture-hdmi-4k-plus) (for 4K), and their [SDI capture dongles](http://www.magewell.com/usb-capture-sdi-4k-plus) (for broadcast hardware).
+
+It is possible for advanced users to write a Media Framework plugin to support their own video capture device.
\ No newline at end of file
diff --git a/Engine/Documentation/Source/Platforms/MR/MRTroubleshooting/MRTroubleshooting.INT.udn b/Engine/Documentation/Source/Platforms/MR/MRTroubleshooting/MRTroubleshooting.INT.udn
new file mode 100644
--- /dev/null
+++ b/Engine/Documentation/Source/Platforms/MR/MRTroubleshooting/MRTroubleshooting.INT.udn
@@ -0,0 +1,57 @@
+Availability: Public
+Crumbs: %ROOT%, Platforms, Platforms/MR
+Title: Troubleshooting Mixed Reality Capture
+Description: Product documentation including reference and guides for Unreal Engine 4
+Version: 4.20
+Parent: Platforms/MR
+Order: 5
+Tags: Reference
+Tags: MR
+Tags: landingpage
+Tags: Early Access
+topic-image:MR_CaptureFlicker.png
+reference-image:MR_CaptureFlicker.gif
+
+
+This page provides troubleshooting information for Mixed Reality Capture (MRC).
+
+## Flickering During Captures
+
+Depending on the resolution of the MR capture target (the default is 1080p), your app may be constrained by its render target pool size. By default, render targets are set to reallocate as needed, which can lead to ping-ponging and flickering as the MR capture targets fight with the stereo render targets. (shown in the video below)
+
+If you see this kind of behavior, you can change the render target resize rule to ‘Grow’ (in your engine.ini file, set *r.SceneRenderTargetResizeMethod=2*). Making this change should resolve the issue.
+
+[REGION:asyncgif]
+![](MR_CaptureFlicker.gif)(w:640)
+[/REGION]
+  
+
+## Capture Not Displaying on Spectator Screen
+
+The MRC framework uses the [Virtual Reality Spectator Screen](Platforms/VR/VRSpectatorScreen) to display the composited scene. This means the Mixed Reality capture will only be displayed when running in Virtual Reality. If your project also uses the spectator screen, then you need to conditionally guard where you use them. There are MRC methods to help with this:
+
+**Get Mixed Reality Capture Texture**  
+Returns the Capture Texture, or a nullptr if there isn't one.
+
+![](MR_BP_GetMixedRealityCaptureTexture.png "MR_BP_GetMixedRealityCaptureTexture.png")
+
+**Is Mixed Reality Capture Broadcasting**  
+Returns true, if the system is sending the capture texture to the spectator screen.  
+
+![](MR_BP_IsMixedRealityCaptureBroadcasting.png "MR_BP_IsMixedRealityCaptureBroadcasting.png")
+
+**Set Mixed Reality Capture Broadcasting**  
+Toggles whether the capture system is sending the capture texture to the spectator screen or not.  
+
+![](MR_BP_SetMixedRealityCaptureBroadcasting.png "MR_BP_SetMixedRealityCaptureBroadcasting.png")
+
+**An Example of Using the MRC Methods**  
+![](MR_MCR_SpecialConsiderations.png "MR_MCR_SpecialConsiderations.png")
+
+## Controller/Trigger Not Responding During Calibration
+
+If your controller input does not appear to respond during calibration, below are several possible causes.
+
+*   **The headset sensor in not covered.** If theheadset’s sensor is not covered, the controllers may not be active.
+*   **The application window needs to have focus.** If theMRCalibration window is not the current active window, the controller actions may not be captured by the calibration program. 
+*   **The Oculus system may not be configured to run from Unknown Sources.** As a part of Early Access, the calibration program is still in active development and has not been reviewed by Oculus.  ![](MR_OculusUnknownSourcesSetting.png "MR_OculusUnknownSourcesSetting.png")
\ No newline at end of file
