Copying //UE4/Dev-Documentation to //UE4/Samples-Main (Source: //UE4/Dev-Documentation @ 4321946)

diff --git a/Engine/Documentation/Source/Support/Builds/ReleaseNotes/4_19/4_19ReleaseNotes.INT.udn b/Engine/Documentation/Source/Support/Builds/ReleaseNotes/4_19/4_19ReleaseNotes.INT.udn
--- a/Engine/Documentation/Source/Support/Builds/ReleaseNotes/4_19/4_19ReleaseNotes.INT.udn
+++ b/Engine/Documentation/Source/Support/Builds/ReleaseNotes/4_19/4_19ReleaseNotes.INT.udn
@@ -8,29 +8,28 @@
 reference-image:image_30.gif
 reference-image:image_43.gif
 
-Unreal Engine 4.19 Release Notes
 
 ![image alt text](419Banner.png)(w:1000 convert:false)
 
-# What’s New
+## What's New
 
-Unreal Engine 4.19 enables you to **step inside the creative process** - the tools become almost transparent so you can spend more time creating. Improvements to rendering, physics, Landscape terrain, and many more systems mean you can **build worlds that run faster than ever before**. The quality of life for the developers using our tools is always top of mind, so we continue to look at areas we can improve to **put the power into the developers’ hands**.
+Unreal Engine 4.19 enables you to **step inside the creative process** - the tools become almost transparent so you can spend more time creating. Improvements to rendering, physics, Landscape terrain, and many more systems mean you can **build worlds that run than ever before**. The quality of life for the developers using our tools is always top of mind, so we continue to look at areas we can improve to **put the power into the developers' hands**.
 
 Whether you are creating games, linear media, architectural visualizations, or product design tools, Unreal Engine 4.19 enables you to know exactly what the finished product will look like every step of the way. The new Live Link plugin seamlessly blends workflows between external content creation tools and Unreal Engine so you can see updates as you make changes to source content. And with the continued improvements to Sequencer, you can be the director with even more control of your scenes in real time.
 
 When it comes to bringing the worlds in your imagination to life, the sky's the limit. Create breathtaking vistas in large, open worlds thanks to Landscape rendering optimizations. With the new Dynamic Resolution feature that adjusts the resolution as needed to achieve desired frame rates, those worlds will run smoother than ever before on PlayStation 4 and Xbox One.
 
-It wouldn’t be a complete release without a mountain of workflow and usability improvements, and Unreal Engine 4.19 does not disappoint in this respect. Working with Material layers and parameters is easier and more intuitive. Features for debugging Blueprints are more powerful with the ability to step over, step into, and step out. You can now save content folders as favorites. Animation tools have been improved with pinnable commands, ability to have multiple viewports, and lots more.
+It wouldn't be a complete release without a mountain of workflow and usability improvements, and Unreal Engine 4.19 does not disappoint in this respect. Working with Material layers and parameters is easier and more intuitive. Features for debugging Blueprints are more powerful with the ability to step over, step into, and step out. You can now save content folders as favorites. Animation tools have been improved with pinnable commands, ability to have multiple viewports, and lots more.
 
-In addition to all of the updates from Epic, this release includes **128**** improvements ****submitted by the incredible community of Unreal Engine developers on GitHub**! Thanks to each of these contributors to Unreal Engine 4.19:
+In addition to all of the updates from Epic, this release includes **128 improvements submitted by the incredible community of Unreal Engine developers on GitHub**! Thanks to each of these contributors to Unreal Engine 4.19:
 
  
 
 [REDACTED_CONTRIBUTORS]
 
-# Major Features 
+## Major Features 
 
-## New: Temporal Upsampling
+### New: Temporal Upsampling
 
 We have added a new upscaling method called Temporal Upsample that performs both a temporal accumulation of the frame at lower resolution and primary spatial upscale, reducing output blur.
 
@@ -64,11 +63,11 @@
 
 ![image alt text](image_1.png)(w:1000 convert:false)
 
-The temporal upscaler that happens in the Temporal Anti-Aliasing (TAA) pass enables consistent geometry sharpness from varying primary screen percentages from `50%` up to `200%`. In effect, even though the screen percentage may be lowered, with TAAU enabled, geometry in the background that would usually become muddy or blend together can now maintain its detail and complexity, like the fence and telephone pole in the example above.
+The temporal upscaler that happens in the Temporal Anti-Aliasing (TAA) pass enables consistent geometry sharpness from varying primary screen percentages from 50\% up to 200\%. In effect, even though the screen percentage may be lowered, with TAAU enabled, geometry in the background that would usually become muddy or blend together can now maintain its detail and complexity, like the fence and telephone pole in the example above.
 
 For additional information, see the [Screen Percentage with Temporal Upsample](https://docs.unrealengine.com/en-us/Engine/Rendering/ScreenPercentage) page.
 
-## New: Dynamic Resolution
+### New: Dynamic Resolution
 
 We now have support for Dynamic Resolution, which **adjusts the resolution as needed to achieve a desired framerate**, for games on PlayStation 4 and Xbox One! This works by using a heuristic to set the primary screen percentage based on the previous frames GPU workload.
 
@@ -76,19 +75,19 @@
 
 Support for additional platforms will be coming in a future release. For additional information, see the [Dynamic Resolution](https://docs.unrealengine.com/en-us/Engine/Rendering/DynamicResolution) page.
 
-## New: Unified Unreal AR Framework
+### New: Unified Unreal AR Framework
 
 With this release of Unreal Engine, the Unified Unreal Augmented Reality Framework provides a rich, unified framework for building Augmented Reality (AR) apps for both Apple and Google handheld platforms. The framework provides a single path of development, allowing developers to build AR apps for both platforms using a single code path. The Unified Unreal AR Framework includes functions supporting Alignment, Light Estimation, Pinning, Session State, Trace Results, and Tracking.
 
 ![image alt text](image_3.png)(w:1000 convert:false)
 
-## New: Unified Unreal AR Framework Project Template
+### New: Unified Unreal AR Framework Project Template
 
 Also new is the Blueprint template HandheldAR, which provides a complete example project demonstrating the new functionality.
 
 ![image alt text](image_4.png)(w:849 h:590 convert:false)
 
-## New: Physical Light Units
+### New: Physical Light Units
 
 New to Unreal Engine 4.19, all light units are now defined using physically based units. Some light units were already well defined, but others were using undefined, engine specific units. The unit selection of a light is done through a drop-down menu (where applicable). For compatibility reasons, the default light units are kept compatible with previous versions of the engine. The new light unit property can be edited per light, changing how the engine interprets the "Intensity" property when doing lighting related computations.
 
@@ -96,7 +95,7 @@
 
 *Project courtesy of Litrix*
 
-## New: Live Link Plugin Improvements
+### New: Live Link Plugin Improvements
 
 The **Maya Live Link Plugin** is now available and can be used to establish a connection between Maya and UE4 enabling you to **preview changes made in Maya in real-time inside the UE4 Editor**. 
 
@@ -124,18 +123,18 @@
 
 The Maya Live Link UI window can be enabled through the MEL console with the command **MayaLiveLinkU**. At the top right is a display that shows whether this instance of Maya is connected to an Unreal client. Below Unreal Engine Live Link is a list of all the subjects currently being streamed (in the image above, only one is being streamed) and in the lower window, controls for adding and removing subjects from streaming is available. 
 
-### Motionbuilder Live Link Plugin
+#### Motionbuilder Live Link Plugin
 
 **
 **The Motionbuilder Plugin offers the same functionality as the Maya Plugin and shows up in the Editor as a connection in a similar way. It also has a custom UI for managing streaming:
 
 Objects can be selected from the the current scene and added to the streamed list (as shown above). From there, their names can be set in the **Subject Name** column and their **Stream Type** (Camera, Skeleton, etc.) can be set. Streaming on the subject can also be enabled and disabled from here. 
 
-### Stream Active Camera as Editor Active Camera
+#### Stream Active Camera as Editor Active Camera
 
-Controlling the active camera within Maya will now manipulate the Editor’s active camera. We’ve also reworked the Editor update hook so that streaming is more robust and facial rigs/leaf bones now correctly update. 
+Controlling the active camera within Maya will now manipulate the Editor's active camera. We've also reworked the Editor update hook so that streaming is more robust and facial rigs/leaf bones now correctly update. 
 
-### Added Virtual Subjects to Live Link
+#### Added Virtual Subjects to Live Link
 
 Virtual subjects are a way to combine multiple subjects coming into Live Link into one subject that can then be used by the Editor. Virtual Subjects are created within the client and contain the bones of multiple real subjects, all tied to a common root. 
 
@@ -147,7 +146,7 @@
 
 Here we have created a new virtual subject (using **Add > Add Virtual Subject**) and set it to read subjects from the Maya Live Link source. The virtual subject is set to consist of the three Item subjects being sent from Maya. 
 
-### Live Link Plugin Development
+#### Live Link Plugin Development
 
 The purpose of Live Link is to provide a common interface for streaming and consuming any kind of animation data from sources outside of UE4 (for example, DDC tools and Motion Capture servers). It is designed to be extensible via Unreal Plugins, allowing third parties to develop new features with no need to make and then maintain Engine changes. 
 
@@ -159,7 +158,7 @@
 
 For an overview of general Plugin development, see the [Plugins](https://docs.unrealengine.com/latest/INT/Programming/Plugins/) documentation pages. 
 
-### Live Link Motion Controller Support
+#### Live Link Motion Controller Support
 
 Live Link can now be used with Motion Controllers. The motion source property of a Motion Controller can be set to a subject within Live Link. When set in this way, the position of the Motion Controller Component is governed by the first transform of the subject.
 
@@ -171,7 +170,7 @@
 
 ![image alt text](image_8.png)(w:925 h:289 convert:false)
 
-### Miscellaneous Improvements
+#### Miscellaneous Improvements
 
 * **Added ability for Live Link Sources to define their own custom settings - **This was requested by partners like IKinema building Live Link support
 
@@ -179,11 +178,11 @@
 
 * **Added virtual Initialization function and Update DeltaTime parameter to Live Link Retargeter API**
 
-## New: Sequencer Improvements
+### New: Sequencer Improvements
 
 We continue to make improvements to the functionality and workflow of Sequencer to make it more powerful and increase efficiency.
 
-### Copy/Paste/Duplicate Object Tracks
+#### Copy/Paste/Duplicate Object Tracks
 
 You can now copy/paste/duplicate object tracks and their child tracks from the right-click context menu. You can copy a spawnable to the same Level Sequence or another Sequence or a possessable from one Level Sequence to another and the object will be bound to the same object in the other Level Sequence. You can now also copy from one UMG animation to another.
 
@@ -205,11 +204,11 @@
 	[/PARAMLITERAL]
 [/OBJECT]
 
-### Level Sequence Dynamic Transform Origin
+#### Level Sequence Dynamic Transform Origin
 
 You can now offset the Actors controlled by Sequence Transform tracks with a global offset at runtime. This allows you to also reuse a Level Sequence in different coordinate spaces. To use this feature, inside the Details panel for your Level Sequence, enable **Override Instance Data** then assign a **Transform Origin Actor. **
 
-The **Transform Origin **section specifies a transform from which all absolute transform sections inside the Sequence should be added to (Scale is ignored). For the best results, keyframe your Actor’s starting transform at 0,0,0 inside your Level Sequence and let the Transform Origin Actor you define drive the position from which to start in the world (depicted below). 
+The **Transform Origin **section specifies a transform from which all absolute transform sections inside the Sequence should be added to (Scale is ignored). For the best results, keyframe your Actor's starting transform at 0,0,0 inside your Level Sequence and let the Transform Origin Actor you define drive the position from which to start in the world (depicted below). 
 
 [OBJECT:EmbeddedVideo]
 	[PARAMLITERAL:videoid]
@@ -231,7 +230,7 @@
 
 Whenever we press a key (see above), the cube moves and our Actor continues to walk along a path, starting from the location of the **Transform Origin Actor**.
 
-### Sequencer Anim BP Weight Control
+#### Sequencer Anim BP Weight Control
 
 Sequencer weight blending now works with Animation Blueprints. You can use the same slot for the animation, and control weight by curve. Assign the same Slot Name under Properties for each animation to blend. Then keyframe the Weight values you desire. Inside your Anim BP, use the Slot Node. The AnimBP will use whatever Weight values you provided when blending.
 
@@ -253,22 +252,22 @@
 	[/PARAMLITERAL]
 [/OBJECT]
 
-### Improved Sequencer Editor Performance
+#### Improved Sequencer Editor Performance
 
 * Compile On The Fly Logic - Sequencer is now able to compile partially or completely out-of-date evaluation templates from the source data as it needs. This affords much more efficient compilation when working within Sequencer in the editor.
 
 * Details Panel Updates - The details panel now defers updates while scrubbing/playing in the Sequencer Editor.
 
-### Minor Improvements
+#### Minor Improvements
 
 * New .ini setting for the default behavior of "When Finished". 
 
-    * For level sequences, it continues to default to RestoreState. For UMG, it’s now set to Keep State.
+    * For level sequences, it continues to default to RestoreState. For UMG, it's now set to Keep State.
 	
 	[/Script/LevelSequence.LevelSequence]
 DefaultCompletionMode=RestoreState
 
-* Always create a camera cut track when adding a camera so that there’s less confusion as to why a camera doesn’t take control when sequencer is activated.
+* Always create a camera cut track when adding a camera so that there's less confusion as to why a camera doesn't take control when sequencer is activated.
 
 * Rotations will no longer flip when autokeying/adding a new key. 
 
@@ -282,7 +281,7 @@
 
 * Can now drag sections up a row, instead of only just down.
 
-## New: Landscape Rendering Optimization
+### New: Landscape Rendering Optimization
 
 The Landscape level of detail (LOD) system has been changed so that instead of being distance-based it now uses screen size to determine detail for a component, similar to how the Static Mesh LOD system works. The LOD distribution now gives more coherent sizes on distant triangles based on their screen size, which can maintain detail that was not possible before.
 
@@ -298,7 +297,7 @@
 
 ![image alt text](image_15.png)(w:441 convert:false)
 
-## New: Proxy LOD System (Experimental)
+### New: Proxy LOD System (Experimental)
 
 The new Proxy LOD system is an experimental Plugin for producing low poly LOD with baked materials for multiple meshes. The new system is used by HLOD and is a replacement for Simplygon. 
 
@@ -315,7 +314,7 @@
 
 
 
-## New: Material Parameters Editing and Saving
+### New: Material Parameters Editing and Saving
 
 You can now **save parameter values to a new Child Instance or Sibling Instance** in the Material Editor and the Material Instance Editor!
 
@@ -325,11 +324,11 @@
 
 * **Save to Child Instance** enables you to save the current parameters and overrides to a new Material Instance with the current instance as the Parent Material.
 
-The Material Editor now features a **Parameter Defaults** panel. Here, you’ll have access to the default values set in your Material Graph for any parameter. You can easily change any of your parameter default values here.
+The Material Editor now features a **Parameter Defaults** panel. Here, you'll have access to the default values set in your Material Graph for any parameter. You can easily change any of your parameter default values here.
 
 ![image alt text](image_19.png)(w:1000 convert:false)
 
-## New: Material Layering (Experimental)
+### New: Material Layering (Experimental)
 
 Material Layering enables you to **combine your Materials in a stack**, using the new Material Layer and Material Layer Blend assets! This enables you to build the correct Material Graph without building sections of nodes by hand. This functionality is similar to Material Functions, but supports the creation of child instances.
 
@@ -355,17 +354,17 @@
 
 For additional information, see our post on the [forums](https://forums.unrealengine.com/development-discussion/rendering/1414193-material-layering-feedback-for-4-19) to get started and leave us feedback!
 
-## New: Upgraded 'Auto Convex' Tool
+### New: Upgraded 'Auto Convex' Tool
 
 You can now specify **Max Hulls **instead of accuracy when generating simplified collision for a Static Mesh using the Auto Convex Tool. Computation is now a background task in the Static Mesh Editor. The Auto Convex Tool also uses a newer version of V-HACD library that should give more accurate results. 
 
 ![image alt text](image_21.png)(w:1000 convert:false)
 
-## New: Improved Facial Animation Sharing using Curve only animations
+### New: Improved Facial Animation Sharing using Curve only animations
 
-You can share facial animations between characters by animating a curve to drive a pose without having any bone transform. Any bone transform, even with one mesh’s ref pose, won’t work for other meshes, so it is important to remove them and start with the mesh’s own ref pose if you want to share the curve between different meshes. This allows you to share the facial curves between different faces. 
+You can share facial animations between characters by animating a curve to drive a pose without having any bone transform. Any bone transform, even with one mesh's ref pose, won't work for other meshes, so it is important to remove them and start with the mesh's own ref pose if you want to share the curve between different meshes. This allows you to share the facial curves between different faces. 
 
-We’ve added the ability to allow users to delete all bone tracks, leaving only animation curves. You can access the **Remove All Bone Tracks **option from the **Asset **menu under **Animation. **
+We've added the ability to allow users to delete all bone tracks, leaving only animation curves. You can access the **Remove All Bone Tracks **option from the **Asset **menu under **Animation. **
 
 [OBJECT:EmbeddedVideo]
 	[PARAMLITERAL:videoid]
@@ -389,11 +388,11 @@
 
 ![image alt text](image_23.png)(w:1000 convert:false)
 
-We’ve also allowed the opting out of importing bone tracks (curve only animation) during the import process.
+We've also allowed the opting out of importing bone tracks (curve only animation) during the import process.
 
-We’ve fixed Pose Preview to select weights that the asset contains. You can see what bone tracks the animation contains inside the Asset Details panel under **Animation Track Names.**
+We've fixed Pose Preview to select weights that the asset contains. You can see what bone tracks the animation contains inside the Asset Details panel under **Animation Track Names.**
 
-## New: Cloth Updates
+### New: Cloth Updates
 
 The new "Anim Drive" feature added to clothing assets and runtime interaction for simulations adds a new mask target for "Anim Drive Multipliers" that is used to drive springs that pull clothing towards its animated position, giving control to animators in cinematics or animation driven scenes.
 
@@ -427,7 +426,7 @@
 
 * **Clothing Create Menu now defaults to use the Skeletal Mesh Physics Asset**
 
-## New: Animation Blueprint Improvements
+### New: Animation Blueprint Improvements
 
 We have made several quality of life improvements to the Animation Blueprints system:
 
@@ -439,7 +438,7 @@
 
 * **Native AnimInstance selection on SkeletalMeshComponent**** - **Allow native Animation Instances to be set on Skeletal Mesh Components. Previously, only Anim Blueprints could be set on Skeletal Mesh Components.
 
-## New: Animation Tools Improvements
+### New: Animation Tools Improvements
 
 **Added Pinnable Command List to animation editors**
 
@@ -563,7 +562,7 @@
 
 * **Added a hotkey and command for toggling Post Processing in Preview Scenes**
 
-* **Added a new shortcut for switching foreground and background colors in the Mesh Painter** - By default the shortcut is bound to the ‘X’ Key when the Mesh Painter is active.
+* **Added a new shortcut for switching foreground and background colors in the Mesh Painter** - By default the shortcut is bound to the ‘X' Key when the Mesh Painter is active.
 
 * **Material Baking has been moved out of the experimental phase**
 
@@ -577,7 +576,7 @@
 
 * **Exposed Damping and Stiffness as pins on Spring AnimNode** - Spring Anim Node now allows exposing the Spring Damping and Spring Stiffness as pins.
 
-## New: Motion Controller Component Visualization
+### New: Motion Controller Component Visualization
 
 Motion Controllers in Unreal Engine 4.19 now have a new Visualization category that enables you to quickly (and more easily) add a display model Static Mesh to your Motion Controllers. By default, the system attempts to load a Static Mesh model matching the device driving the Motion Controller.
 
@@ -605,19 +604,19 @@
 
     * For more complicated, custom models, manually attach your own Static Mesh model.
 
-    * ‘*Display Model Source*’ can be manually set if the target XR system isn’t loaded.
+    * ‘*Display Model Source*' can be manually set if the target XR system isn't loaded.
 
-* By default, when empty, ‘*Display Model Source*’ will attempt to load a model from the backing XR system.
+* By default, when empty, ‘*Display Model Source*' will attempt to load a model from the backing XR system.
 
     * Distinct XR systems support this feature by implementing the IXRSystemAssets interface - currently, only Oculus and SteamVR support this functionality.
 
-## New: Vive Pro Support
+### New: Vive Pro Support
 
 ![image alt text](image_32.jpg)(w:1000 convert:false)
 
-All existing UE4 content that supports SteamVR is compatible with HTC’s newly-announced Vive Pro!  While no modifications are necessary to use the Vive Pro with your existing UE4 project, you will need to check the performance of your UE4 project with the Vive Pro and optimize accordingly as the Vive Pro operates at a higher resolution. 
+All existing UE4 content that supports SteamVR is compatible with HTC's newly-announced Vive Pro!  While no modifications are necessary to use the Vive Pro with your existing UE4 project, you will need to check the performance of your UE4 project with the Vive Pro and optimize accordingly as the Vive Pro operates at a higher resolution. 
 
-## New: Platform SDK Upgrades
+### New: Platform SDK Upgrades
 
 In every release, we update the engine to support the latest SDK releases from platform partners.
 
@@ -667,7 +666,7 @@
 
 * **tvOS:** SDK 11
 
-## New: Low Input Latency Frame Syncing (Xbox One and PlayStation 4)
+### New: Low Input Latency Frame Syncing (Xbox One and PlayStation 4)
 
 Low-latency frame syncing mode modifies the way thread syncing is performed between the game, rendering and RHI threads, and the GPU to greatly reduce and control input latency. The effects of Low-latency frame syncing can be controlled using the r.GTSyncType (game thread sync type) console variable. This console variable drives how the new frame syncing mechanism works and allows the following values:
 
@@ -679,7 +678,7 @@
 
 For additional information, see the [Low Latency Frame Syncing](https://docs.unrealengine.com/en-us/Platforms/LowLatencyFrameSyncing) page.
 
-## New: Encryption/Signing Key Generator
+### New: Encryption/Signing Key Generator
 
 It is now possible to **configure .pak file signing and encryption fully from the Editor** via a single unified settings panel. These features can now be controlled in Project Settings, under the Project section, in the Crypto panel.
 
@@ -687,23 +686,23 @@
 
 From this panel, it is possible to generate or clear keys for either signing or encryption, as well as enable or disable the various types of signing and encryption that are available. The following features can be activated or deactivated from this menu:
 
-* **Encrypt Pak .ini Files** - Encrypts all .ini files that exist in your project’s .pak files. This will prevent easy mining or tampering with configuration data for your product, at minimal runtime cost.
+* **Encrypt Pak .ini Files** - Encrypts all .ini files that exist in your project's .pak files. This will prevent easy mining or tampering with configuration data for your product, at minimal runtime cost.
 
-* **Encrypt Pak Index** - Encrypts the pak file index, which prevents UnrealPak from opening, viewing, and unpacking your product’s pak files, at minimal runtime cost.
+* **Encrypt Pak Index** - Encrypts the pak file index, which prevents UnrealPak from opening, viewing, and unpacking your product's pak files, at minimal runtime cost.
 
 * **Encrypt UAsset Files** - Encrypts the uasset files in the pak file. uasset files contain header information about the assets inside, but not the actual asset data itself. Encrypting this data provides extra security to your data, but does add a small runtime cost and an increase in data entropy which can have a negative impact on patch sizes.
 
 * **Encrypt Assets** - Fully encrypts all Assets within the pak file. Please note that this setting will have a measurable effect on run-time file I/O performance, and increases the amount of entropy in your final packaged data, making the distribution patching system less efficient.
 
-## New: Call Stack Display for Blueprint Debugging
+### New: Call Stack Display for Blueprint Debugging
 
 A control that displays the current Blueprint Call Stack has been added to the Developer Tools menu. Double-clicking entries in the Call Stack will focus the corresponding node in the Blueprint Editor.
 
 ![image alt text](image_35.png)(w:1000 convert:false)
 
-## New: Single stepping improvements to Blueprint Debugging
+### New: Single stepping improvements to Blueprint Debugging
 
-The Engine’s Blueprint Debugger now supports **Step Into (F11)**, **Step Over (F10)**, and **Step Out (Alt+Shift+F11)** operations when stopped at a Blueprint Breakpoint, making debugging sessions quicker and easier.
+The Engine's Blueprint Debugger now supports **Step Into (F11)**, **Step Over (F10)**, and **Step Out (Alt+Shift+F11)** operations when stopped at a Blueprint Breakpoint, making debugging sessions quicker and easier.
 
 ![image alt text](image_36.png)(w:1000 convert:false)
 
@@ -727,15 +726,15 @@
 	[/PARAMLITERAL]
 [/OBJECT]
 
-## New: Improved Tools For Optimizing Disk Size
+### New: Improved Tools For Optimizing Disk Size
 
 Several tools have been added or improved to help analyze disk and download sizes for packaged games. The Asset Audit window has a new "Add Chunks" option which shows all packaged Chunk/Pak files, and the Reference Viewer and Size Map can be used from that view to analyze why Assets are being cooked into a specific chunk.
 
-The Size Map tool has been changed to show Disk Size by default, which is more accurate than the estimated memory size. If a platform is selected in the Asset Audit window, the actual cooked size for that platform will be shown. The Size Map’s user interface has also been improved to support right-click operations and better performance.
+The Size Map tool has been changed to show Disk Size by default, which is more accurate than the estimated memory size. If a platform is selected in the Asset Audit window, the actual cooked size for that platform will be shown. The Size Map's user interface has also been improved to support right-click operations and better performance.
 
 ![image alt text](image_38.png)(w:1000 convert:false)
 
-## New: Audio Improvements
+### New: Audio Improvements
 
 **Upgrade to Steam Audio Beta 10 (New Audio Engine Only)**
 
@@ -765,7 +764,7 @@
 
 Oculus Audio plugin has been updated with support for Android.
 
-In addition to providing HRTF spatialization, we now support Oculus Audio’s room modeling reverb with volumetric parameters as well as Ambisonics Source playback.  
+In addition to providing HRTF spatialization, we now support Oculus Audio's room modeling reverb with volumetric parameters as well as Ambisonics Source playback.  
 
 **Added Resonance Audio Plugin (New Audio Engine Only)**
 
@@ -833,7 +832,7 @@
 
 This Source Effect will allow designers the ability to drive stereo panning values from Blueprints as a Source Effect. This allows you to change the pan of a mono or stereo 3d or 2d source.
 
-## New: Physics Improvements
+### New: Physics Improvements
 
 Several improvements and optimizations have been made to the physics system to make simulating physical objects easier and efficient!
 
@@ -855,11 +854,11 @@
 
 * **Engine now compiles cleanly without PhysX - **Can set WITH_PHYSX define directive to 0 and compile without errors. 
 
-## New: Automatic Code Signing on iOS
+### New: Automatic Code Signing on iOS
 
 Automatic Code Signing (ACS) when using Xcode has now been introduced in UE4.19. ACS plus a Team ID allows a package to be created and signed for an iOS project. This also works when there are no mobile provisions pre-installed on the Mac but the Mac must contain the certificates. A development team will still need to add iOS devices to their Apple account in order to deploy their apps, but they will no longer need to download and install mobile provisions on to a Mac if they use this feature. 
 
-## New: Android 4.19 Improvements
+### New: Android 4.19 Improvements
 
 * Added Proxy network support.
 
@@ -867,7 +866,7 @@
 
 * Added NDK and SDK platform validation before you compile and pack your project to let you know sooner if the build has failed or not.
 
-## New: Folder Favorites for the Content Browser
+### New: Folder Favorites for the Content Browser
 
 You can now **mark any folder in your Content Browser as a "favorite"** to make it easily accessible! When a folder is marked as a favorite, it will populate under the new "Favorites" section located above the folder tree. Simply right-click to add or remove folders from this section.
 
@@ -879,11 +878,11 @@
 
 ## RELEASE NOTES
 
-#### AI
+### AI
 
 * Bug Fix: Fixed the On Perception Updated delegate to be correctly bindable to AI Perception Component.
 
-##### Behavior Tree
+#### Behavior Tree
 
 * Bug Fix: Changed the order of the Behavior Tree tick to allow Blueprint-based nodes to update before processing search data.
 
@@ -895,7 +894,7 @@
 
 * New: Added a simple mechanism for logging Behavior Tree** **nodes that reference empty or removed Blackboard keys.
 
-##### Debugging Tools
+#### Debugging Tools
 
 * New: Extended the EQS Testing Pawn with an option to specify navigation agent properties to support navigation-related EQS tests and features.
 
@@ -903,7 +902,7 @@
 
 * Bug Fix: Fixed some existing Visual Logging Blueprint functions.
 
-##### Navigation
+#### Navigation
 
 * Bug Fix: Fixed cases where changes to Nav Mesh Bounds Volumes properties were not causing Nav Mesh to rebuild.
 
@@ -911,7 +910,7 @@
 
 * New: Made updates to allow the proper use of Nav Link Component via Blueprints.
 
-#### Animation
+### Animation
 
 * Crash Fix: Fixed a crash when opening Animation Compression settings for newly imported Alembic File.
 
@@ -999,7 +998,7 @@
 
 * New: Skeletal Mesh Render Data Get Max Bones Per Section is now an ENGINE_API exported function.
 
-##### Animation Assets
+#### Animation Assets
 
 * Crash Fix: Fixed a crash with streaming if a montage is streamed out while an event is queued.
 
@@ -1039,7 +1038,7 @@
 
 * New: Added inline time/frame settings to notify context menu.
 
-##### Animation Blueprints
+#### Animation Blueprints
 
 * Crash Fix: Loading an Anim Instance with a re-instanced class no longer crashes.
 
@@ -1073,7 +1072,7 @@
 
 * Deprecated: Get Scalar Parameter Default has been deprecated and replaced with Get Scalar Parameter Default Value.
 
-##### Animation Tools
+#### Animation Tools
 
 * Crash Fix: Removing bones in a LOD on a Skeleton that has Virtual Bones no longer crashes.
 
@@ -1119,11 +1118,11 @@
 
 * New: Moved Live Link Message Bus Source heartbeats to their own thread using a new  Heartbeat Manager singleton. This prevents sources from incorrectly being removed during Slate UI operations.
 
-##### Import/Export
+#### Import/Export
 
 * New: Refactored Facial Animation importing so it can be used from inside other modules.
 
-##### Skeletal Mesh
+#### Skeletal Mesh
 
 * Crash Fix: Undoing a change on a Skeletal Mesh that uses Virtual Bones no longer crashes.
 
@@ -1145,7 +1144,7 @@
 
 * New: You can set attached and parent to sync their LOD by setting Sync Attach Parent LOD to true.
 
-##### Tools
+#### Tools
 
 * Crash Fix: Skeletal Mesh Editor no longer crashes when importing more than 5 LODs.
 
@@ -1185,7 +1184,7 @@
 
 * New: Added a search bar to the Advanced Preview Scene settings panel.
 
-#### Audio
+### Audio
 
 * Crash Fix: Fixed an issue with VOIP crashing on Mac.
 
@@ -1273,7 +1272,7 @@
 
 * New: Added "Owner" parameter to Play Sound At Location and Play Sound 2D functions to support "Limit to Owner" concurrency settings.
 
-#### Automation
+### Automation
 
 * Crash Fix: Fixed a crash in Screen Shot Manager caused by an attempt to access a deleted String in an async lambda expression.
 
@@ -1291,7 +1290,7 @@
 
 * New: Automated screenshot tests that fail due to size mismatch will now show a delta image, where mismatched bounds are colored red.
 
-#### Blueprints
+### Blueprints
 
 * Crash Fix: Fixed inclusion of private plugin header files in nativized Blueprint C++ source.
 
@@ -1327,13 +1326,13 @@
 
 * Bug Fix: Removed the Inverse Lerp node and redirected it to call Normalize To Range. Also Fixed incorrect value being returned by Normalize To Range when Range Min and Range Max are equal.
 
-* Bug Fix: Fixed a rare linker issue where Is Blueprint Finalization Pending would return true when trying to force-load subobjects that weren’t ready to be loaded. This would prevent some placeholder objects from being replaced.
+* Bug Fix: Fixed a rare linker issue where Is Blueprint Finalization Pending would return true when trying to force-load subobjects that weren't ready to be loaded. This would prevent some placeholder objects from being replaced.
 
 * Bug Fix: We now ensure that "-build" is always passed to BuildCookRun automation workflows in the editor for projects configured with Blueprint nativization enabled so that it doesn't skip that stage.
 
 * Bug Fix: Revised both UAT/UBT-driven and custom code/data workflows to more seamlessly work with Blueprint nativization settings.
 
-* Bug Fix: Ed Graph Schema’s Break Single Pin Link function is now correctly marked as a const function.
+* Bug Fix: Ed Graph Schema's Break Single Pin Link function is now correctly marked as a const function.
 
 * Bug Fix: Local variables can now be created in blueprint function libraries.
 
@@ -1435,7 +1434,7 @@
 
 * New: Added Vector Length XY to Blueprint math library.
 
-#### Core
+### Core
 
 * Crash Fix: Fixed a crash on shutdown due to component visualizers not being cleaned up properly by the Unregister Component Visualizer function.
 
@@ -1463,11 +1462,11 @@
 
 * Crash Fix: Running DDC commandlet with an invalid package name no longer causes crashes.
 
-* Bug Fix: Multicast Delegate Property’s Identical function now has the same behavior as Delegate Property’s Identical function. This prevents it from giving incorrect results and causing serialization errors.
+* Bug Fix: Multicast Delegate Property's Identical function now has the same behavior as Delegate Property's Identical function. This prevents it from giving incorrect results and causing serialization errors.
 
 * Bug Fix: The Asset Registry module is now always preloaded on startup, even if EDL is disabled.
 
-* Bug Fix: Fixed Thread Sanitizer warnings in stats system. Optimized Cycle Counter’s Start function to only read the stat name once (atomically).
+* Bug Fix: Fixed Thread Sanitizer warnings in stats system. Optimized Cycle Counter's Start function to only read the stat name once (atomically).
 
 * Bug Fix: Fix ambiguities in calculating root directory.
 
@@ -1509,9 +1508,9 @@
 
 * Bug Fix: Deleted an unnecessary specialization of Make Array View.
 
-* Bug Fix: Removed String's char pointer constructor’s confusing void pointer parameter.
+* Bug Fix: Removed String's char pointer constructor's confusing void pointer parameter.
 
-* Bug Fix: Base64’s Encode function is now const-correct.
+* Bug Fix: Base64's Encode function is now const-correct.
 
 * Bug Fix: Json Value no longer stringifies whole numbers as floats.
 
@@ -1547,21 +1546,21 @@
 
 * Bug Fix: If UAT is invoked with platforms in both the "-platform" and “-targetplatform” command-line switches, we will now build using all of them rather than just the ones in “-targetplatform”.
 
-* Bug Fix: Fixed cases where String’s Printf function was used without a formatting string.
+* Bug Fix: Fixed cases where String's Printf function was used without a formatting string.
 
 * Bug Fix: Updated loading movie code to properly calculate delta time when throttled at 60 frames per second.
 
 * Bug Fix: Fixed an inconsistency with how Soft Object Ptr case is managed between Linker Save and Archive Save Tag Imports, which could cause a cook ensure under some circumstances.
 
-* Bug Fix: Fixed an occasional loading assert in Script Struct’s Post Load function caused by a race condition on the Prepare CPP Struct Ops Completed variable.
+* Bug Fix: Fixed an occasional loading assert in Script Struct's Post Load function caused by a race condition on the Prepare CPP Struct Ops Completed variable.
 
 * Bug Fix: Changed Visual Studio Code project file generation to use "taskName" rather than "label" for C# build task names.
 
-* Bug Fix: Improved whitespace handling in both Map Property’s and Set Property’s "ImportText_Internal" function. Emptied the map or set if an error occurs during a text import of a property to avoid bad container states.
+* Bug Fix: Improved whitespace handling in both Map Property's and Set Property's "ImportText_Internal" function. Emptied the map or set if an error occurs during a text import of a property to avoid bad container states.
 
 * Bug Fix: Fixed a potential buffer overrun in Set Property.
 
-* Bug Fix: Map Property’s Import Text function now correctly skips duplicate keys.
+* Bug Fix: Map Property's Import Text function now correctly skips duplicate keys.
 
 * Bug Fix: Unreal Build Tool now enforces Development and Debug configurations in existing C# projects. Fixes command-line builds and prevents solution from being marked dirty on initial load.
 
@@ -1575,7 +1574,7 @@
 
 * Bug Fix: Dedicated servers do not include RHI modules except NullRHI.
 
-* Bug Fix: Changed String’s Parse Into Array function to use Reset instead of Empty on the passed in array, allowing it to be approximately resized
+* Bug Fix: Changed String's Parse Into Array function to use Reset instead of Empty on the passed in array, allowing it to be approximately resized
 
 * Bug Fix: Fixed accounting for Curve Table assets in "obj list".
 
@@ -1589,7 +1588,7 @@
 
 * Bug Fix: Asset registry memory usage reported by memreport is now correct.
 
-* Bug Fix: Fixed Platform File Pak interface’s Find Files Recursively function not correctly calling the base class implementation.
+* Bug Fix: Fixed Platform File Pak interface's Find Files Recursively function not correctly calling the base class implementation.
 
 * Bug Fix: Literal soft object pins are correctly cooked and tracked in the reference viewer. This may cause assets referenced by disabled nodes to be cooked into your project. In this case, the reference viewer will track which blueprints are causing this.
 
@@ -1597,7 +1596,8 @@
 
 * Bug Fix: Hot reloads no longer occur during PIE, which was causing bad states in blueprint types when they are reinstanced.
 
-* New: Encryption/Signing Key Generator.A plugin which allows the automatic creation and configuration of in-engine data cryptography through the project settings, under the "crypto" header.
+* New: Encryption/Signing Key Generator.
+A plugin which allows the automatic creation and configuration of in-engine data cryptography through the project settings, under the "crypto" header.
 
     * Added options to plugin settings which allow encryption to be used on the various asset files in a pak. Encryption can be used on just the .uasset file, which contains only the package header information, or the .uexp and .ubulk files too, which contain the rest of the package data. Encryption on .uasset files provides a good compromise between data security and runtime performance.
 
@@ -1621,7 +1621,7 @@
 
 * New: Array now has a "_GetRef" variants for many of its functions, which will create an element and return a reference to it rather than its index in the Array. These functions are: “Insert_GetRef”, “InsertZeroed_GetRef”, “InsertDefaulted_GetRef”, “Emplace_GetRef”, “EmplaceAt_GetRef”, “Add_GetRef”, “AddZeroed_GetRef” and “AddDefaulted_GetRef”.
 
-* New: There is a new Multi Map function, Append, to match the Map class’ function of the same name.
+* New: There is a new Multi Map function, Append, to match the Map class' function of the same name.
 
 * New: Base64 encoding library now supports ANSI/wide characters, decoding streams without padding marks, and decoding into preallocated buffers without writing dummy padding bytes.
 
@@ -1633,7 +1633,7 @@
 
 * New: Visual Studio Code - Add launch task for generating project files for the given folder
 
-* New: Ensure that all branches within Generic Platform Misc’s Root Dir function produce an absolute path with no duplicate slashes. Remove relative-to-absolute conversion of the root directory from the Paths function, Make Standard Filename, now that we know Root Dir always returns an absolute path.
+* New: Ensure that all branches within Generic Platform Misc's Root Dir function produce an absolute path with no duplicate slashes. Remove relative-to-absolute conversion of the root directory from the Paths function, Make Standard Filename, now that we know Root Dir always returns an absolute path.
 
 * New: Added Net Serialize to Date Time.
 
@@ -1713,7 +1713,7 @@
 
 * Removed: Removed deprecated delegate functions.
 
-#### Editor
+### Editor
 
 * Crash Fix: Fixed a crash when opening the level viewport context menu if the actor-component selection is out of sync.
 
@@ -1815,7 +1815,7 @@
 
 * New: Added a setting to enable/disable high DPI support in editor.
 
-* New: Replaced the skeletal mesh import option "keep overlapping vertex" with thresholds, enabling the user to control the vertices’ welding thresholds.
+* New: Replaced the skeletal mesh import option "keep overlapping vertex" with thresholds, enabling the user to control the vertices' welding thresholds.
 
 * New: The Persistent Level in the Levels tab can now be locked.
 
@@ -1855,7 +1855,7 @@
 
 * Removed: Removed .fbx export file version 2010 compatibility. The 2018 .fbx SDK does not export to 2010 and earlier file formats.
 
-##### Content Browser
+#### Content Browser
 
 * Crash Fix: Fixed a crash when renaming an asset without an enabled source control via the Asset Tool.
 
@@ -1877,7 +1877,7 @@
 
 * New: Added the ability for the Reference Viewer to show and copy references lists for nodes with multiple objects, or multiple selected nodes.
 
-##### Foliage
+#### Foliage
 
 * Crash Fix: Fixed a crash for Instanced Static meshes and Foliage when used with a Static Mesh that had an invalid lightmap coordinate index set.
 
@@ -1899,7 +1899,7 @@
 
 * Bug Fix: Fixed a special case where adding and removing multiple times would cause the regeneration of random streams for all changed instances.
 
-##### Landscape
+#### Landscape
 
 * Crash Fix: Fixed a crash when updating Landscape grass.
 
@@ -1911,7 +1911,7 @@
 
 * New: Added the possibility to update the Landscape Material at runtime using Material Instance Dynamic Materials.
 
-##### Material Editor
+#### Material Editor
 
 * Crash Fix: Fixed a crash when filtering Material Layers drop-down with no set value.
 
@@ -1921,7 +1921,7 @@
 
 * Bug Fix: Fixed an issue to not show empty Global group if it only contained the Material Layer parameter.
 
-* Bug Fix: Fixed an issue with Details panels to now "own" a color picker so that a different Detail panel refreshing doesn’t close it. This fixed the refreshing state of the graph after changing the texture or color parameter values. 
+* Bug Fix: Fixed an issue with Details panels to now "own" a color picker so that a different Detail panel refreshing doesn't close it. This fixed the refreshing state of the graph after changing the texture or color parameter values. 
 
 * Bug Fix: Fixed an issue to allow visible parameters parameter retrieval to correctly recurse through internally called functions. Previous checks were intended to prevent function previews from leaving their graph through unhooked inputs but unintentionally blocked all function inputs. 
 
@@ -1977,9 +1977,9 @@
 
 * Removed: Parent drop-down has been removed from Layers and Clends. A filter button has been added instead.
 
-##### Media Framework
+#### Media Framework
 
-* Crash Fix: Fixed a crash with Media assets in Media Sound Component’s "Update Player."
+* Crash Fix: Fixed a crash with Media assets in Media Sound Component's "Update Player."
 
 * Crash Fix: Fixed a crash when opening AVF Media after opening multiple media players.
 
@@ -1989,7 +1989,7 @@
 
 * Bug Fix: Fixed external textures not registering correctly in all cases for the Engine.
 
-* Bug Fix: Fixed Media Asset issue with Media Player’s Previous and Next not playing automatically if the player was playing previously. 
+* Bug Fix: Fixed Media Asset issue with Media Player's Previous and Next not playing automatically if the player was playing previously. 
 
 * Bug Fix: Fixed an issue with image media player never finishing initialization if loading failed.
 
@@ -2013,15 +2013,15 @@
 
 * Removed: Looping option has been removed from playlists for Media Assets.
 
-##### Sequencer
+#### Sequencer
 
 * Crash Fix: Fixed a copy/paste crash that could happen when only processing Movie Scene Copyable Binding and objects that can be spawned by the movie scene spawn were registered. 
 
 * Crash Fix: Fixed a crash when an existing hotspot is null.
 
-* Crash Fix: Fixed a crash when pasting a camera cut track that happened when the outer track isn’t set to the movie scene. 
+* Crash Fix: Fixed a crash when pasting a camera cut track that happened when the outer track isn't set to the movie scene. 
 
-* Crash Fix: Prevented a crash by ensuring that the null object in property actuate when there’s a track without a binding object. This occured when pasting a property track to the root. 
+* Crash Fix: Prevented a crash by ensuring that the null object in property actuate when there's a track without a binding object. This occured when pasting a property track to the root. 
 
 * Crash Fix: Fixed a crash when dragging a Level Sequence into the tree area.
 
@@ -2031,7 +2031,7 @@
 
 * Crash Fix: Prevented a reload crash that could happen if a Sequence was being recorded in a currently opened Sequencer. On Recording Started and On Recording Finished are now set to multicast delegates.
 
-* Crash Fix: Added a nullptr check to fix a crash in Actor Sequence when it doesn’t exist.
+* Crash Fix: Added a nullptr check to fix a crash in Actor Sequence when it doesn't exist.
 
 * Crash Fix: Fixed a crash opening a new Sequence while another is still active.
 
@@ -2067,11 +2067,11 @@
 
 * Bug Fix: Fixed a delay before warm up causing Sequences not to render.
 
-* Bug Fix: Fixed an issue to enable the camera cut track when popping back to the master only if there’s a camera cut track in the Master Sequence. It addresses an issue where if you don’t have a camera cut track in the Master, the camera gets locked to a camera cut in a sub-scene and you can’t toggle out of it. 
+* Bug Fix: Fixed an issue to enable the camera cut track when popping back to the master only if there's a camera cut track in the Master Sequence. It addresses an issue where if you don't have a camera cut track in the Master, the camera gets locked to a camera cut in a sub-scene and you can't toggle out of it. 
 
 * Bug Fix: Pivot locations are now updated for a selection when closing Sequencer.
 
-* Bug Fix: Fixed an overlap issue by looking at the key behind it. It addresses if there’s 3 keyframes on consecutive frames, so if you zoom out, you should see two bordered keys when the overlap threshold is passed. 
+* Bug Fix: Fixed an overlap issue by looking at the key behind it. It addresses if there's 3 keyframes on consecutive frames, so if you zoom out, you should see two bordered keys when the overlap threshold is passed. 
 
 * Bug Fix: Fixed UMG animations not working with Blueprint Nativization.
 
@@ -2087,7 +2087,7 @@
 
 * Bug Fix: Fixed an issue with Sequence Recorder to close the target animation asset editor if it exists before recording into it. 
 
-* Bug Fix: Fixed an issue with Sequence Recorder to record to the specified target animation for the target Actor only. Newly tracked components will have newly created animations so that they don’t record to the same target animation assets.
+* Bug Fix: Fixed an issue with Sequence Recorder to record to the specified target animation for the target Actor only. Newly tracked components will have newly created animations so that they don't record to the same target animation assets.
 
 * Bug Fix: Sequencer now uses drag and drop actor factory if specified. This ensures that the correct Actor Factory is used in creating the object template for the Sequencer spawner. It fixes some spawnables not getting created properly (for example, an Empty Actor).
 
@@ -2099,7 +2099,7 @@
 
 * Bug Fix: Fixed WMF Media that caused H.265 frames to be dropped due to false negative buffer size checks.
 
-* Bug Fix: Fixed an issue in Sequencer to always save the default spawnable states regardless of the focused Sequence. Addresses issue if you step back to the Master Sequence (and the spawnable if it still exists) and then scrubs outside the region where the spawnable exists. It gets destroyed but saved default spawnable state doesn’t get called because it’s no longer the focused Sequence.
+* Bug Fix: Fixed an issue in Sequencer to always save the default spawnable states regardless of the focused Sequence. Addresses issue if you step back to the Master Sequence (and the spawnable if it still exists) and then scrubs outside the region where the spawnable exists. It gets destroyed but saved default spawnable state doesn't get called because it's no longer the focused Sequence.
 
 * Bug Fix: Fixed Sequencer viewport invalidation so that it happens on Sequence evaluation.
 
@@ -2109,7 +2109,7 @@
 
 * Bug Fix: Fixed an issue with Sequence Recorder when recording spawnables not getting the correct position for being spawned at.
 
-* Bug Fix: Fixed an issue with Sequence Recorder to record spawn Actors immediately so that they won’t be missed if they’re deleted before Tick.
+* Bug Fix: Fixed an issue with Sequence Recorder to record spawn Actors immediately so that they won't be missed if they're deleted before Tick.
 
 * Bug Fix: Fixed a problem where Sequencer was not updating after toggling Bind Sequencer to PIE/Simulate while PIE is active.
 
@@ -2127,11 +2127,11 @@
 
 * Bug Fix: Fixed an issue in Sequencer when exporting unmatched float properties to custom attributes. 
 
-* Bug Fix: Fixed an issue in Sequencer to override the animation asset in the player state if it doesn’t match the animation asset being evaluated. 
+* Bug Fix: Fixed an issue in Sequencer to override the animation asset in the player state if it doesn't match the animation asset being evaluated. 
 
-* Bug Fix: Fixed an issue in Sequencer so that we don’t update current time to be within the view range when stepping into a Sequence. 
+* Bug Fix: Fixed an issue in Sequencer so that we don't update current time to be within the view range when stepping into a Sequence. 
 
-* Bug Fix: Fixed an issue to skip binding IDs if the Sequence can’t be found. 
+* Bug Fix: Fixed an issue to skip binding IDs if the Sequence can't be found. 
 
 * New: Added some options to the Curve Editor for Modify Owner Change to the Curve Owner Interface. We call mark as changed when modifying keys or tangents. 
 
@@ -2151,13 +2151,13 @@
 
     * Removed the concept of shared tracks. Any previous uses should port over to shared execution tokens instead.
 
-* New: Now automatically create a camera cut track if a camera is dropped and there’s no existing camera cut or there’s no existing camera cut sections. 
+* New: Now automatically create a camera cut track if a camera is dropped and there's no existing camera cut or there's no existing camera cut sections. 
 
 * New: The sub-Sequence name is now displayed on binding ID pickers rather than sub-section name.
 
 * New: Added an active column for Sequence Recorder.
 
-* New: Added an option to export an object’s transform to a Camera Anim asset.
+* New: Added an option to export an object's transform to a Camera Anim asset.
 
 * New: We now allow blending for vector tracks.
 
@@ -2191,7 +2191,7 @@
 
 * Removed: In Sequencer, the hotkeys Shift-C and Shift-D for toggling the Cinematic viewport have been removed. This was causing some confusion when users were accidently hitting them. 
 
-##### Static Mesh Editor
+#### Static Mesh Editor
 
 * Crash Fix: Fixed an Editor crash when a distributed Simplygon job fails.
 
@@ -2199,11 +2199,11 @@
 
 * New: Added LOD menu to viewport toolbar.
 
-##### World Browser
+#### World Browser
 
 * Bug Fix: Fixed issues with re-importing of tiled Landscape data in the Word Composition window.
 
-#### Gameplay Framework
+### Gameplay Framework
 
 * Crash Fix: Fixed a crash in the Get Asset Registry Tags function in Static Mesh that happened when the Find Object function was called during saving.
 
@@ -2215,11 +2215,11 @@
 
 * Crash Fix: Removed an extra dereference that was causing a crash in the Gameplay Cue Manager.
 
-* Bug Fix: Added tracking of Idle Seconds to the Performance Data Consumer interface’s Frame Data.
+* Bug Fix: Added tracking of Idle Seconds to the Performance Data Consumer interface's Frame Data.
 
 * Bug Fix: Fixed the Load Asset blueprint node so it can be correctly called from a loop with different inputs.
 
-* Bug Fix: Player Camera Manager’s Update View Target Internal function will no longer apply values from Blueprint Update Camera if it returns false.
+* Bug Fix: Player Camera Manager's Update View Target Internal function will no longer apply values from Blueprint Update Camera if it returns false.
 
 * Bug Fix: Fixed an issue where module redirects could cause the Asset Manager to incorrectly allow Primary Assets with the wrong base class.
 
@@ -2231,7 +2231,7 @@
 
 * Bug Fix: Components and Tick functions will no longer be unregistered and re-registered when testing a rename for an Actor.
 
-* Bug Fix: Modified Gameplay Tag’s Import Text Item function to account for redirects when establishing Tag Name.
+* Bug Fix: Modified Gameplay Tag's Import Text Item function to account for redirects when establishing Tag Name.
 
 * Bug Fix: Fixed character proxies doing up to two floor checks when only rotation is changed (there shouldn't be any). Added some optional verbose logging to the Find Floor and Compute Floor Dist functions.
 
@@ -2251,7 +2251,7 @@
 
 * Bug Fix: Optimized client movement to avoid updating skeletal mesh bones on the client when teleporting back to combine a network move to the server. This can be avoided since the following Perform Movement call will move the mesh and cause the update there.
 
-* Bug Fix: Added stat for Force Position Update (when server ticks client that has timed out) so that it doesn’t show up as Perform Movement.
+* Bug Fix: Added stat for Force Position Update (when server ticks client that has timed out) so that it doesn't show up as Perform Movement.
 
 * Bug Fix: Redirectors or otherwise invalid .umap files are no longer considered sub levels for world composition.
 
@@ -2279,7 +2279,7 @@
 
 * New: Primary Data Asset can now be directly blueprinted to enable blueprint-only Primary Assets, even if Should Manager Determine Type And Name is disabled. If you do this, the Primary Asset Type will be the blueprint below Primary Data Asset, which you can then inherit from to create your Primary Assets.
 
-* New: If overrides of Actor Component’s Begin Play function are not routed up the hierarchy code, the engine will now use Ensure to warn that initialization steps were missed.
+* New: If overrides of Actor Component's Begin Play function are not routed up the hierarchy code, the engine will now use Ensure to warn that initialization steps were missed.
 
 * New: Alpha Blend struct is now marked as BlueprintType.
 
@@ -2295,11 +2295,11 @@
 
 * New: Reduced overhead of evaluating streaming level volumes when all Player Controllers have Is Using Streaming Volumes set to false.
 
-* New: Significance Manager’s Update function now takes viewpoints in as an Array View instead of a const Array reference.
+* New: Significance Manager's Update function now takes viewpoints in as an Array View instead of a const Array reference.
 
 * New: When keys with unknown character code are received by the input system, a new Key struct will now be dynamically created to represent it rather than it being rejected as an invalid key.
 
-* New: Significance Managers’ Post Significance Update functions can now specify to be executed sequentially on the game thread instead of only as concurrently in the parallel for.
+* New: Significance Managers' Post Significance Update functions can now specify to be executed sequentially on the game thread instead of only as concurrently in the parallel for.
 
 * New: Conditional Init Axis Properties is now protected instead of private, so it can be seen by child classes.
 
@@ -2313,7 +2313,7 @@
 
 * New: Hooked up gameplay tag containers to use Post Script Construct to populate parent tags.
 
-#### Localization
+### Localization
 
 * Crash Fix: Fixed a crash if a break iterator outlives ICU during shutdown.
 
@@ -2349,7 +2349,7 @@
 
 * New: Added a way to display and reset the package localization ID via the Content Browser.
 
-#### Networking
+### Networking
 
 * Crash Fix: Fixed a rare crash that could occur on servers where replicated, placed-in-map actors could be destroyed, and clients could send references to these actors in server function parameters.
 
@@ -2379,7 +2379,7 @@
 
 * Bug Fix: Fixed a replication checksum issue with Slate brushes.
 
-* Bug Fix: Added stats to Net Driver’s Tick Flush function and analytic stat-gathering blocks within that function.
+* Bug Fix: Added stats to Net Driver's Tick Flush function and analytic stat-gathering blocks within that function.
 
 * Bug Fix: When a client informs the server that it has loaded a streaming level, we now force a net update on all dormant actors that have, at any point, replicated data to relevant clients and ensure that the connection's destroyed startup/dormant actors list is properly populated.
 
@@ -2449,7 +2449,7 @@
 
 * New: The "net.UseEncryptionToken" console variable has been renamed to “net.AllowEncryption” and made more robust.
 
-#### Online
+### Online
 
 * Crash Fix: Fixed a crash in web browser module destruction related to the freeing of memory via default destructors after Cef Shutdown has been called.
 
@@ -2479,7 +2479,7 @@
 
 * New: Added the ability to update cacert.pem from Curl's website. ([https://curl.haxx.se/docs/caextract.html](https://curl.haxx.se/docs/caextract.html))
 
-#### Physics
+### Physics
 
 * Crash Fix: Fixed a crash that would delete max distance mask from a clothing data entry.
 
@@ -2529,7 +2529,7 @@
 
 * New: Added 'Set Use CCD' function for changing Continuous Collision Detection option at run-time.
 
-##### Clothing
+#### Clothing
 
 * Crash Fix: Fixed a crash in clothing actor creation when the clothing simulation mesh has no simulated vertices.
 
@@ -2545,7 +2545,7 @@
 
 * New: Changed "bCompileNvCloth" compile option to use the same pattern as “bCompileAPEX” (on by default, disabled on some platforms). This allows game projects to disable it.
 
-##### Collision Detection
+#### Collision Detection
 
 * New: Improved "Auto Convex" tool for generating simplified convex collision for meshes.
 
@@ -2555,7 +2555,7 @@
 
     * Now using V-HACD version 2.3.
 
-##### Physics Asset Editor
+#### Physics Asset Editor
 
 * Crash Fix: Fixed a crash undoing primitive regeneration when stopping an in-progress simulation.
 
@@ -2591,7 +2591,7 @@
 
 * New: Added error-reporting to the physics profile name widgets.
 
-#### Platforms
+### Platforms
 
 * Crash Fix: Updated Intel SPC Tex Comp DLLs to address crashes with some processors on Windows during texture compression.
 
@@ -2607,7 +2607,7 @@
 
 * New: Added a descriptive error message when using an Invalid iOS Mobile Provision.
 
-##### All Mobile
+#### All Mobile
 
 * Bug Fix: Fixed an issue for functional code checks being removed for shipping builds in media framework.
 
@@ -2617,7 +2617,7 @@
 
 * New: Disabled manual vertex fetch on mobile.
 
-##### Android
+#### Android
 
 * Crash Fix: Fixed crash if quickly clicking on "Configure Now" button in Android Project Settings multiple times.
 
@@ -2629,9 +2629,9 @@
 
     * Backspace on an empty line would crash.
 
-    * Virtual keyboard text field isn’t visible after rotating from landscape to portrait. 
+    * Virtual keyboard text field isn't visible after rotating from landscape to portrait. 
 
-    * Chinese and Korean virtual keyboards don’t allow native characters.
+    * Chinese and Korean virtual keyboards don't allow native characters.
 
     * Number pad reverts back to normal keyboard
 
@@ -2639,7 +2639,7 @@
 
     * Cursor and UMG do not match. 
 
-    * Text field doesn’t appear if there was too much text.
+    * Text field doesn't appear if there was too much text.
 
 * Crash Fix: Samsung Galaxy S8 with Android 7.0 is not hiding suggestions and disabling predictive input. There are cases with this that can cause a crash. Global solution: On text change, downgrade to simple suggestions all the easy correction spans that are not a spell check span (remove android.text.style.SuggestionSpan.FLAG_EASY_CORRECT flags).
 
@@ -2703,7 +2703,7 @@
 
 * New: Added LOD Proxy support.
 
-##### HTML5
+#### HTML5
 
 * Bug Fix: Chrome renders transparent on OSX even though canvas has been set with alpha to false. 
 
@@ -2711,7 +2711,7 @@
 
 * Bug Fix: Fixed an issue for "Runtime Error: integer result unrepresentable" compiler error removed the following linker options: 
 
-    * -s OUTLINTING_LIMIT flags -s BINARYEN_TRAP_MODE=’clamp’
+    * -s OUTLINTING_LIMIT flags -s BINARYEN_TRAP_MODE='clamp'
 
 * Bug Fix: Disabled Support Screen Percentage to render the full screen otherwise. This was causing HTML5 screens to render only a portion of the screen and in black. 
 
@@ -2727,7 +2727,7 @@
 
 * New: Resurrected vehicle games and removed "IsContentOnly" check during HTML file generation to allow code projects to work again. This corrected optimization suffix for PhysX 3 Vehicles. 
 
-##### iOS
+#### iOS
 
 * Bug Fix: Fixed an issue while remote compiling Metal shaders to now check the connection before starting the compile process. 
 
@@ -2751,7 +2751,7 @@
 
 * New: Disabled the validation layer in iOS App Store builds for shipping builds.
 
-##### Linux
+#### Linux
 
 * Crash Fix: Fixed a crash on a headless system.
 
@@ -2863,7 +2863,7 @@
 
 * Removed: Removed the support for /-prefixed command line switches.
 
-##### Mac
+#### Mac
 
 * Crash Fix: Fixed a rare crash caused by a deferred call to Close Window function executing during Slate's PrepassWindowAndChildren call.
 
@@ -2903,7 +2903,7 @@
 
 * New: Added support for changing monitor's display mode on Mac in fullscreen mode. This greatly improves performance on Retina screens when playing in resolutions lower than native.
 
-##### Windows
+#### Windows
 
 * Bug Fix: Switched Windows to use wprintf() consistently.
 
@@ -2911,11 +2911,11 @@
 
 * Bug Fix: Fixed an issue with custom window positions in "windowed mode".
 
-#### Programming
+### Programming
 
 * New: Minimized number of fractional digits you can have in a resultant string while using FString::SanitizeFloat, defaulting to 1. Setting this to 0 prevents adding '.0' to integers.
 
-* New: Added ‘.idea/misc.xml’ file generation to speed up CLion indexing.
+* New: Added ‘.idea/misc.xml' file generation to speed up CLion indexing.
 
 * New: Added option to generate minimal set of targets for CMake files.
 
@@ -2935,11 +2935,11 @@
 
 * New: Reduced calls to malloc by using larger memory blocks for hlslcc Type allocation pools.
 
-* New: Optimized hlslcc memory allocation policy on Mac, resulting in `34%` faster and `83%` cheaper memory moves.
+* New: Optimized hlslcc memory allocation policy on Mac, resulting in 34\% faster and 83\% cheaper memory moves.
 
 * New: Added basic support for hotfixing live assets from .ini files.
 
-#### Rendering
+### Rendering
 
 * Crash Fix: Fixed a Metal shader standards iOS crash when users would try to load a project that tries to use an incompatible shader version.
 
@@ -3131,13 +3131,13 @@
 
 * New: Added RHI Command List Enqueue Lambda method which can be used to enqueue arbitrary tasks on the RHI thread from the render thread without having to write lots of boilerplate FRHI Command functor classes.
 
-##### FX
+#### FX
 
 * Bug Fix: Fixed texture streaming with particle Subuv.
 
 * New: Implemented a fast particle pool memory which default is 2MB and automatically cleans up oldest used pool slots. 
 
-##### Lighting
+#### Lighting
 
 * Bug Fix: Fixed Lightmass correctness issues so that, Point, Spot, Directional and Skylights more closely match reference renderer. 
 
@@ -3181,7 +3181,8 @@
 
 * Bug FIx: Added a New World Settings Lightmass property, Volumetric LightmapSphericalHarmonic Smoothing, which  controls the global amount of smoothing applied.
 
-* Bug Fix: Hidden levels now keep their last build data when using lighting scenarios.Fixed light build bug where hidden levels affected volumetric lightmaps.
+* Bug Fix: Hidden levels now keep their last build data when using lighting scenarios.
+Fixed light build bug where hidden levels affected volumetric lightmaps.
 
 * Bug Fix: Prevent an assert from firing when LPV Intensity is saved to a value of 0.0.
 
@@ -3201,7 +3202,7 @@
 
 * New: Added a checkpoint for Screen Shadow Mask Texture, allowing Vis Screen Shadow Mask Texture.
 
-##### Materials
+#### Materials
 
 * Bug Fix: Return error for external texture if not used in pixel shader
 
@@ -3217,7 +3218,7 @@
 
 * New: Material Layers experimental enable flag has moved to Rendering > Experimental > Support Material Layers. 
 
-##### Mobile Rendering
+#### Mobile Rendering
 
 * Bug Fix: Cancelled the applied exposure scale with planar reflections for non-hdr mobile
 
@@ -3237,7 +3238,7 @@
 
 * New: Add mobile CSM shader culling for movable directional lights.
 
-##### Optimizations
+#### Optimizations
 
 * Bug Fix: Vulkan descriptor pools are now allocated per PSO instead of globally to reduce peak mem consumption and fragmentation (enabled on Windows only).
 
@@ -3269,7 +3270,7 @@
 
 * Depreciation: Removed deprecated Lightmap UV Bias, Shadow Map UV Bias from Instanced Static Mesh Component.
 
-##### Post Processing
+#### Post Processing
 
 * Bug Fix: Fix the black ground scattering on Metal by ensuring that input into abs() is non-negative for atmosphere rendering.
 
@@ -3279,7 +3280,7 @@
 
 * New: Implemented a SSR input Post Process Material location to compose linear color space backplate after TAA for next frame's SSR.
 
-#### Tools
+### Tools
 
 * Crash Fix: Fixed a crash when attempting to move windows in VR Mode.
 
@@ -3299,7 +3300,7 @@
 
 * New: Added support for Debug Symbols (.udebugsymbols) files generated for iOS builds to be used for memory debugging with the Memory Profiler 2 tool.
 
-##### AutomationTool
+#### AutomationTool
 
 * Bug Fix: Fixed parallel executor getting the working directory from task element, not tool node.
 
@@ -3321,7 +3322,7 @@
 
 * Removed: Code that executes pre-built steps from UE4Build have been removed. These build steps should be executed before generating the action graph. So, they have to be done inside of Unreal Build Tool. Running them here causes them to execute twice.
 
-##### BuildGraph
+#### BuildGraph
 
 * New: Added an example script for building the Editor and Game targets for a project, and submitting the binaries to source control.
 
@@ -3329,13 +3330,13 @@
 
 * New: Added an <MsBuild> task for compiling MSBuild projects.
 
-##### Network Profiler
+#### Network Profiler
 
 * New: The Network Profile is now 64 bit and the solution has been converted to VS2015 and are now about `6%` smaller.
 
 * New: Added a checkbox above the connections list that will select or deselect all connections in the Network Profiler.
 
-##### UnrealBuildTool
+#### UnrealBuildTool
 
 * Bug Fix: Fixed an issue causing Unreal Build Tool on Mac to not use makefiles which made the time it needed to determine what needs to be compiled and linked longer than on other platforms.
 
@@ -3371,7 +3372,7 @@
 
     * Changed Launch.build.cs to set private defines for the build changelist in non-formal builds. This allows the version to be set correctly with executables copies to console. 
 
-    * Modular builds on desktop platforms still read the Version file from disk since it’s quicker than compiling a value in. They need the build ID to prevent mismatched binaries. 
+    * Modular builds on desktop platforms still read the Version file from disk since it's quicker than compiling a value in. They need the build ID to prevent mismatched binaries. 
 
 * New: Version files are only generated for non-promoted builds, and are included in the target receipt.
 
@@ -3393,7 +3394,7 @@
 
 * Removed: UDN documentation source files are no longer included in the generated project files by default.
 
-##### UnrealGameSync
+#### UnrealGameSync
 
 * Bug Fix: Fixed an issue so that the project filename is now included in the Editor build command. The project may not be in one of the .uprojectdirs paths.
 
@@ -3403,11 +3404,11 @@
 
 * New: Added an in-place error panel to show when a project fails to open. This enables the user to retry and have their tabs saved instead of creating a modal dialog.
 
-##### UnrealVS
+#### UnrealVS
 
 * New: Added Log GenerateProjectFiles to file when using Unreal VS.
 
-#### UI
+### UI
 
 * Crash Fix: Fixed a line-ending inconsistency when retrieving rich-text that could cause Input Method crashes. 
 
@@ -3437,7 +3438,7 @@
 
 * New: Added Editable Text functionality.
 
-##### Slate
+#### Slate
 
 * Crash Fix: Fixed a crash when trying to use the widget reflector in a client cooked build.
 
@@ -3455,7 +3456,7 @@
 
     * This differs from Custom, which is a full override of the navigation behavior.
 
-##### UMG
+#### UMG
 
 * Crash Fix: Fixed a crash on compile due to partial GC of a widget.
 
@@ -3493,7 +3494,7 @@
 
 * New: Now able to set escape keys and no key specified text on Input Key Selector.
 
-#### VR
+### VR
 
 * Crash Fix: Fixed a crash that could occur when increasing and then decreasing screen percentage in the editor's VR preview.
 
@@ -3541,7 +3542,7 @@
 
 * New: Enabled the debug canvas layer by default for the console command window on Gear VR and Google VR.
 
-#### XR
+### XR
 
 * Bug Fix: Fixed a bug with Oculus, that caused the tracking origin to be reset on startup.
 
@@ -3573,25 +3574,25 @@
 
 ## UPGRADE NOTES
 
-#### Animation
+### Animation
 
-##### Animation Blueprints
+#### Animation Blueprints
 
 * Bug Fix: We changed the order of how to get the material default value. Previous implementation was to get parent's default value first before itself, which was incorrect, so we fixed it by getting own default value if exists. However, if you were relying on the previous behavior, you're going to see the issue. Make sure you set correct default value for the material instance if it was incorrect.
 
 * Bug Fix: If this code executes (it will add a log warning if it does) it could cause problems with deterministic cooking / patch creation as new guids will be generated during cook. Nothing will break but patch sizes could be bigger than they need to be. Re-saving the problematic assets will correct the issue.
 
-##### Skeletal Mesh
+#### Skeletal Mesh
 
 * The function name of SkeletalMeshComponent::PostBlendPhysics() has changed to FinalizeAnimationUpdate() for clarification.
 
-#### Blueprints
+### Blueprints
 
 * New: Step Into has been remapped to the F11 key. The Step Over command is now mapped to F10 instead.
 
 * User Defined Structure default values now correctly propagate to blueprints and other user structs the same way that native structs do, which may cause issues if your blueprints were depending on this not happening. If this is a problem you may need to modify your outer blueprint or struct and set some values back to empty or 0.
 
-#### Core
+### Core
 
 * VisitTupleElements has had its parameters swapped - the visitor is now the first parameter, followed by the tuple (now multiple tuples). Any existing calls will need to swap arguments.
 
@@ -3619,41 +3620,47 @@
 
 * Bug Fix: Calling FReferenceCollector::AddReferencedObjects() with a container will no longer compile if the container has pointers to forward declared types. You should ensure the types are included before calling this function.
 
-#### Editor
+### Editor
 
 * The debug canvas for stats is always dpi scaled in editor and pie. Any existing tools that directly use viewport positions and size as positions for canvas elements will need to divide out dpi scale using 1/CanvasObject->GetDPIScale().
 
-* Bug Fix: Note: This only applies to perforce users NOT using ticket based perforce servers.The perforce source control login flow will no longer make any calls to the perforce api which results in a password being stored on the local users machine by perforce. We have determined that the security level of this feature is not high enough to continue using it. Users can now expect to have to type in their perforce password each time they open the editor. To work around this users can call p4 passwd directly from a command window to save a password on their machine if they so choose.
+* Bug Fix: Note: This only applies to perforce users NOT using ticket based perforce servers.
 
-##### Sequencer
+The perforce source control login flow will no longer make any calls to the perforce api which results in a password being stored on the local users machine by perforce. We have determined that the security level of this feature is not high enough to continue using it. Users can now expect to have to type in their perforce password each time they open the editor. To work around this users can call p4 passwd directly from a command window to save a password on their machine if they so choose.
 
+#### Sequencer
+
 * 'Shared' templates are no longer supported in sequencer. Any previous usage should be easily converted to a shared execution token instead (FMovieSceneExecutionTokens::AddShared()). Such tokens are shared globally, rather than per sub-sequence which makes them more robust for shared responsibilities.
 
-#### Gameplay Framework
+### Gameplay Framework
 
 * Bug Fix: AActor subclasses that intended to be tick when the tick type is Viewports Only and did not set ShouldTickIfViewportsOnly to return true, but instead overrode TickActor and did not filter on Viewports Only will now need to correctly implement the ShouldTickIfViewportsOnly virtual to achieve this behavior.
 
-#### Networking
+### Networking
 
 * UNetConnection::ClientWorldPackageName is now a private member and should accessed through the available Get/Set methods.
 
 * Bug Fix: The net.UseAdaptiveNetUpdateFrequency console variable is now disabled by default, since it was causing replication issues in many projects. If you want to keep using this feature, you may enable it in your project's configuration files.
 
-#### Rendering
+### Rendering
 
 * We changed the Vertexbuffer layout from AoS to SoA because after this change we can bind the Input Attributes via SRV and load them in the shader directly (also from compute). Before this change it was very difficult and involved manual interpretation of the Raw Data in the Shader. We recommend everyone to follow this route and change their Vertexbuffers as we probably will deprecate mixed AoS Vertexbuffers in the Future.
 
 * Tube shaped point lights will be automatically rotated to compensate for this change. Dynamically generated point lights with a source length will need to have a 90 degrees pitch added to their orientation.
 
-* Inside Sequencer->Render->CompositionGraphOptionsSet Output Format to "Custom Render Passes"Set Include Render Passes to "Final Image"Set Capture Frames in HDR "On"Set Capture Gamut to "Linear"
+* Inside Sequencer->Render->CompositionGraphOptions
+Set Output Format to "Custom Render Passes"
+Set Include Render Passes to "Final Image"
+Set Capture Frames in HDR "On"
+Set Capture Gamut to "Linear"
 
-#### UI
+### UI
 
-##### UMG
+#### UMG
 
 * Bug Fix: If a game was depending on this functionality to keep things loaded, change your blueprint to set the return value into a variable, that will keep it hard referenced normally
 
-#### XR
+### XR
 
 * Deprecated the MotionController component's EControllerHand 'Hand' property and replaced it with a 'MotionSource' name. To support the same functionality, use names corresponding to the old enum values ("Left", "Right", etc.).
 
@@ -3661,23 +3668,23 @@
 
 ## PROGRAMMING UPGRADE NOTES
 
-#### AI
+### AI
 
 * Bug Fix: Fixed a bug in UAISense_Sight::UnregisterSource that resulted in sources being removed not getting removed from the ObservedTargets collection.
 
-##### Debugging Tools
+#### Debugging Tools
 
 * Redone how visual log is using categories. Now it's using FName rather than FLogCategoryBase to indicated log category. All UE_VLOG macros have been updated. This change made it possible to actually vislog things using blueprint functions.
 
-##### Navigation
+#### Navigation
 
 * Bug Fix: Fixed navmesh not loading properly in PIE when owning world has been duplicated-for-play
 
-#### Animation
+### Animation
 
 * Bug Fix: Now when meshes are set to EMeshComponentUpdateFlag::AlwaysTickPose, we optionally kick of a task to perform parallel update only (no evaluation). Previously we needed to run this on the game thread as if we skipped evaluation we did not run a worker task.
 
-#### Audio
+### Audio
 
 * New: The new Microphone Capture Component enables feeding audio directly into a game for use with the Audio Engine and for driving in-game parameters. Limitations:
 
@@ -3687,14 +3694,23 @@
 
     * Doesn't support recording microphone output to assets.
 
-#### Automation
+### Automation
 
-* Added automated test per-platform blacklist, defined in the project DefaultEngine.ini or platform equivalent e.g. MacEngine.ini:	[AutomationTestBlacklist]	;Use format +BlacklistTest=(Map=/Game/Tests/MapName, Test=TestName)	+BlacklistTest=(Map=/Game/Tests/Physics/Stability, Test=StackCubes)
+* Added automated test per-platform blacklist, defined in the project DefaultEngine.ini or platform equivalent e.g. MacEngine.ini:
 
-* Automated test screenshots are now able to test against a hierarchical fallback platform and RHI. This is defined in DefaultEngine.ini, for example:	[AutomationTestFallbackHierarchy]	;Use format +FallbackPlatform=(Child=/Platform/RHI, Parent=/Platform/RHI)	+FallbackPlatform=(Child=/Windows/OpenGL_SM4, Parent=/Windows/D3D11_SM5)	+FallbackPlatform=(Child=/Mac/SM5, Parent=/Windows/D3D11_SM5)
+	[AutomationTestBlacklist]
+	;Use format +BlacklistTest=(Map=/Game/Tests/MapName, Test=TestName)
+	+BlacklistTest=(Map=/Game/Tests/Physics/Stability, Test=StackCubes)
 
-#### Core
+* Automated test screenshots are now able to test against a hierarchical fallback platform and RHI. This is defined in DefaultEngine.ini, for example:
 
+	[AutomationTestFallbackHierarchy]
+	;Use format +FallbackPlatform=(Child=/Platform/RHI, Parent=/Platform/RHI)
+	+FallbackPlatform=(Child=/Windows/OpenGL_SM4, Parent=/Windows/D3D11_SM5)
+	+FallbackPlatform=(Child=/Mac/SM5, Parent=/Windows/D3D11_SM5)
+
+### Core
+
 * Any usage of FStartupPackages in project code should be removed
 
 * "Native" (now called “FNativeFuncPtr”) is now a function pointer that takes a UObject* context, rather than a UObject member function pointer. Use “P_THIS” if you were previously using the “this” pointer in your native function.
@@ -3715,7 +3731,7 @@
 
 * Bug Fix: Cleaned up delegate definitions in UObjectGlobals.h and moved several to be WITH_EDITOR. If you were using those definitions in a game you will need to wrap your use in WITH_EDITOR.
 
-#### Editor
+### Editor
 
 * Being an early feature, some Editor UI buttons do not work during playback of demos.
 
@@ -3725,7 +3741,7 @@
 
 * The SizeMap and ReferenceViewer have been merged into the AssetManagerEditor plugin so they can share the data needed to show cooked sizes.
 
-#### Gameplay Framework
+### Gameplay Framework
 
 * C++ Header cleanups for Actor, Pawn, Character, Controller, and PlayerController. Marked several fields as private that were accidentally exposed in 4.17.
 
@@ -3733,35 +3749,41 @@
 
 * Overriden versions of USignificanceManager::Update will need to change signature to receive a TArrayView instead of const TArray&.
 
-#### Networking
+### Networking
 
 * Bug Fix: Programmers using HasPendingConnection() on a listening FScocket should consider using the new WaitForPendingConnection() function instead, which prevents wasting CPU cycles due to polling. This change supported a change to the behavior of the FTcpListener constructor's InSleepTime parameter which now specifies the maximum time to wait for a connection in its polling loop. It now defaults to 1 second, but using a value of 0 will restore the previous polling behavior.
 
-#### Online
+### Online
 
-* This Online Presence Test requires that the testing user to have friends and be connected to the platform.Usage is "online test presence [OPTIONAL ID]". The ID passed should be an arbitrary non-friend. If no ID is passed, the arbitrary lookup test is skipped. This test does wait for presence updates for about two minutes, after this time if no updates are received, the harness will mark the test as a failure.
+* This Online Presence Test requires that the testing user to have friends and be connected to the platform.
+Usage is "online test presence [OPTIONAL ID]". The ID passed should be an arbitrary non-friend. If no ID is passed, the arbitrary lookup test is skipped. This test does wait for presence updates for about two minutes, after this time if no updates are received, the harness will mark the test as a failure.
 
-#### Rendering
+### Rendering
 
 * The STRONG_TYPE #define now seen in UE4 .usf/HLSL is actually hlslcc's "invariant" keyword applied as a type-qualifier to a Buffer<>/RWBuffer<> type - only valid when using Metal which exports this through ILanguageSpec and #define'd out for everyone else. When this is applied the MetalBackend will emit a "raw" buffer whose type must match the shader exactly, which avoids the overheads of Linear-Textures & Function-Constants.
 
 * Vertexbuffer Objects now bind their SRVs directly instead of filling out a StreamDescriptor. Care must be taken when the Vertexfactory is initialized before the VertexBuffer as the bound SRV will become invalid if the underlying vertexbuffer is re-created. We do not support late binding of the data anymore.
 
-##### Mobile Rendering
+#### Mobile Rendering
 
 * Removed cvar "r.AllReceiveDynamicCSM" and primitive flag “FPrimitiveSceneProxy.bReceiveCombinedCSMAndStaticShadowsFromStationaryLights”. Added “FPrimitiveSceneProxy.bReceiveMobileCSMShadows”, which now applies to both stationary and movable lights and defaults to on.
 
-#### Tools
+### Tools
 
-* UEdGraphPin::PinCategory, UEdGraphPin::PinSubcategory, and UEdGraphPin::PinName are now stored as FName instead of FString.UEdGraphPin::CreatePin has several simplified overrides so you can only specify Subcategory or SubcategoryObject or neither.UEdGraphPin::CreatePin also takes a parameter bundle for reference, const, container type, index, and value terminal type rather than a long list of default parameters.Material Expressions now store input and output names as FName instead of FString.FNiagaraParameterHandle now stores the parameter handle, namespace, and name as FName instead of FString.Most existing pin related functions using string have been deprecated.
+* UEdGraphPin::PinCategory, UEdGraphPin::PinSubcategory, and UEdGraphPin::PinName are now stored as FName instead of FString.
+UEdGraphPin::CreatePin has several simplified overrides so you can only specify Subcategory or SubcategoryObject or neither.
+UEdGraphPin::CreatePin also takes a parameter bundle for reference, const, container type, index, and value terminal type rather than a long list of default parameters.
+Material Expressions now store input and output names as FName instead of FString.
+FNiagaraParameterHandle now stores the parameter handle, namespace, and name as FName instead of FString.
+Most existing pin related functions using string have been deprecated.
 
-#### UI
+### UI
 
-##### Slate
+#### Slate
 
 * Renamed FTextBlockLayout to FSlateTextBlockLayout to reflect that it's a Slate specific type
 
-#### XR
+### XR
 
 * If you maintain a VR or AR plugin, you will have to remove your RefreshPoses() override and move any late update functionality to OnBeginRendering_RenderThread/GameThread(). Alternatively, if your plugin does not support late update, you can remove RefreshPoses() completely and override DoesSupportLateUpdate(), returning false.
 
@@ -3773,44 +3795,44 @@
 
     * Added two new functions to the IMotionController API - EnumerateSources() & GetCustomParameterValue(). Use EnumerateSources() to return a list of the supported devices (replaces the static EControllerHand enum). GetCustomParameterValue() can be used to expose float values associated with the motion device (such as FOV with cameras, etc.).
 
-#### VR
+### VR
 
 * The GetHMDDistortionEnabled method on IHeadMountedDisplay how takes a single argument of type EShaderPath. If you maintain a custom HMD plugin, you will need to update your method signature.
 
 ## DEPRECATIONS
 
-#### Automation
+### Automation
 
 * UFunctionalTestingManager::RunAllFunctionalTests has dropped an unused parameter. The original version is now deprecated.
 
-#### Core
+### Core
 
 * Removed the PLATFORM_HAS_64BIT_ATOMICS macro, as it is always 1.
 
 * Deprecated the VARARG macros and replaced them with variadic templates.
 
-#### Editor and Tools
+### Editor and Tools
 
 * Deprecated Material Baking functionality inside of the MaterialUtilities module. Users should use the functionality provided by the Material Baking module instead.
 
-#### Gameplay Framework
+### Gameplay Framework
 
 * Add APlayerController::ClientReturnToMainMenuWithTextReason to support an FText reason parameter so that the message can be localized. Deprecated APlayerController::ClientReturnToMainMenu.
 
-#### Rendering
+### Rendering
 
 * The Editor is now only supported on Feature Level 5 platforms. Feature Level 4 platforms (D3D10, OpenGL3) are no longer supported.
 
-##### Materials
+#### Materials
 
 * Material Layer experimental support feature, Shared Input Collection and associated material graph nodes have been removed. See the forum feedback thread for more insight.
 
-#### Tools
+### Tools
 
-##### UnrealBuildTool
+#### UnrealBuildTool
 
 * The bOutputToEngineBinaries flag has been removed from the target rules. The output directory for targets is always determined by the location of that target.
 
-#### XR
+### XR
 
 * The IHeadMountedDisplay::GetHMDDeviceType() method has been removed as it was not extensible enough. Code that needs know which XR plugin is currently active should use IXRTrackingSystem::GetSystemName() instead.
