Copying //UE4/Dev-Documentation to //UE4/Samples-Main (Source: //UE4/Dev-Documentation @ 4321946)

diff --git a/Engine/Documentation/Source/UE4Home.INT.udn b/Engine/Documentation/Source/UE4Home.INT.udn
--- a/Engine/Documentation/Source/UE4Home.INT.udn
+++ b/Engine/Documentation/Source/UE4Home.INT.udn
@@ -6,4 +6,69 @@ Navigation:home
 social-image:home_social.png
 type:landing
 
+[PUBLISH:Rocket]
+## Search
+
+[REGION:frontpage]
+	[OBJECT:SearchBox]
+		[PARAMLITERAL:path]
+		https://www.unrealengine.com/bing-search
+		[/PARAMLITERAL]
+		[PARAMLITERAL:placeholder]
+		Need Help? Find a topic by searching
+		[/PARAMLITERAL]
+		[PARAM:message]
+		[/PARAM]
+		[PARAMLITERAL:label]
+			Quick Search Terms:
+		[/PARAMLITERAL]
+		[PARAMLITERAL:terms]
+			Quick Start
+		[/PARAMLITERAL]
+		[PARAMLITERAL:terms]
+			How-To
+		[/PARAMLITERAL]
+		[PARAMLITERAL:terms]
+			C++
+		[/PARAMLITERAL]
+		[PARAMLITERAL:terms]
+			Blueprint
+		[/PARAMLITERAL]
+		[PARAMLITERAL:terms]
+			iOS
+		[/PARAMLITERAL]
+		[PARAMLITERAL:terms]
+			Android
+		[/PARAMLITERAL]
+		[PARAMLITERAL:terms]
+			Actor
+		[/PARAMLITERAL]
+		[PARAMLITERAL:terms]
+			Component
+		[/PARAMLITERAL]
+		[PARAMLITERAL:terms]
+			Materials
+		[/PARAMLITERAL]
+		[PARAMLITERAL:terms]
+			Lighting
+		[/PARAMLITERAL]
+		[PARAMLITERAL:terms]
+			Camera
+		[/PARAMLITERAL]
+		[PARAMLITERAL:terms]
+			UMG UI
+		[/PARAMLITERAL]
+		[PARAMLITERAL:terms]
+			Sequencer
+		[/PARAMLITERAL]
+		[PARAMLITERAL:terms]
+			VR
+		[/PARAMLITERAL]
+	[/OBJECT]
+[/REGION]
+
+## Browse by Topic
+
+[/PUBLISH:Rocket]
+
 [DIR(end:"1" output:"role" org:"hierarchy" parent:"%ROOT%" type:"landing")]
