Merging //UE4/Dev-Main to Dev-Anim (//UE4/Dev-Anim) @ CL 4324949

diff --git a/Engine/Documentation/Source/Engine/Rendering/PostProcessEffects/DepthOfField/Cinematic/CinematicDepthofField.INT.udn b/Engine/Documentation/Source/Engine/Rendering/PostProcessEffects/DepthOfField/Cinematic/CinematicDepthofField.INT.udn
new file mode 100644
--- /dev/null
+++ b/Engine/Documentation/Source/Engine/Rendering/PostProcessEffects/DepthOfField/Cinematic/CinematicDepthofField.INT.udn
@@ -0,0 +1,476 @@
+title:Cinematic Depth of Field
+description:The following depth of field methods provides a cinematic look that closely resembles photography and film for desktop and console platforms using the Deferred rendering path.
+crumbs:
+availability:Public
+version:4.20
+parent:Engine/Rendering/PostProcessEffects/DepthOfField
+order:1
+type:overview
+tags:Post Process
+tags:Depth of Field
+topic-image:DOF_Cine_Topic.png
+
+The following depth of field method provides a cinematic look that closely resembles photography and film for desktop and console platforms using the Deferred Shading renderer and Clustered Forward renderer.
+
+## Cinematic
+
+[REGION:note]
+This depth of field method replaces **Circle** DOF as the default cinematic method used in Unreal Engine 4, version 4.20. It is enabled by default on supported hardware.
+[/REGION]
+
+**Cinematic** DOF closely matches real-world cameras, similarly to Circle and Bokeh DOF, you can see circular shaped Bokeh with sharp [High Dynamic Range](Engine/Rendering/HDRDisplayOutput) (HDR) content. This method uses a procedural Bokeh simulation that provides dynamic resolution stability and alpha channel support while being faster, scalable, and optimized for projects developing on desktop and consoles.
+
+[OBJECT:ComparisonSlider]
+	[PARAM:before]
+	![Depth of Field Disabled](image_0.png)
+	[/PARAM]
+	[PARAM:after]
+	![Cinematic Depth of Field](image_1.png)
+	[/PARAM]
+	[PARAMLITERAL:position]
+	bottom
+	[/PARAMLITERAL]
+[/OBJECT]
+
+### Focusing the Subject of the Shot
+
+The key to achieving an aesthetically pleasing depth of field effect is all in how the subject(s) of the shot is focused. There are three core factors that affect setting up Depth of Field for any given shot:
+
+* Decide on the **Focal Length** of the lens.
+
+* Choose an appropriate **Aperture** (f-stop number).
+
+* Choose a **Distance to your Subject(s) **from the camera.
+
+Let's break down the elements that make up the camera and scene being captured in order to understand what's going on when we adjust these settings:
+
+![image_2.png](image_2.png)
+
+1. Points in the Scene
+
+    1. Red = Out of Focus
+
+    2. Blue = Perfectly in Focus
+
+    3. Green = Marginally in Focus
+
+2. Focus Distance to Subject of the Shot (Blue Dot)
+
+3. Lens of the camera
+
+4. Diaphragm of the Lens measured in F-stop
+
+5. Focal Length of the Lens
+
+6. Filmback / Image Sensor
+
+7. Rendered Image Result
+
+[REGION:note]
+Take note that the rendered image on the right is inverted through the lens of the camera.
+[/REGION]
+
+Here we have the multiple points (1) representing objects captured by the camera to a defined **Focus Distance** (2), in this case, the blue subject. The **Aperture** (4) defines how blurred objects in the foreground and background will be that are not in focus, and, finally, the **Focal Length** (5) of the lens controls the field of view, or how zoomed in the image will be. 
+
+#### Aperture
+
+The **Aperture** defines how sharp or blurred the foreground and background are based on the diameter of the **Diaphragm** which is controlled in f-stop. 
+
+![image_3.png](image_3.png)
+
+1. Lens of the Camera
+
+2. Aperture Diaphragm of the Lens measured in F-stop
+
+This diagram demonstrates the Diaphragm (2) blocking light passing through the Lens (1). The size of the Aperture (or f-stop number) sets how much light passes through the Lens in turn controlling how much the foreground and background is blurred  based on the Focus Distance of the shot. 
+
+There are two elements that make up the Aperture; the **F-stop** number and the **Diaphragm**.
+
+![image_4.png](image_4.png)
+
+1. The** Aperture** is the diameter of the opening through which light passes. The Aperture is defined by [REDACTED_PERSON] divided by an F-stop number.
+
+2. The **Diaphragm** is the mechanical structure made up of multiple blades used to block light. It opens and closes based on the supplied F-stop.
+
+The len's Aperture is defined by the diameter of the Diaphragm which decreases as the F-stop number increases controlling how shallow or wide the depth of field effect is. See Focus Distance for a diagram showing this effect. 
+
+To demonstrate how the Aperture works, drag the slider to change the F-stop number between 1.4, 2.8, and 5.6. Smaller F-stop numbers create a shallower DOF with more background and foreground blur. Larger F-stops increase the depth of DOF with less blur. Note that only the F-stop is being changed here, the Focus Distance and Focal Length remain constant at 7m and 75mm, respectively.
+
+[OBJECT:SliderSequence]
+	[PARAMLITERAL:id]
+	aperture
+	[/PARAMLITERAL]
+	[PARAMLITERAL:max]
+	2
+	[/PARAMLITERAL]
+	[PARAMLITERAL:caption]
+	Drag the slider to change the Aperture F-stop.
+	[/PARAMLITERAL]
+	[PARAMLITERAL:ext]
+	jpg
+	[/PARAMLITERAL]
+	[PARAM:image_list]
+	![aperture_0.png](aperture_0.png)
+	![aperture_1.png](aperture_1.png)
+	![aperture_2.png](aperture_2.png)
+	[/PARAM]
+[/OBJECT]
+
+[REGION:note]
+Normally when you would adjust the Aperture settings of a real camera, you'd have to also adjust the Exposure settings at the same time to maintain same light intensity received by the filmback / image sensor. However, in UE4, this is not a real camera, so adjusting the F-stop and Diaphragm does not control the light intensity. 
+[/REGION]
+
+#### Focal Length
+
+The **Focal Length** is the distance from the center of the lens to the **Filmback** (or image sensor) measured in millimeters (mm). The focal length of the lens your image: if the filmback grows, the field of view (FOV) will grow with it. If you were to decrease the Focal Length, the FOV would increase too. 
+
+![image_8.png](image_8.png)
+
+1. Lens of the camera.
+
+2. Focal Length
+
+To demonstrate how this works, drag the slider to change the Focal Length of the lens between 50mm, 75mm, and 100mm. Take notice of how effectively it changes the Field of View (FOV) of the shot; as the focal length increase, the FOV decreases. You can think of the focal length as being similar to zooming in except you also zoom in on the out-of-focus areas of the shot as well, increasing the blur for the foreground and background areas of the image. 
+
+Note that only the Focal Length is being changed here, the Focus Distance and F-stop remain constant at 7m and f/2.8, respectively.
+
+[OBJECT:SliderSequence]
+	[PARAMLITERAL:id]
+	focal
+	[/PARAMLITERAL]
+	[PARAMLITERAL:max]
+	2
+	[/PARAMLITERAL]
+	[PARAMLITERAL:ext]
+	jpg
+	[/PARAMLITERAL]
+	[PARAM:image_list]
+	![focal_0.png](focal_0.png)
+	![focal_1.png](focal_1.png)
+	![focal_2.png](focal_2.png)
+	[/PARAM]
+[/OBJECT]
+
+[REGION:note]
+Note that the FOV of the camera is controllable from the Cine Camera Actor through the filmback and focal length, since this is a physically accurate camera. However, from the Camera Actor you can directly control the FOV angle in degrees since it does not have a Focal Length setting. The focal length used for DOF is automatically derived by the renderer from the FOV angle, and the filmback setting.
+[/REGION]
+
+#### Focus Distance
+
+The **Focus Distance** is the distance from the center of the camera lens to the subject of the shot that is put into focus, creating a focal plane. The closer the camera is to the subject, the more the background will end up out-of-focus. 
+
+![image_12.png](image_12.png)
+
+1. Points in the Scene
+
+    1. Red = Out of Focus
+
+    2. Blue = Perfectly in Focus
+
+    3. Green = Marginally in Focus
+
+2. Focus Distance to Subject (Blue Dot)
+
+3. Lens of the Camera
+
+This diagram demonstrates the different object of the scene (1) that make up the scene; background, subject(s), and foreground. The distance from the Lens (3) to the focal subject (3) - in this case the blue dot - is the Focus Distance. Areas in the background and foreground are blurred depending on the F-stop used and the Focal Length of the camera. 
+
+We can visualize the DOF Layers by using the **Show** dropdown in the Viewport to select **Advanced > Depth of Field Layers**. In this instance, Green represents the foreground objects, Black is the area where objects are focused, and Blue the background objects.
+
+[OBJECT:ComparisonSlider]
+	[PARAM:before]
+	![Rendered Scene](image_13.png)
+	[/PARAM]
+	[PARAM:after]
+	![Depth of Field Layers Visualization](image_14.png)
+	[/PARAM]
+	[PARAMLITERAL:position]
+	bottom
+	[/PARAMLITERAL]
+[/OBJECT]
+
+To demonstrate how this works, drag the slider to change the Focus Distance between 4m, 7m, and 10m. Adjusting the Focus Distance adjusts the Focal Plane (purple dotted line) in the scene, indicating where the in focus region of the shot is. The blue reference point in the diagram uses the Character as the focal subject here. As camera and diagram shift focus, objects in the foreground and background will shift too. 
+
+Starting with 4m, the Police car in the foreground is in sharp focus and the character is out of focus in the background. The blue reference in the diagram indicates this with its focal point intersection in front of the filmback. At 7m, the character is completely in focus with background and foreground objects out of focus. At 10m, the construction barriers are in sharp focus and the character and Police car are out of focus in the foreground. The blue reference in the diagram indicates this since the intersection point extends beyond the filmback. 
+
+Through all these Focus Distance changes, the DOF effect increases or decreases depending on an object being in the foreground or background. Note that the Focus Distance is being changed here, the F-stop and Focal Length remain constant at f/1.4 and 75mm, respectively.
+
+[OBJECT:SliderSequence]
+	[PARAMLITERAL:id]
+	distance
+	[/PARAMLITERAL]
+	[PARAMLITERAL:max]
+	2
+	[/PARAMLITERAL]
+	[PARAMLITERAL:ext]
+	jpg
+	[/PARAMLITERAL]
+	[PARAM:image_list]
+	![distance_0.png](distance_0.png)
+	![distance_1.png](distance_1.png)
+	![distance_2.png](distance_2.png)
+	[/PARAM]
+[/OBJECT]
+
+### Procedural Diaphragm Simulation
+
+The job of Depth of Field is to apply focal importance to a subject(s) in your scene. Just as importantly, areas that are out-of-focus can be as important. Unreal Engine 4's depth of field simulates the lens diaphragm by enabling you to aesthetically control the shape of the Bokeh (out-of-focus areas) by adjusting the number and curvature of diaphragm blades. 
+
+![image_18.png](image_18.png)
+
+[REGION:note]
+For performance reasons, the diaphragm simulation is only supported on Engine Scalability settings of **Epic **and **Cinematic**. Lower quality settings will fall back to circle Bokeh shapes of equal area to ensure similar behavior with auto exposure functionality, in turn disabling the diaphragm blade simulation. 
+[/REGION]
+
+#### Number of Diaphragm Blades
+
+The **Diaphragm **(or Iris) is the mechanism made up of blades that can adjust to aesthetically inform the look of Bokeh shapes. Using a lower number of blades will result in a polygon-shaped diaphragm whereas a higher number of blades improves the curvature of the diaphragm to appear more natural with rounded out-of-focus Bokeh. 
+
+[OBJECT:SliderSequence]
+	[PARAMLITERAL:id]
+	blades
+	[/PARAMLITERAL]
+	[PARAMLITERAL:max]
+	3
+	[/PARAMLITERAL]
+	[PARAMLITERAL:ext]
+	jpg
+	[/PARAMLITERAL]
+	[PARAM:image_list]
+	![blades_0.png](blades_0.png)
+	![blades_1.png](blades_1.png)
+	![blades_2.png](blades_2.png)
+	![blades_3.png](blades_3.png)
+	[/PARAM]
+[/OBJECT]
+
+The number of diaphragm blades directly influences the shape of the Bokeh when the **Maximum Aperture** is increased (minimum F-stop number is decreased).
+
+[OBJECT:SliderSequence]
+	[PARAMLITERAL:id]
+	blades_zoom
+	[/PARAMLITERAL]
+	[PARAMLITERAL:max]
+	3
+	[/PARAMLITERAL]
+	[PARAMLITERAL:ext]
+	jpg
+	[/PARAMLITERAL]
+	[PARAM:image_list]
+	![blades_zoom_0.png](blades_zoom_0.png)
+	![blades_zoom_1.png](blades_zoom_1.png)
+	![blades_zoom_2.png](blades_zoom_2.png)
+	![blades_zoom_3.png](blades_zoom_3.png)
+	[/PARAM]
+[/OBJECT]
+
+#### Aperture and Max Aperture
+
+The **Aperture** is the opening of the lens controlled by the diameter of the **Diaphragm** which is measured in **f-stops**. The amount of light allowed to pass through the lens is controlled by the size of the aperture. Adjusting the size of the aperture sets the focal plane to adjust how in focus or out of focus areas outside of the focal region are. 
+
+The following chart shows the correlation between aperture size (f-stop number), max aperture (min f-stop number) and depth of field effects. 
+
+![image_27.png](image_27.png)
+
+A larger aperture (lower f-stop number) has a shallow focal region, causing the foreground and background to become more out-of-focus. A small aperture (higher f-stop number) has a wider focal region that includes more objects in focus in the foreground and background.
+
+[OBJECT:SliderSequence]
+	[PARAMLITERAL:id]
+	max_aperture
+	[/PARAMLITERAL]
+	[PARAMLITERAL:max]
+	5
+	[/PARAMLITERAL]
+	[PARAMLITERAL:caption]
+	Aperture F-stops: 1.4, 2.0, 2.8, 4.0, 5.6, 8.0
+	[/PARAMLITERAL]
+	[PARAMLITERAL:ext]
+	jpg
+	[/PARAMLITERAL]
+	[PARAM:image_list]
+	![max_aperture_0.png](max_aperture_0.png)
+	![max_aperture_1.png](max_aperture_1.png)
+	![max_aperture_2.png](max_aperture_2.png)
+	![max_aperture_3.png](max_aperture_3.png)
+	![max_aperture_4.png](max_aperture_4.png)
+	![max_aperture_5.png](max_aperture_5.png)
+	[/PARAM]
+[/OBJECT]
+
+While the Aperture can be as small as you'd like, it can only ever be as large as the maximum size of the lens. The **Maximum Aperture** (measured in f-stop) defines how large the lens can be. This also defines the curvature of the blades that make up the diaphragm by limiting the size that the aperture can open; limiting the size of the aperture shortens the effective focal region which limits how out-of-focus areas can be. 
+
+[OBJECT:SliderSequence]
+	[PARAMLITERAL:id]
+	max_aperture_blades
+	[/PARAMLITERAL]
+	[PARAMLITERAL:max]
+	3
+	[/PARAMLITERAL]
+	[PARAMLITERAL:caption]
+	5 Diaphragm Blades with Different Max Aperture Values
+	[/PARAMLITERAL]
+	[PARAMLITERAL:ext]
+	jpg
+	[/PARAMLITERAL]
+	[PARAM:image_list]
+	![max_aperture_blades_0.png](max_aperture_blades_0.png)
+	![max_aperture_blades_1.png](max_aperture_blades_1.png)
+	![max_aperture_blades_2.png](max_aperture_blades_2.png)
+	![max_aperture_blades_3.png](max_aperture_blades_3.png)
+	[/PARAM]
+[/OBJECT]
+
+In this example, the **Number of Diaphragm Blades** is set to 5. With a smaller Aperture, the blades of the diaphragm become more visible in the Bokeh shapes. Using a larger Aperture (f-stop = minimal f-stop) the Bokeh shapes will begin to become circular.
+
+In this example, a Cine Camera Actor is being used with following values:
+
+* **Min FStop**: 1.4
+* **Max FStop**: 4.0
+* **Diaphragm Blade Count**: 5
+* **Current Aperture**: 1.4 - 4.0
+
+With these values, the F-stop can only accept values between **1.4** and **4.0**. As the depth of field effect is widened with a larger F-stop (smaller aperture), the diaphragm blades become more visible in the Bokeh shapes.
+
+[REGION:note]
+Keep in mind that the aperture is not controlling the light intensity. This choice has been made so that you don't have to continually adjust the exposure at the same time like with a real camera. 
+[/REGION]
+
+### Available Settings
+
+#### Cinematic Camera
+
+The following settings are specific to the Cinematic Camera Actor. Note that you have access to the Camera and Depth of Field settings as well. 
+
+| Property | Description |
+| --- | --- |
+| Filmback Settings | |
+| Sensor Width | Sets the horizontal size of the filmback or digital sensor (measured in millimeters (mm)). |
+| Lens Settings | |
+| Min Focal Length | Sets the minimum focal length for the lens (measured in millimeters (mm)). |
+| Max Focal Length | Sets the maximum focal length for the lens (measured in millimeters (mm)). |
+| Min FStop | The maximum aperture for this lens. For example, 2.8 for an f/2.8 lens. This will also define the curvature of the blades of the diaphragm. |
+| Max FStop | The minimum aperture for this lens. For example, 8.0 for an f/8.0 lens. |
+| Diaphragm Blade Count | The number of blades that make up the diaphragm mechanism. |
+| Focus Settings | |
+| Focus Method | Select the type of focus method to use to control camera focus; None, Manual, or Tracking. |
+| Manual Focus Distance | Set the manually-controlled focus distance. This is only available when the Focus Method is set to Manual. |
+| Draw Debug Focus Plane | Enable to draw a translucent plane at the current focus depth. This enables a visual way to track where focus is for the shot. |
+| Debug Focus Plane Color | Set a color for Draw Debug Focus Plane when it is enabled. |
+| Smooth Focus Changes | Enable this to use interpolation to smooth out changes in focus distance. When disabled, the focus changes will be instantaneous. |
+| Focus Smoothing Interp Speed | Controls the interpolation speed when smoothing focus distance changes. This option is ignored if Smooth Focus Changes is not enabled. |
+| Focus Offset | This adds additional offset to focus depth which can be used to manually tweak focus if the chosen focus method needs adjustment. |
+| Current Focal Length | Controls the current focal length of the camera which controls the field of view (FoV) and zoom. |
+| Current Aperture | Sets the current aperture size based on an f-stop number. Note that this setting only accepts values that are within the Min FStop and Max FStop range. |
+| Current Focal Distance | Displays a read-only value that is used by the Focus Settings. |
+| Current Horizontal FOV | Displays a read-only value that is used by [REDACTED_PERSON] and Filmback Settings. |
+
+#### Post Process Volume and Camera Actor
+
+The following settings are available with the Cine Camera, Camera, and Post Process Volumes.
+
+| Property | Description |
+| --- | --- |
+| Camera Settings | |
+| Aperture (F-stop) | Defines the opening of the camera lens. Aperture is 1/f-stop with typical lens' going down to f/1.2 (larger opening). Larger numbers will reduce the DoF effect. |
+| Maximum Aperture (minimum f-stop) | Defines the maximum opening of the camera lens to control the curvature of the diaphragm. Set this to 0 to get straight blades. |
+| Number of Diaphragm Blades | Defines the number of blades of the diaphragm within the lens. Values between 4 and 16 can be used. |
+| Depth of Field Settings | |
+| Focal Distance | The distance in which the depth of field effect should be sharp. This value is measured in Unreal Units (cm). |
+| Depth Blur km for 50% | Defines the distance at which the pixel is is blurred with a radius half of Depth blur Radius. This is particularly useful to emulate cheaply atmospheric scattering. |
+| Depth Blur Radius | Radius in pixels at 1080p resolution to apply according to distance from camera to emulate atmospheric scattering. |
+
+### Useful Console Variables for Optimization
+
+The challenge of DOF is it may varies according to your content and how important artistically it has been used. For that purpose, the DOF implementation offers a variety of console variables ready to be customised per platforms, for you to control the budget you would like to allocate to your project on weaker platforms according to your priorities.
+
+Below are some console variables that can be useful for binding Cinematic DoF performance. Note that these are a good starting point and additional console variables can be found under r.DOF.*.
+
+* r.DOF.Kernel.MaxBackgroundRadius - The maximum size of the background blurring radius in horizontal screen space. 
+* r.DOF.Kernel.MaxForegroundRadius - The maximum size of the foreground blurring radius in horizontal screen space. 
+* r.DOF.Scatter.MaxSpriteRatio - The maximum ratio of scattered pixel quads as sprites. Useful to control DOF's scattering upperbound. 
+
+## Bokeh DOF method (Legacy)
+
+[REGION:warning]
+This depth of field method will be deprecated in a future release.
+[/REGION]
+
+[EXCERPT:DepthOfField_Bokeh]
+**Bokeh** DOF refers to the name of the shape that can be seen in photos or movies when objects are out of focus. This method renders a textured quad for each pixel using a user-specified texture to define a shape with which to reproduce the effect of a camera lens. 
+
+The implementation used only uses half-resolution to perform this intensive effect. It attempts to save render performance in areas where the effect is not noticeable by using Adaptive Depth of Field. Bokeh DOF is more costly than other DOF methods available in UE4, making this a prime candidate for cinematics and showcases where appealing visuals often outweigh any performance concerns in those situations. 
+[/EXCERPT:DepthOfField_Bokeh]
+
+[OBJECT:ComparisonSlider]
+	[PARAM:before]
+	![Depth of Field Disabled](image_38.png)
+	[/PARAM]
+	[PARAM:after]
+	![Bokeh Depth of Field](image_39.png)
+	[/PARAM]
+	[PARAMLITERAL:position]
+	bottom
+	[/PARAMLITERAL]
+[/OBJECT]
+
+### Adaptive Depth of Field for Bokeh DOF
+
+When using Bokeh DOF, it renders at quarter-resolution (half-resolution in each direction) by default for performance reasons. The downsampling that happens is barely noticeable, in most cases, and is perfectly acceptable; however, it can lead to artifacts and unwanted results in some situations. **Adaptive Depth of Field** resolves these types of artifacts where possible.
+
+[OBJECT:ComparisonSlider]
+	[PARAM:before]
+	![Adaptive DOF without Downsampling](image_40.png)
+	[/PARAM]
+	[PARAM:after]
+	![Adaptive DOF with Downsampling](image_41.png)
+	[/PARAM]
+	[PARAMLITERAL:position]
+	bottom
+	[/PARAMLITERAL]
+[/OBJECT]
+
+The characters being blurred in the background appear blocky when using only the downsampled DOF technique. Also, some artifacts appear around the horns of the character in the foreground. Adaptive DOF eliminates these types of artifacts, and the background characters have a much smoother appearance. 
+
+Adaptive DOF can be visualized by using the **Adaptive Depth of Field** show flag in the Level Viewport under **Show > Visualize**. When enabled, an overlay is shown where the downsampled effect is being used (green) and where the full resolution effect is being used (red). The regular scene color is shown where no blurring is being applied. 
+
+![image_42.png](image_42.png)
+
+In general, you want the Adaptive DOF to display all green. The more red shown by the visualization, the more expensive the scene is to render each frame.
+
+[OBJECT:ComparisonSlider]
+	[PARAM:before]
+	![Adaptive DOF mostly optimized](image_43.png)
+	[/PARAM]
+	[PARAM:after]
+	![Adaptive DOF unoptimized](image_44.png)
+	[/PARAM]
+	[PARAMLITERAL:position]
+	bottom
+	[/PARAMLITERAL]
+[/OBJECT]
+
+### Available Settings
+
+The following settings are available with the Cine Camera, Camera, and Post Process Volumes, located under the **Lens** tab in the **Depth of Field** section.
+
+| Property | Description |
+| --- | --- |
+| Focal Distance | The distance in which the depth of field effect should be sharp. This value is measured in Unreal Units (cm). |
+| Near Transition Region | The distance in Unreal Units from the focal region on the side nearer to the camera over which the scene transitions from focused to blurred when using Bokeh or Gaussian DOF. |
+| Far Transition Region | The distance in Unreal units from the focal region on the side farthest from the camera over which the scene transitions from focused to blurred when using Bokeh or Gaussian DOF. |
+| Scale | An overall scaling factor for Bokeh-based blurring. |
+| Max Bokeh Size | The maximum size (in percent of the view width) of the blur for the Bokeh-based DOF. Note that performance scales with the size. |
+| Shape Texture | A texture that defines the shape of the Bokeh when objects get out of focus. Note that these cannot be blended with other Post Process Volumes. |
+| Occlusion | Controls how much the blurred geometry extends beyond its usual silhouette and its opacity. A value of 0.18 gives a fairly natural occlusion result. A value of 0.4 may be necessary to solve layer color leaking issues. Very small values (less than 0.18) generally result in negating the blur effect, though can work well when objects are very near to the camera. |
+| Color Threshold | Controls the threshold where adaptive DOF switches to using full-resolution based on color. Smaller values cause more of the scene to use full-resolution processing. |
+| Size Threshold | Controls the threshold where adaptive DOF switches to using full-resolution based on size. Larger values cause more of the scene to use full-resolution processing. |
+
+## Resources
+
+* [Diaphragm (Optics)](http://en.wikipedia.org/wiki/Diaphragm_(optics))
+
+* [Aperture](http://en.wikipedia.org/wiki/Aperture)
+
+* [F-Number](http://en.wikipedia.org/wiki/F-number)
+
+* [Understanding the Factors That Affect Depth of Field](http://photography.tutsplus.com/articles/understanding-the-factors-that-affect-depth-of-field--photo-6844)
diff --git a/Engine/Source/Editor/MovieSceneTools/Private/MatineeImportTools.cpp b/Engine/Source/Editor/MovieSceneTools/Private/MatineeImportTools.cpp
--- a/Engine/Source/Editor/MovieSceneTools/Private/MatineeImportTools.cpp
+++ b/Engine/Source/Editor/MovieSceneTools/Private/MatineeImportTools.cpp
@@ -114,7 +114,8 @@ bool FMatineeImportTools::TryConvertMatineeToggleToOutParticleKey( ETrackToggleA
 }
 
 
-void FMatineeImportTools::SetOrAddKey(TMovieSceneChannelData<FMovieSceneFloatValue>& ChannelData, FFrameNumber Time, float Value, float ArriveTangent, float LeaveTangent, EInterpCurveMode MatineeInterpMode )
+void FMatineeImportTools::SetOrAddKey(TMovieSceneChannelData<FMovieSceneFloatValue>& ChannelData, FFrameNumber Time, float Value, float ArriveTangent, float LeaveTangent, EInterpCurveMode MatineeInterpMode
+	, ERichCurveTangentWeightMode WeightedMode, float ArriveTangentWeight, float LeaveTangentWeight )
 {
 	if (ChannelData.FindKey(Time) == INDEX_NONE)
 	{
@@ -124,6 +125,9 @@ void FMatineeImportTools::SetOrAddKey(TMovieSceneChannelData<FMovieSceneFloatVal
 		NewKey.TangentMode = MatineeInterpolationToRichCurveTangent( MatineeInterpMode );
 		NewKey.Tangent.ArriveTangent = ArriveTangent;
 		NewKey.Tangent.LeaveTangent = LeaveTangent;
+		NewKey.Tangent.TangentWeightMode = WeightedMode;
+		NewKey.Tangent.ArriveTangentWeight = ArriveTangentWeight;
+		NewKey.Tangent.LeaveTangentWeight = LeaveTangentWeight;
 		ChannelData.AddKey( Time, NewKey );
 	}
 }
diff --git a/Engine/Source/Editor/MovieSceneTools/Public/MatineeImportTools.h b/Engine/Source/Editor/MovieSceneTools/Public/MatineeImportTools.h
--- a/Engine/Source/Editor/MovieSceneTools/Public/MatineeImportTools.h
+++ b/Engine/Source/Editor/MovieSceneTools/Public/MatineeImportTools.h
@@ -47,7 +47,8 @@ public:
 	static bool TryConvertMatineeToggleToOutParticleKey( ETrackToggleAction ToggleAction, EParticleKey& OutParticleKey );
 
 	/** Adds a key to a rich curve based on matinee curve key data. */
-	static void SetOrAddKey( TMovieSceneChannelData<FMovieSceneFloatValue>& Curve, FFrameNumber Time, float Value, float ArriveTangent, float LeaveTangent, EInterpCurveMode MatineeInterpMode );
+	static void SetOrAddKey( TMovieSceneChannelData<FMovieSceneFloatValue>& Curve, FFrameNumber Time, float Value, float ArriveTangent, float LeaveTangent, EInterpCurveMode MatineeInterpMode
+	, ERichCurveTangentWeightMode WeightedMode = RCTWM_WeightedNone, float ArriveTangentWeight = 0.0f, float LeaveTangentWeight = 0.0f);
 
 	/** Copies keys from a matinee bool track to a sequencer bool track. */
 	static bool CopyInterpBoolTrack( UInterpTrackBoolProp* MatineeBoolTrack, UMovieSceneBoolTrack* BoolTrack );
diff --git a/Engine/Source/Editor/Sequencer/Private/SequencerEdMode.cpp b/Engine/Source/Editor/Sequencer/Private/SequencerEdMode.cpp
--- a/Engine/Source/Editor/Sequencer/Private/SequencerEdMode.cpp
+++ b/Engine/Source/Editor/Sequencer/Private/SequencerEdMode.cpp
@@ -107,7 +107,7 @@ void FSequencerEdMode::DrawHUD(FEditorViewportClient* ViewportClient,FViewport*
 {
 	FEdMode::DrawHUD(ViewportClient,Viewport,View,Canvas);
 
-	if( ViewportClient->AllowsCinematicPreview() )
+	if( ViewportClient->AllowsCinematicControl() )
 	{
 		// Get the size of the viewport
 		const int32 SizeX = Viewport->GetSizeXY().X;
@@ -142,6 +142,8 @@ void FSequencerEdMode::OnKeySelected(FViewport* Viewport, HMovieSceneKeyProxy* K
 
 	for (TWeakPtr<FSequencer> WeakSequencer : Sequencers)
 	{
+		bool bChangedSelection = false;
+
 		TSharedPtr<FSequencer> Sequencer = WeakSequencer.Pin();
 		if (Sequencer.IsValid())
 		{
@@ -150,6 +152,12 @@ void FSequencerEdMode::OnKeySelected(FViewport* Viewport, HMovieSceneKeyProxy* K
 			FSequencerSelection& Selection = Sequencer->GetSelection();
 			if (!bAddToSelection && !bToggleSelection)
 			{
+				if (!bChangedSelection)
+				{
+					Sequencer->GetSelection().SuspendBroadcast();
+					bChangedSelection = true;
+				}
+
 				Sequencer->GetSelection().EmptySelectedKeys();
 			}
 
@@ -174,6 +182,12 @@ void FSequencerEdMode::OnKeySelected(FViewport* Viewport, HMovieSceneKeyProxy* K
 							TSharedPtr<IKeyArea> KeyArea = KeyAreaNode->GetKeyArea(Section);
 							if (KeyArea.IsValid() && KeyArea->GetName() == KeyData.ChannelName)
 							{
+								if (!bChangedSelection)
+								{
+									Sequencer->GetSelection().SuspendBroadcast();
+									bChangedSelection = true;
+								}
+
 								Sequencer->SelectKey(Section, KeyArea, KeyData.KeyHandle.GetValue(), bToggleSelection);
 								goto next_key;
 							}
@@ -183,6 +197,12 @@ void FSequencerEdMode::OnKeySelected(FViewport* Viewport, HMovieSceneKeyProxy* K
 next_key:
 				continue;
 			}
+			if (bChangedSelection)
+			{
+				Sequencer->GetSelection().ResumeBroadcast();
+				Sequencer->GetSelection().GetOnKeySelectionChanged().Broadcast();
+				Sequencer->GetSelection().GetOnOutlinerNodeSelectionChangedObjectGuids().Broadcast();
+			}
 		}
 	}
 }
@@ -309,7 +329,7 @@ bool FSequencerEdMode::GetParentTM(FTransform& CurrentRefTM, const TSharedPtr<FS
 			}
 		}
 	}
-
+	
 	return false;
 }
 
@@ -354,7 +374,6 @@ FTransform FSequencerEdMode::GetRefFrame(const TSharedPtr<FSequencer>& Sequencer
 FTransform FSequencerEdMode::GetRefFrame(const TSharedPtr<FSequencer>& Sequencer, const UObject* InObject, FFrameTime KeyTime)
 {
 	FTransform RefTM = FTransform::Identity;
-
 	const AActor* Actor = Cast<AActor>(InObject);
 	if (Actor != nullptr)
 	{
@@ -477,14 +496,14 @@ void FSequencerEdMode::GetLocationAtTime(FMovieSceneEvaluationTrack* Track, UObj
 }
 
 void FSequencerEdMode::DrawTransformTrack(const TSharedPtr<FSequencer>& Sequencer, FPrimitiveDrawInterface* PDI,
-	UMovieScene3DTransformTrack* TransformTrack, const TArray<TWeakObjectPtr<UObject>>& BoundObjects, const bool bIsSelected)
+											UMovieScene3DTransformTrack* TransformTrack, const TArray<TWeakObjectPtr<UObject>>& BoundObjects, const bool bIsSelected)
 {
 	bool bHitTesting = true;
-	if (PDI != nullptr)
+	if( PDI != nullptr )
 	{
 		bHitTesting = PDI->IsHitTesting();
 	}
-
+	
 	ASequencerMeshTrail* TrailActor = nullptr;
 	// Get the Trail Actor associated with this track if we are drawing mesh trails
 	if (bDrawMeshTrails)
@@ -501,24 +520,24 @@ void FSequencerEdMode::DrawTransformTrack(const TSharedPtr<FSequencer>& Sequence
 
 	bool bShowTrajectory = TransformTrack->GetAllSections().ContainsByPredicate(
 		[bIsSelected](UMovieSceneSection* Section)
-	{
-		UMovieScene3DTransformSection* TransformSection = Cast<UMovieScene3DTransformSection>(Section);
-		if (TransformSection)
 		{
-			switch (TransformSection->GetShow3DTrajectory())
+			UMovieScene3DTransformSection* TransformSection = Cast<UMovieScene3DTransformSection>(Section);
+			if (TransformSection)
 			{
-			case EShow3DTrajectory::EST_Always:				return true;
-			case EShow3DTrajectory::EST_Never:				return false;
-			case EShow3DTrajectory::EST_OnlyWhenSelected:	return bIsSelected;
+				switch (TransformSection->GetShow3DTrajectory())
+				{
+				case EShow3DTrajectory::EST_Always:				return true;
+				case EShow3DTrajectory::EST_Never:				return false;
+				case EShow3DTrajectory::EST_OnlyWhenSelected:	return bIsSelected;
+				}
 			}
+			return false;
 		}
-		return false;
-	}
 	);
-
+	
 	FFrameRate TickResolution = Sequencer->GetFocusedTickResolution();
 	FMovieSceneEvaluationTemplate* Template = Sequencer->GetEvaluationTemplate().FindTemplate(Sequencer->GetFocusedTemplateID());
-	if (!bShowTrajectory || !Template || !TransformTrack->GetAllSections().ContainsByPredicate([](UMovieSceneSection* In) { return In->IsActive(); }))
+	if (!bShowTrajectory || !Template || !TransformTrack->GetAllSections().ContainsByPredicate([](UMovieSceneSection* In){ return In->IsActive(); }))
 	{
 		return;
 	}
@@ -539,7 +558,7 @@ void FSequencerEdMode::DrawTransformTrack(const TSharedPtr<FSequencer>& Sequence
 			GetParents(Parents, BoundObject);
 
 			FVector OldKeyPos(0);
-			FFrameTime OldKeyTime = 0;
+ 			FFrameTime OldKeyTime = 0;
 			int KeyTimeIndex = 0;
 			FTransform OldPosRefTM;;
 			FVector OldPos_G;
@@ -553,7 +572,7 @@ void FSequencerEdMode::DrawTransformTrack(const TSharedPtr<FSequencer>& Sequence
 				FVector WorldPosition;
 				FTrajectoryKey TrajectoryKey;
 
-				FKeyPositionRotation(const FTrajectoryKey &InTrajKey, const FVector &InPos, const FRotator &InRot, const FVector &InWorldPos) :
+				FKeyPositionRotation(const FTrajectoryKey &InTrajKey, const FVector &InPos, const FRotator &InRot, const FVector &InWorldPos) : 
 					 Position(InPos), Rotation(InRot), WorldPosition(InWorldPos), TrajectoryKey(InTrajKey) {}
 			};
 			TArray<FKeyPositionRotation> KeyPosRots;
@@ -563,31 +582,31 @@ void FSequencerEdMode::DrawTransformTrack(const TSharedPtr<FSequencer>& Sequence
 				FFrameTime NewKeyTime = NewTrajectoryKey.Time;
 
 				FVector NewKeyPos(0);
-				FRotator NewKeyRot(0, 0, 0);
+				FRotator NewKeyRot(0,0,0);
 
 				GetLocationAtTime(EvalTrack, BoundObject, NewKeyTime, NewKeyPos, NewKeyRot, Sequencer);
 				FTransform NewPosRefTM = GetRefFrame(Sequencer, Parents, NewKeyTime);
 				FVector NewKeyPos_G = NewPosRefTM.TransformPosition(NewKeyPos);
-				FKeyPositionRotation KeyPosRot(NewTrajectoryKey, NewKeyPos, NewKeyRot, NewKeyPos_G);
+				FKeyPositionRotation KeyPosRot(NewTrajectoryKey,NewKeyPos, NewKeyRot, NewKeyPos_G);
 				KeyPosRots.Push(KeyPosRot);
 				// If not the first keypoint, draw a line to the last keypoint.
-				if (KeyTimeIndex > 0)
+				if(KeyTimeIndex > 0)
 				{
 					int32 NumSteps = FMath::CeilToInt((TickResolution.AsSeconds(NewKeyTime) - TickResolution.AsSeconds(OldKeyTime)) / SequencerEdMode_Draw3D::DrawTrackTimeRes);
 					// Limit the number of steps to prevent a rendering performance hit
-					NumSteps = FMath::Min(100, NumSteps);
-					FFrameTime DrawSubstep = NumSteps == 0 ? 0 : (NewKeyTime - OldKeyTime)*(1.f / NumSteps);
+					NumSteps = FMath::Min( 100, NumSteps );
+					FFrameTime DrawSubstep = NumSteps == 0 ? 0 : (NewKeyTime - OldKeyTime)*(1.f/NumSteps);
 					// Find position on first keyframe.
 					FFrameTime OldTime = OldKeyTime;
 					FVector OldPos(0);
-					FRotator OldRot(0, 0, 0);
+					FRotator OldRot(0,0,0);
 					GetLocationAtTime(EvalTrack, BoundObject, OldKeyTime, OldPos, OldRot, Sequencer);
 
 					const bool bIsConstantKey = NewTrajectoryKey.Is(ERichCurveInterpMode::RCIM_Constant);
 					// For constant interpolation - don't draw ticks - just draw dotted line.
 					if (bIsConstantKey)
 					{
-						if (PDI != nullptr)
+						if(PDI != nullptr)
 						{
 							DrawDashedLine(PDI, OldPos_G, NewKeyPos_G, TrackColor, 20, SDPG_Foreground);
 						}
@@ -595,12 +614,12 @@ void FSequencerEdMode::DrawTransformTrack(const TSharedPtr<FSequencer>& Sequence
 					else
 					{
 						// Then draw a line for each substep.
-						for (int32 j = 1; j<NumSteps + 1; j++)
+						for (int32 j=1; j<NumSteps+1; j++)
 						{
 							FFrameTime NewTime = OldKeyTime + DrawSubstep*j;
 
 							FVector NewPos(0);
-							FRotator NewRot(0, 0, 0);
+							FRotator NewRot(0,0,0);
 							GetLocationAtTime(EvalTrack, BoundObject, NewTime, NewPos, NewRot, Sequencer);
 
 							FTransform RefTM = GetRefFrame(Sequencer, Parents, NewTime);
@@ -637,9 +656,9 @@ void FSequencerEdMode::DrawTransformTrack(const TSharedPtr<FSequencer>& Sequence
 
 			// Draw keypoints on top of curve
 			FColor KeyColor = TrackColor.ToFColor(true);
-			for (const FKeyPositionRotation& KeyPosRot : KeyPosRots)
+			for (const FKeyPositionRotation& KeyPosRot: KeyPosRots)
 			{
-				if (bHitTesting && PDI)
+				if (bHitTesting && PDI) 
 				{
 					PDI->SetHitProxy(new HMovieSceneKeyProxy(TransformTrack, KeyPosRot.TrajectoryKey));
 				}
@@ -666,7 +685,7 @@ void FSequencerEdMode::DrawTransformTrack(const TSharedPtr<FSequencer>& Sequence
 					}
 				}
 
-				if (bHitTesting && PDI)
+				if (bHitTesting && PDI) 
 				{
 					PDI->SetHitProxy(nullptr);
 				}
@@ -691,7 +710,7 @@ void FSequencerEdMode::DrawTracks3D(FPrimitiveDrawInterface* PDI)
 		// Map between object binding nodes and selection
 		TMap<TSharedRef<FSequencerDisplayNode>, bool > ObjectBindingNodesSelectionMap;
 
-		for (auto ObjectBinding : Sequencer->GetNodeTree()->GetObjectBindingMap())
+		for (auto ObjectBinding : Sequencer->GetNodeTree()->GetObjectBindingMap() )
 		{
 			if (!ObjectBinding.Value.IsValid())
 			{
@@ -730,8 +749,8 @@ void FSequencerEdMode::DrawTracks3D(FPrimitiveDrawInterface* PDI)
 							});
 							if (TrailPtr == nullptr)
 							{
-								UViewportWorldInteraction* WorldInteraction = Cast<UViewportWorldInteraction>(GEditor->GetEditorWorldExtensionsManager()->GetEditorWorldExtensions(GetWorld())->FindExtension(UViewportWorldInteraction::StaticClass()));
-								if (WorldInteraction != nullptr)
+								UViewportWorldInteraction* WorldInteraction = Cast<UViewportWorldInteraction>( GEditor->GetEditorWorldExtensionsManager()->GetEditorWorldExtensions( GetWorld() )->FindExtension( UViewportWorldInteraction::StaticClass() ) );
+								if( WorldInteraction != nullptr )
 								{
 									ASequencerMeshTrail* TrailActor = WorldInteraction->SpawnTransientSceneActor<ASequencerMeshTrail>(TEXT("SequencerMeshTrail"), true);
 									FMeshTrailData MeshTrail = FMeshTrailData(TransformTrack, TrailActor);
@@ -748,7 +767,6 @@ void FSequencerEdMode::DrawTracks3D(FPrimitiveDrawInterface* PDI)
 	}
 }
 
-
 FSequencerEdModeTool::FSequencerEdModeTool(FSequencerEdMode* InSequencerEdMode) :
 	SequencerEdMode(InSequencerEdMode)
 {
