Updating lerp tooltip based on updates by [REDACTED_PERSON].

diff --git a/Engine/Documentation/Source/Shared/GraphNodes/Material/Material.INT.udn b/Engine/Documentation/Source/Shared/GraphNodes/Material/Material.INT.udn
--- a/Engine/Documentation/Source/Shared/GraphNodes/Material/Material.INT.udn
+++ b/Engine/Documentation/Source/Shared/GraphNodes/Material/Material.INT.udn
@@ -606,7 +606,7 @@ Engine/Rendering/Materials/ExpressionReference/Math#if
 [EXCERPT:UMaterialExpressionLinearInterpolate]
 The LinearInterpolate expression blends between two input value(s) based on a third input value used as a mask. This can be thought of as a mask to define transitions between two textures, like a layer mask in Photoshop. The intensity of the mask Alpha determines the ratio of color to take from the two input values. 
 
-If Alpha is 0.0/black, the first input is used. If Alpha is 1.0/white, the second input is used. If Alpha is grey (somewhere between 0.0 and 1.0), the output is a blend between the two inputs. Keep in mind that the blend happens per channel. So, if Alpha is an RGB color, Alpha's red channel value defines the blend between A and B's red channels independently of Alpha's green channel, which defines the blend between A and B's green channels.
+If Alpha is 0.0, the first input is used. If Alpha is 1.0, the second input is used. If Alpha is grey (somewhere between 0.0 and 1.0), the output is a blend between the two inputs. Keep in mind that the blend happens per channel.  So, if Alpha is an RGB color, Alpha's red channel value defines the blend between A and B's red channels independently of Alpha's green channel, which defines the blend between A and B's green channels.
 
 [VAR:ToolTipFullLink]
 Engine/Rendering/Materials/ExpressionReference/Math#linearinterpolate
