Merging //UE4/Dev-Main @ 4096029 to Dev-Core (//UE4/Dev-Core)

diff --git a/Engine/Documentation/Source/DocumentationGuidelines/BrandGlossary/F/F.INT.udn b/Engine/Documentation/Source/DocumentationGuidelines/BrandGlossary/F/F.INT.udn
--- a/Engine/Documentation/Source/DocumentationGuidelines/BrandGlossary/F/F.INT.udn
+++ b/Engine/Documentation/Source/DocumentationGuidelines/BrandGlossary/F/F.INT.udn
@@ -47,4 +47,4 @@ The Foliage Instanced Mesh system allows you to quickly paint or erase sets of S
 
 ##Foliage Mode
 
-The Foliage Instanced Mesh system enables you to quickly paint or erase batched sets of Static Meshes that use hardware instancing for efficient rendering on Landscape Actors, other Static Mesh Actors, and BSP geometry.
\ No newline at end of file
+The Foliage Mode enables you to quickly paint or erase batched sets of Static Meshes that use hardware instancing for efficient rendering on Landscape Actors, other Static Mesh Actors, and BSP geometry.
\ No newline at end of file
diff --git a/Engine/Documentation/Source/DocumentationGuidelines/Syntax/EpicMarkdownSyntax.INT.udn b/Engine/Documentation/Source/DocumentationGuidelines/Syntax/EpicMarkdownSyntax.INT.udn
--- a/Engine/Documentation/Source/DocumentationGuidelines/Syntax/EpicMarkdownSyntax.INT.udn
+++ b/Engine/Documentation/Source/DocumentationGuidelines/Syntax/EpicMarkdownSyntax.INT.udn
@@ -123,7 +123,7 @@ It was desirable to have the ability to provide some formatting instructions for
     [refepiclogo]: Logo_Epic-New.jpg "Here's a title"(w:50 h:50 a:left convert:true quality:75 fill:#000000)
 	This text is in a paragraph alongside the image. A special div element follows ending floating.
 	[REGION:clear]
-		[COMMENT]
+		[COMMENT:none]
  clear float 
 [/COMMENT]
 	[/REGION]    
@@ -133,7 +133,7 @@ It was desirable to have the ability to provide some formatting instructions for
 	<p>This is the reference image <img width="50" height="50" style="float: left;" title="Here's a title" alt="Epic Logo" src="./../../../images\EpicMarkdownDocGuides\EpicMarkdownExtension\Logo_Epic-New.jpg"/> with width, height, and left float.</p>
 	<p>This text is in a paragraph alongside the image. A special div element follows ending floating.</p>
 	[REGION:clear]
-		[COMMENT]
+		[COMMENT:none]
  clear float 
 [/COMMENT]
 	[/REGION]
@@ -148,7 +148,7 @@ This is the reference image ![Epic Logo][refepiclogo] with width, height, and le
 This text is in a paragraph alongside the image. A special div element follows ending floating.
 
 [REGION:clear]
-	[COMMENT]
+	[COMMENT:none]
  clear float 
 [/COMMENT]
 [/REGION]
diff --git a/Engine/Documentation/Source/Engine/AI/BehaviorTrees/HowUE4BehaviorTreesDiffer/HowUE4BehaviorTreesDiffer.INT.udn b/Engine/Documentation/Source/Engine/AI/BehaviorTrees/HowUE4BehaviorTreesDiffer/HowUE4BehaviorTreesDiffer.INT.udn
--- a/Engine/Documentation/Source/Engine/AI/BehaviorTrees/HowUE4BehaviorTreesDiffer/HowUE4BehaviorTreesDiffer.INT.udn
+++ b/Engine/Documentation/Source/Engine/AI/BehaviorTrees/HowUE4BehaviorTreesDiffer/HowUE4BehaviorTreesDiffer.INT.udn
@@ -5,9 +5,7 @@ Description:How Unreal Engine 4 Behavior Trees Differ From Traditional Behavior
 
 [TOC(start:2 end:3)]
 
-This section is for people who are familiar with **Behavior Trees** generally and would like to dive into the Unreal Engine 4 implementation as quickly as possible.  For those who have not used behavior trees before, you may find some of the explanation here confusing. [COMMENT:none]
- You may want to check out the user guide first. You can return to this section if you're interested in what we've done differently than others. 
-[/COMMENT]
+This section is for people who are familiar with **Behavior Trees** generally and would like to dive into the Unreal Engine 4 implementation as quickly as possible.  For those who have not used behavior trees before, you may find some of the explanation here confusing.
 
 There are three critical ways in which the UE4 implementation of Behavior Trees differs from "standard" behavior trees:
 
@@ -19,9 +17,7 @@ Having an event-driven architecture grants improvements to both performance and
 
 Since the code does not have to iterate through the entire tree every tick, performance is much better!  Conceptually, instead of constantly asking "Are we there yet?", we can just rest until we are prodded and told "We are there!"
 
-When stepping forward and backward through the behavior tree's execution history to visually debug the behavior, it is ideal to make the history show relevant changes and not show irrelevant ones.  In the event-driven implementation, it is not necessary to filter out irrelevant steps that iterated over the tree and chose the same behavior as before, because that additional iteration never had to happen in the first place!  Instead, only changes to the location of execution in the tree or to blackboard values matter, and it is easy to just show those differences.  [COMMENT:none]
- (See the debugging section for more information on how to step through the execution history and other debugging tricks.) 
-[/COMMENT]
+When stepping forward and backward through the behavior tree's execution history to visually debug the behavior, it is ideal to make the history show relevant changes and not show irrelevant ones.  In the event-driven implementation, it is not necessary to filter out irrelevant steps that iterated over the tree and chose the same behavior as before, because that additional iteration never had to happen in the first place!  Instead, only changes to the location of execution in the tree or to blackboard values matter, and it is easy to just show those differences.
 
 ## Conditionals Are Not Leaf Nodes
 
@@ -35,9 +31,7 @@ First, conditional decorators make the behavior tree UI more intuitive and easie
 
 In this section of a Behavior Tree, the Decorators **Close Enough** and **Blackboard** can prevent the execution of the Sequence node's children.
 
-Another advantage of conditional decorators is that it is easy to make those decorators act as observers (waiting for events) at critical nodes in the tree.  This feature is critical to gaining full advantage from the event-driven nature of the trees.  [COMMENT:none]
- Read more about how to use Conditional Decorators and their special properties here. 
-[/COMMENT]
+Another advantage of conditional decorators is that it is easy to make those decorators act as observers (waiting for events) at critical nodes in the tree.  This feature is critical to gaining full advantage from the event-driven nature of the trees. 
 
 ## Special Handling for Concurrent Behaviors
 
@@ -81,20 +75,10 @@ For example, a service can be used to determine which enemy is the best choice f
 
 Services are active only as long as execution remains in the subtree rooted at the composite node with which the service is associated.
 
-[COMMENT:none]
- Add screenshot showing what Services look like. 
-[/COMMENT]
-
 #### Decorator "Observer Aborts" Property
 
 One common usage case for standard Parallel nodes is to constantly check conditions so that a task can abort if the conditions it requires becomes false.  For example, if you have a cat that performs a sequence, such as "Shake Rear End" and "Pounce", you may want to give up immediately if the mouse escapes into its mouse hole.  With Parallel nodes, you would have a child that checks if the mouse can be pounced on, and then another child that's the sequence to perform.  Since our Behavior Trees are event-driven, we instead handle this by having our conditional decorators observe their values and abort when necessary.  (In this example, you would just have the "Mouse Can Be Pounced On?" decorator on the sequence itself, with "Observer Aborts" set to "Self". )
 
-[COMMENT:none]
- For more information, see the <LINK>"Observer Aborts" explanation page. 
-[/COMMENT]
-
-
-
 
 ## Advantages of UE4's Approach to Concurrent Behaviors
 
diff --git a/Engine/Documentation/Source/Engine/Animation/Persona/PersonaEditor.INT.udn b/Engine/Documentation/Source/Engine/Animation/Persona/PersonaEditor.INT.udn
--- a/Engine/Documentation/Source/Engine/Animation/Persona/PersonaEditor.INT.udn
+++ b/Engine/Documentation/Source/Engine/Animation/Persona/PersonaEditor.INT.udn
@@ -125,7 +125,7 @@ See the individual editor pages for information on their interfaces and usage:
 
 
 
-[COMMENT]
+[COMMENT:none]
 
 [EXCERPT:SkeletonOverview]
 **Skeleton mode** allows for visualization and control over those aspects of a character or Skeletal Mesh involving its bone or joint hierarchy. Example uses include:
@@ -163,10 +163,6 @@ See the individual editor pages for information on their interfaces and usage:
 * Blending between different poses and animation results using the AnimGraph.
 [/EXCERPT:GraphOverview]
 
-[/COMMENT]
-
-[COMMENT]
-
 [EXCERPT:SkeletonUI]
 [REGION:simpletable]
 |Available Interface Components						|																					|
diff --git a/Engine/Documentation/Source/Engine/Animation/Persona/Viewport/PersonaViewport.INT.udn b/Engine/Documentation/Source/Engine/Animation/Persona/Viewport/PersonaViewport.INT.udn
--- a/Engine/Documentation/Source/Engine/Animation/Persona/Viewport/PersonaViewport.INT.udn
+++ b/Engine/Documentation/Source/Engine/Animation/Persona/Viewport/PersonaViewport.INT.udn
@@ -112,8 +112,62 @@
 | **Mute Audio** | Mutes audio from the preview. |
 | **Use Audio Attenuation** | Use audio attenuation when playing back audio in the preview. |
 
+### Physics
+
+The **Physics** button provides the following options: 
+
+![](PhysicsMenu.png)
+
+| Option | Description |
+| ------ | ----------- |
+| [REGION:tablesection] Anim [/REGION] ||
+| **Physics Bend** | Lets you manually control the physics and animation. |
+| **Update Joint from Animations** | When enabled, any joints will update from its animations while simulating. |
+| **Physics Update Mode** | Determines whether the simulations of the root body also updates the component transform. |
+| **Poke Pause Time** | The time between poking the ragdoll and starting to blend back. |
+| **Poke Blend Time** | The time that is taken to blend from physics to the animation. |
+| [REGION:tablesection] Simulation [/REGION] ||
+| **Grav Scale** | Scale factor for the gravity used while simulating. |
+| **Max FPS** | The max frames per second (FPS) for simulation in the Physics Asset Editor. This is helpful for targeting the same FPS as your game. A value of -1 will disable a set Max FPS. |
+| [REGION:tablesection] Mouse Spring [/REGION] ||
+| **Handle Linear Damping** | Linear damping for mouse spring forces. |
+| **Handle Linear Stiffness** | Linear stiffness for mouse spring forces. |
+| **Handle Angular Damping** | Angular damping for mouse springs forces.  |
+| **Handle Angular Stiffness** | Angular stiffness for mouse spring forces. |
+| **Interpolation Speed** | How quickly interpolation takes place for the physics target transform for mouse spring forces. |
+| [REGION:tablesection] Poking [/REGION] ||
+| **Poke Strength** | The strength of the impulse used when poking with left mouse button. |
+
+
+### Preview LODs
+
+The **Preview LODs** button provides the following options: 
+
+![](lodmenu.png)
+
+| Option | Description |
+| ------ | ----------- |
+| **LOD Auto** | Automatically selects the Level Of Detail (LOD) (this is the default setting). |
+| **LOD 0** | Selects LOD 0 as the LOD. |
+
+### Playback Speed
+
+The **Playback Speed** button provides the following options: 
+
+![](playbackspeed.png)
+
+| Option | Description |
+| ------ | ----------- |
+| **x0.1** | Sets the playback speed to 0.1.|
+| **x0.25** | Sets the playback speed to 0.25. |
+| **x0.5** | Sets the playback speed to 0.5. |
+| **x1.0** | Sets the playback speed to 1.0 (this is the default setting). |
+| **x2.0** | Sets the playback speed to 2.0. |
+| **x5.0** | Sets the playback speed to 5.0. |
+| **x10.0** | Sets the playback speed to 10.0. |
+[/EXCERPT:viewoptions]
+
 [COMMENT:none]
--
 [INCLUDE:#advanced]
 
 [EXCERPT:Animation]
@@ -178,67 +232,8 @@
 | **Mesh Info** | How mesh info should be displayed in the viewport. |
 | **Mesh Overlay Drawing** | How the mesh should be displayed while in the viewport. |
 [/EXCERPT:mesh]
-
--
 [/COMMENT]
 
-
-### Physics
-
-The **Physics** button provides the following options: 
-
-![](PhysicsMenu.png)
-
-| Option | Description |
-| ------ | ----------- |
-| [REGION:tablesection] Anim [/REGION] ||
-| **Physics Bend** | Lets you manually control the physics and animation. |
-| **Update Joint from Animations** | When enabled, any joints will update from its animations while simulating. |
-| **Physics Update Mode** | Determines whether the simulations of the root body also updates the component transform. |
-| **Poke Pause Time** | The time between poking the ragdoll and starting to blend back. |
-| **Poke Blend Time** | The time that is taken to blend from physics to the animation. |
-| [REGION:tablesection] Simulation [/REGION] ||
-| **Grav Scale** | Scale factor for the gravity used while simulating. |
-| **Max FPS** | The max frames per second (FPS) for simulation in the Physics Asset Editor. This is helpful for targeting the same FPS as your game. A value of -1 will disable a set Max FPS. |
-| [REGION:tablesection] Mouse Spring [/REGION] ||
-| **Handle Linear Damping** | Linear damping for mouse spring forces. |
-| **Handle Linear Stiffness** | Linear stiffness for mouse spring forces. |
-| **Handle Angular Damping** | Angular damping for mouse springs forces.  |
-| **Handle Angular Stiffness** | Angular stiffness for mouse spring forces. |
-| **Interpolation Speed** | How quickly interpolation takes place for the physics target transform for mouse spring forces. |
-| [REGION:tablesection] Poking [/REGION] ||
-| **Poke Strength** | The strength of the impulse used when poking with left mouse button. |
-
-
-### Preview LODs
-
-The **Preview LODs** button provides the following options: 
-
-![](lodmenu.png)
-
-| Option | Description |
-| ------ | ----------- |
-| **LOD Auto** | Automatically selects the Level Of Detail (LOD) (this is the default setting). |
-| **LOD 0** | Selects LOD 0 as the LOD. |
-
-### Playback Speed
-
-The **Playback Speed** button provides the following options: 
-
-![](playbackspeed.png)
-
-| Option | Description |
-| ------ | ----------- |
-| **x0.1** | Sets the playback speed to 0.1.|
-| **x0.25** | Sets the playback speed to 0.25. |
-| **x0.5** | Sets the playback speed to 0.5. |
-| **x1.0** | Sets the playback speed to 1.0 (this is the default setting). |
-| **x2.0** | Sets the playback speed to 2.0. |
-| **x5.0** | Sets the playback speed to 5.0. |
-| **x10.0** | Sets the playback speed to 10.0. |
-[/EXCERPT:viewoptions]
-
-
 ## Movement Toolbar
 
 [EXCERPT:movement]
diff --git a/Engine/Documentation/Source/Engine/Animation/Sequences/Notifies/AnimNotifies.INT.udn b/Engine/Documentation/Source/Engine/Animation/Sequences/Notifies/AnimNotifies.INT.udn
--- a/Engine/Documentation/Source/Engine/Animation/Sequences/Notifies/AnimNotifies.INT.udn
+++ b/Engine/Documentation/Source/Engine/Animation/Sequences/Notifies/AnimNotifies.INT.udn
@@ -158,7 +158,7 @@ To use an AnimTrail emitter you must:
 1. Create a [Particle System](Engine/Rendering/ParticleSystems) that contains an emitter with AnimTrail type data. 
 2. Create and add a Notify State of the **Trail** type to an animation.
 
-[COMMENT]
+[COMMENT:none]
  To be included as part of a HowTo.
 Setting Up the Particle System
 
@@ -231,7 +231,7 @@ There are also some debug rendering options available:
 | **Render Tangents** | This toggles the debug rendering of the curve tangents. |
 | **Render Tessellation** | This toggles debug rendering of all vertices on the curve. This shows how the curve is tessellated between the spawn points. |
 
-[COMMENT]
+[COMMENT:none]
  To be included with AnimTrail HowTo.
 ### Controlling Animation Trails in Blueprint
 While the above method is the standard approach to controlling AnimTrails, you can also control them from inside a Blueprint without the need of an Anim Notify. 
diff --git a/Engine/Documentation/Source/Engine/Animation/SyncGroups/SyncGroups.INT.udn b/Engine/Documentation/Source/Engine/Animation/SyncGroups/SyncGroups.INT.udn
--- a/Engine/Documentation/Source/Engine/Animation/SyncGroups/SyncGroups.INT.udn
+++ b/Engine/Documentation/Source/Engine/Animation/SyncGroups/SyncGroups.INT.udn
@@ -43,9 +43,7 @@ Leaders and followers are designated simply by which animation has the greater b
 |[INCLUDE:#walkanimationscaled] |In this instance, the blend weight for the run animation would be greater than the walk animation. Since they are in the same Sync Group, the walk scales down to match the length of the run animation.|
 |[INCLUDE:#runanimationscaled]|Here, the weighting has switched such that the walking animation is receiving a greater blend weight. This causes the run animation to scale up to match the length of the walk. | 
 
-[COMMENT]
--
-
+[COMMENT:none]
 
 [EXCERPT:RunAnimation]
 [REGION:raw]
@@ -81,7 +79,6 @@ Leaders and followers are designated simply by which animation has the greater b
 * Has its Notifies suppressed, deferring instead to Leader's Notifies.
 [/EXCERPT:FollowerList]
 
---
 [/COMMENT]
 
 
@@ -96,13 +93,11 @@ To set the Sync Group, select the animation node in the AnimGraph and look in th
 |**Group Name**|This is where the Sync Group can be set. Generally, you will want all related animations that are meant to blend together (such as body motions for sneak, walk, run, sprint) to be in the same Sync Group.|
 |**Group Role**|This is where you set the type of membership in the group. The available values are "Can Be Leader" (the default), "Always Leader", and "Never Leader".|
 
-[COMMENT]
--
+[COMMENT:none]
 Group Role has 3 options:
 * Can Be Leader - This is the default setting. This means that if this animation has the greatest blend weight, then it can become the Leader and other animations will scale to match its length.
 * Always Leader - If set, then this animation will always be the Leader and all animations will always scale to match its length.
 * Never Leader - If set, then this animation can never be the Leader and will always scale to match the next weighted animation.
--
 [/COMMENT]
 
 
diff --git a/Engine/Documentation/Source/Engine/Audio/WAV/WAV.INT.udn b/Engine/Documentation/Source/Engine/Audio/WAV/WAV.INT.udn
--- a/Engine/Documentation/Source/Engine/Audio/WAV/WAV.INT.udn
+++ b/Engine/Documentation/Source/Engine/Audio/WAV/WAV.INT.udn
@@ -31,41 +31,7 @@ The engine currently supports importing uncompressed little endian 16 bit wave f
 
 Importing with the normal sound factory yields a simple sound node wave that can be referenced by a sound cue for playback. These sounds can be mono or stereo. Importing with a special naming convention gives the potential of multichannel (e.g. 5.1) sounds. Up to 8 channel sounds can be compressed on all platforms, although not all platforms can play them back efficiently. The package saving process does all the necessary work to convert to the platform's native format. There is no strict speaker mapping, so the number of channels infers which channels are played on which speakers (see table below).
 
-[EXCERPT:multiChannelExtensionTable]
-
-[REGION:raw]
-|                   | Extension | 4.0         | 5.1         | 6.1         | 7.1         |
-| ---               | ---       | ---         | ---         | ---         | ---         |
-| **FrontLeft**     | _fl       |  &#x2713;   |  &#x2713;   |  &#x2713;   |  &#x2713;   |
-| **FrontRight**    | _fr       |  &#x2713;   |  &#x2713;   |  &#x2713;   |  &#x2713;   |
-| **FrontCenter**   | _fc       |             |  &#x2713;   |  &#x2713;   |  &#x2713;   |
-| **LowFrequency**  | _lf       |             |  &#x2713;   |  &#x2713;   |  &#x2713;   |
-| **SideLeft**      | _sl       |  &#x2713;   |  &#x2713;   |  &#x2713;   |  &#x2713;   |
-| **SideRight**     | _sr       |  &#x2713;   |  &#x2713;   |  &#x2713;   |  &#x2713;   |
-| **BackLeft**      | _bl       |             |             |  &#x2713;   |  &#x2713;   |
-| **BackRight**     | _br       |             |             |             |  &#x2713;   |
-[/REGION]
-
-[/EXCERPT:multiChannelExtensionTable]
-
-[COMMENT:none]
-
-
-[REGION:raw]
-|                   | Extension | 4.0 | 5.1 | 6.1 | 7.1 |  
-| --- | --- | --- | --- | --- | --- | 
-| **FrontLeft**     | _fl       | ![](icon_MatEd_Apply_40x.png) | ![](icon_MatEd_Apply_40x.png)   | ![](icon_MatEd_Apply_40x.png)   | ![](icon_MatEd_Apply_40x.png)   | 
-| **FrontRight**    | _fr       | ![](icon_MatEd_Apply_40x.png)   | ![](icon_MatEd_Apply_40x.png)   | ![](icon_MatEd_Apply_40x.png)   | ![](icon_MatEd_Apply_40x.png)   | 
-| **FrontCenter**   | _fc       |     | ![](icon_MatEd_Apply_40x.png)   | ![](icon_MatEd_Apply_40x.png)   | ![](icon_MatEd_Apply_40x.png)   | 
-| **LowFrequency**  | _lf       |     | ![](icon_MatEd_Apply_40x.png)   | ![](icon_MatEd_Apply_40x.png)   | ![](icon_MatEd_Apply_40x.png)   | 
-| **SideLeft**      | _sl       | ![](icon_MatEd_Apply_40x.png)   | ![](icon_MatEd_Apply_40x.png)   | ![](icon_MatEd_Apply_40x.png)   | ![](icon_MatEd_Apply_40x.png)   | 
-| **SideRight**     | _sr       | ![](icon_MatEd_Apply_40x.png)   | ![](icon_MatEd_Apply_40x.png)   | ![]( icon_MatEd_Apply_40x.png)   | ![](icon_MatEd_Apply_40x.png)   | 
-|**BackLeft**      | _bl       |     |     | ![](icon_MatEd_Apply_40x.png)   | ![](icon_MatEd_Apply_40x.png)   |  
-| **BackRight**     | _br       |     |     |     | ![](icon_MatEd_Apply_40x.png)   |  
-[/REGION]
-
-
-[/COMMENT]
+[INCLUDE:#multiChannelExtensionTable]
 
 Compression settings are tweaked to maintain similar quality across platforms. After testing many sounds, our musician determined an XMA quality of 40 was equivalent to an ogg vorbis quality of 0.15. 
 
@@ -98,7 +64,24 @@ Or
 
 [/EXCERPT:wav01]
 
+[COMMENT:none]
+[EXCERPT:multiChannelExtensionTable]
 
+[REGION:raw]
+|                   | Extension | 4.0         | 5.1         | 6.1         | 7.1         |
+| ---               | ---       | ---         | ---         | ---         | ---         |
+| **FrontLeft**     | _fl       |  &#x2713;   |  &#x2713;   |  &#x2713;   |  &#x2713;   |
+| **FrontRight**    | _fr       |  &#x2713;   |  &#x2713;   |  &#x2713;   |  &#x2713;   |
+| **FrontCenter**   | _fc       |             |  &#x2713;   |  &#x2713;   |  &#x2713;   |
+| **LowFrequency**  | _lf       |             |  &#x2713;   |  &#x2713;   |  &#x2713;   |
+| **SideLeft**      | _sl       |  &#x2713;   |  &#x2713;   |  &#x2713;   |  &#x2713;   |
+| **SideRight**     | _sr       |  &#x2713;   |  &#x2713;   |  &#x2713;   |  &#x2713;   |
+| **BackLeft**      | _bl       |             |             |  &#x2713;   |  &#x2713;   |
+| **BackRight**     | _br       |             |             |             |  &#x2713;   |
+[/REGION]
+
+[/EXCERPT:multiChannelExtensionTable]
+[/COMMENT]
 
 [REGION:note]
 Exporting a multi channel sound asset will create a mono .WAV file for each channel with the above speaker extension naming convention.
@@ -107,7 +90,5 @@ Exporting a multi channel sound asset will create a mono .WAV file for each chan
 
 ## Related
 
-[REGION:topics half]
 %Engine/Audio:topiccompact%
 %Engine/Audio/Overview:topiccompact%
-[/REGION]
diff --git a/Engine/Documentation/Source/Engine/Basics/Projects/UnrealGameProjects.INT.udn b/Engine/Documentation/Source/Engine/Basics/Projects/UnrealGameProjects.INT.udn
--- a/Engine/Documentation/Source/Engine/Basics/Projects/UnrealGameProjects.INT.udn
+++ b/Engine/Documentation/Source/Engine/Basics/Projects/UnrealGameProjects.INT.udn
@@ -11,11 +11,7 @@ related:Programming/Development/CompilingProjects
 related:Programming/Development/ManagingGameCode
 
 [EXCERPT:Main]
-[EXCERPT:Intro]
-
-[INCLUDE:Shared/Glossary/P/#Glossary_project]
-
-[/EXCERPT:Intro]
+[INCLUDE:#intro]
 
 %Engine/Basics/Projects/Browser:topic%
 [OBJECT:Topic]
@@ -33,20 +29,12 @@ related:Programming/Development/ManagingGameCode
 	[/PARAM]
 [/OBJECT]
 
+[/EXCERPT:Main]
 
 [COMMENT:none]
+[EXCERPT:Intro]
 
-Projects can be one of two types:
-
-* **Content-Only** - These are projects where all development is done through the editor using content assets such as maps, art assets, _Blueprints_, etc.
-* **Code** - These projects allow programmer development in addition to regular content development available in _content-only_ projects. 
-
-[REGION:tip]
-Content projects can be converted to code projects using the [](Programming/Development/ManagingGameCode/CppClassWizard).
-[/REGION]
-
-
-
-[/COMMENT]
+[INCLUDE:Shared/Glossary/P/#Glossary_project]
 
-[/EXCERPT:Main]
+[/EXCERPT:Intro]
+[/COMMENT]
\ No newline at end of file
diff --git a/Engine/Documentation/Source/Engine/Blueprints/BP_HowTo/BlueprintHowTo.INT.udn b/Engine/Documentation/Source/Engine/Blueprints/BP_HowTo/BlueprintHowTo.INT.udn
--- a/Engine/Documentation/Source/Engine/Blueprints/BP_HowTo/BlueprintHowTo.INT.udn
+++ b/Engine/Documentation/Source/Engine/Blueprints/BP_HowTo/BlueprintHowTo.INT.udn
@@ -22,30 +22,12 @@ When you are ready and have an understanding of the core aspects behind Blueprin
 [DIR(output:"fancy" type:"how-to" skill:"Beginner" parent:"Engine/Blueprints/BP_HowTo" org:"hierarchy")]
 
 ## Intermediate
-[REGION:wrapper transparent]
-[](FANCY:Engine/Blueprints/BP_HowTo/ActorReference "%Engine/Blueprints/BP_HowTo/ActorReference:description%")
-[](FANCY:Engine/Blueprints/BP_HowTo/MakingFunctions "%Engine/Blueprints/BP_HowTo/MakingFunctions:description%")
-[](FANCY:Engine/Blueprints/BP_HowTo/WorkingWithArrays "%Engine/Blueprints/BP_HowTo/WorkingWithArrays:description%")
-[](FANCY:Engine/Blueprints/UserGuide/Variables "%Engine/Blueprints/UserGuide/Variables:description%")
-[/REGION]
 
-## Advanced
-[REGION:wrapper transparent]
-[](FANCY:Engine/Blueprints/BP_HowTo/DirectBlueprintComs "%Engine/Blueprints/BP_HowTo/DirectBlueprintComs:description%")
-[](FANCY:Engine/Blueprints/BP_HowTo/Casting "%Engine/Blueprints/BP_HowTo/Casting:description%")
-[](FANCY:Engine/Blueprints/BP_HowTo/Interfaces "%Engine/Blueprints/BP_HowTo/Interfaces:description%")
-[](FANCY:Engine/Blueprints/BP_HowTo/EventDispatcher "%Engine/Blueprints/BP_HowTo/EventDispatcher:description%")
-[](FANCY:Engine/Blueprints/BP_HowTo/MakingMacros "%Engine/Blueprints/BP_HowTo/MakingMacros:description%")
-[](FANCY:Engine/Blueprints/BP_HowTo/MacroLibraries "%Engine/Blueprints/BP_HowTo/MacroLibraries:description%")
-[/REGION]
-
-[COMMENT:none]
- To remove since information is redundant. 
-
-| **[](Engine/Blueprints/BP_HowTo/BasicUsage)** | %Engine/Blueprints/BP_HowTo/BasicUsage:description% |
+[DIR(output:"fancy" type:"how-to" skill:"Intermediate" parent:"Engine/Blueprints/BP_HowTo" org:"hierarchy")]
 
+## Advanced
 
-[/COMMENT]
+[DIR(output:"fancy" type:"how-to" skill:"Advanced" parent:"Engine/Blueprints/BP_HowTo" org:"hierarchy")]
 
 
 
diff --git a/Engine/Documentation/Source/Engine/Blueprints/Editor/UIComponents/GraphEditor/GraphEditor.INT.udn b/Engine/Documentation/Source/Engine/Blueprints/Editor/UIComponents/GraphEditor/GraphEditor.INT.udn
--- a/Engine/Documentation/Source/Engine/Blueprints/Editor/UIComponents/GraphEditor/GraphEditor.INT.udn
+++ b/Engine/Documentation/Source/Engine/Blueprints/Editor/UIComponents/GraphEditor/GraphEditor.INT.udn
@@ -31,6 +31,20 @@ The **Graph Editor** panel is the heart of the Blueprint system. It is here that
 
 The **Graph Editor** tab can be navigated using the controls below:
 
+[INCLUDE:#Controls]
+
+[REGION:tip]
+StaticMesh, SoundCue, SkeletalMesh, and ParticleSystem assets can be dragged from the **Content Browser** and 
+dropped on the **Graph Editor** tab to create a new AddComponent function call with the asset automatically assigned.
+
+[REGION:raw]
+![Blueprints - Drag and Drop Component](k2_component_graph_drag.png)(w:900 convert:false)
+[/REGION:raw]
+
+[/REGION:tip]
+[/EXCERPT:Main]
+
+[COMMENT:none]
 [EXCERPT:Controls]
 
 | Control | Action |
@@ -49,16 +63,4 @@ The **Graph Editor** tab can be navigated using the controls below:
 | **Alt + Click** on a pin | Remove all wires connected to the selected pin. |
 
 [/EXCERPT:Controls]
-
-
-[REGION:tip]
-StaticMesh, SoundCue, SkeletalMesh, and ParticleSystem assets can be dragged from the **Content Browser** and 
-dropped on the **Graph Editor** tab to create a new AddComponent function call with the asset automatically assigned.
-
-[REGION:raw]
-![Blueprints - Drag and Drop Component](k2_component_graph_drag.png)(w:900 convert:false)
-[/REGION:raw]
-
-[/REGION:tip]
-[/EXCERPT:Main]
-
+[/COMMENT]
\ No newline at end of file
diff --git a/Engine/Documentation/Source/Engine/Blueprints/Editor/UIComponents/Menu/BlueprintEditorMenu.INT.udn b/Engine/Documentation/Source/Engine/Blueprints/Editor/UIComponents/Menu/BlueprintEditorMenu.INT.udn
--- a/Engine/Documentation/Source/Engine/Blueprints/Editor/UIComponents/Menu/BlueprintEditorMenu.INT.udn
+++ b/Engine/Documentation/Source/Engine/Blueprints/Editor/UIComponents/Menu/BlueprintEditorMenu.INT.udn
@@ -13,11 +13,61 @@
 
 [EXCERPT:Main]
 
-
-
-
 ### File
 
+[INCLUDE:#FileMenu]
+
+### Edit
+
+[VAR:BlueprintSearch]
+| **Search** | Find references to functions, events, variables, and pins in the current Blueprint.|
+| **Find in Blueprints** | Find references to functions, events, variables, and pins in ALL Blueprints. |
+| **Delete Unused Variables** | Deletes any variables that are never used. |
+[/VAR]
+
+[INCLUDE:#EditMenu]
+
+
+### Asset
+
+[INCLUDE:#AssetMenu]
+
+### View
+
+[INCLUDE:ViewMenu]
+
+### Debug
+
+[INCLUDE:#DebugMenu]
+
+### Window 
+
+The **Window** menu in the **Blueprint Editor** has a specific subsection for displaying Blueprint Editor-specific tabs. Not all tabs will be present in the menu while the Blueprint Editor is in the **Defaults** and **Components** modes.
+
+
+| Command | Description |
+| --- | --- |
+| **Toolbar** | Shows the Toolbar if it is currently not visible. |
+| **Details** | Shows the [Details](Engine/Blueprints\Editor/UIComponents\Details) pane if it is currently not visible. |
+| **Debug** | Shows the **Debug** pane if if it currently not visible. |
+| **Palette** | Shows the [](Engine/Blueprints\Editor/UIComponents\Palette) pane if it is not currently visible. |
+| **My Blueprint** | Shows the [](Engine/Blueprints\Editor/UIComponents\MyBlueprint) pane if it is not currently visible.  |
+| **Compiler Results** | Shows the **Compiler Results** pane if it is not currently visible.  |
+| **Find Results** | Shows the **Find Results** pane if it is not currently visible.  | 
+| **Components** | Shows the **Components panel** if it is not currently visible. |
+| **Viewport** | Shows the **Preview Viewport** pane if it is not currently visible.  |
+| **Content Browser** | Opens a sub-menu with access to all four **Content Browsers**. |
+| **Developer Tools** | [INCLUDE:#devtools] |
+| **Project Launcher** | Shows the Project Launcher which enables you to run your project on any appropriately setup and connected devices. |
+| **Plugins** | Shows the Plugin Tab where you can load/unload plugins. |
+| **Reset Layout...** | Restores the default layout for the entire Unreal Engine 4 Editor. This requires the Editor to restart, but it will re-open the current project. |
+| **Save Layout** | Saves the current interface layout. |
+| **Enable Fullscreen** | Enables fullscreen mode for the Main Editor window. |
+
+
+[/EXCERPT:Main]
+
+[COMMENT:none]
 [EXCERPT:FileMenu]
 | Command | Description |
 | --- | --- |
@@ -34,15 +84,6 @@
 | **Diff** | Diff against previous revisions. Requires source control to be active. |
 | **Developer** | Open the Developer Menu, where you can change compiler settings and recompile modules like the Graph Editor. |
 [/EXCERPT:FileMenu]
-
-### Edit
-
-[VAR:BlueprintSearch]
-| **Search** | Find references to functions, events, variables, and pins in the current Blueprint.|
-| **Find in Blueprints** | Find references to functions, events, variables, and pins in ALL Blueprints. |
-| **Delete Unused Variables** | Deletes any variables that are never used. |
-[/VAR]
-
 [EXCERPT:EditMenu]
 | Command | Description |
 | --- | --- |
@@ -57,9 +98,6 @@
 | **Project Settings** | Opens the settings for the current project. |
 | **Plugins** | Opens the Plugin Browser tab. |
 [/EXCERPT:EditMenu]
-
-
-### Asset
 [EXCERPT:AssetMenu]
 [VAR:AssetMenu]
 | Command | Description |
@@ -71,9 +109,6 @@
 
 %assetmenu%
 [/EXCERPT:AssetMenu]
-
-### View
-
 [EXCERPT:ViewMenu]
 | Command | Description |
 | --- | --- |
@@ -82,12 +117,9 @@
 | **Hide Unused Pins** | Hides all pins with no connections and no default values. |
 | **Hide Unconnected Pins** | Hides all pins with no connections.  This option will hide input pins for inputs that have been set directly on the node. |
 | [REGION:tablesection]**Zoom**[/REGION] ||
-| **Zoom to Graph Extents ** | Fits the current view to the entire graph. |
+| **Zoom to Graph Extents** | Fits the current view to the entire graph. |
 | **Zoom to Selection** | Fits the current view to the selection.  |
 [/EXCERPT:ViewMenu]
-
-### Debug
-
 [EXCERPT:DebugMenu]
 | Command | Description |
 | --- | --- |
@@ -98,36 +130,6 @@
 | [REGION:tablesection]**Watches**[/REGION] ||
 | **Delete All Watches** | Removes all watch values in all graphs of the current Blueprint or Level Blueprint. |
 [/EXCERPT:DebugMenu]
-
-### Window 
-
-The **Window** menu in the **Blueprint Editor** has a specific subsection for displaying Blueprint Editor-specific tabs. Not all tabs will be present in the menu while the Blueprint Editor is in the **Defaults** and **Components** modes.
-
-
-| Command | Description |
-| --- | --- |
-| **Toolbar** | Shows the Toolbar if it is currently not visible. |
-| **Details** | Shows the [Details](Engine/Blueprints\Editor/UIComponents\Details) pane if it is currently not visible. |
-| **Debug** | Shows the **Debug** pane if if it currently not visible. |
-| **Palette** | Shows the [](Engine/Blueprints\Editor/UIComponents\Palette) pane if it is not currently visible. |
-| **My Blueprint** | Shows the [](Engine/Blueprints\Editor/UIComponents\MyBlueprint) pane if it is not currently visible.  |
-| **Compiler Results** | Shows the **Compiler Results** pane if it is not currently visible.  |
-| **Find Results** | Shows the **Find Results** pane if it is not currently visible.  | 
-| **Components** | Shows the **Components panel** if it is not currently visible. |
-| **Viewport** | Shows the **Preview Viewport** pane if it is not currently visible.  |
-| **Content Browser** | Opens a sub-menu with access to all four **Content Browsers**. |
-| **Developer Tools** | [INCLUDE:#devtools] |
-| **Project Launcher** | Shows the Project Launcher which enables you to run your project on any appropriately setup and connected devices. |
-| **Plugins** | Shows the Plugin Tab where you can load/unload plugins. |
-| **Reset Layout...** | Restores the default layout for the entire Unreal Engine 4 Editor. This requires the Editor to restart, but it will re-open the current project. |
-| **Save Layout** | Saves the current interface layout. |
-| **Enable Fullscreen** | Enables fullscreen mode for the Main Editor window. |
-
-
-[/EXCERPT:Main]
-
-[COMMENT:none]
-
 [EXCERPT:devtools]
 [REGION:simpletable]
 | Sub-Menu Item | Description |
diff --git a/Engine/Documentation/Source/Engine/Blueprints/Editor/UIComponents/MyBlueprint/MyBlueprint.INT.udn b/Engine/Documentation/Source/Engine/Blueprints/Editor/UIComponents/MyBlueprint/MyBlueprint.INT.udn
--- a/Engine/Documentation/Source/Engine/Blueprints/Editor/UIComponents/MyBlueprint/MyBlueprint.INT.udn
+++ b/Engine/Documentation/Source/Engine/Blueprints/Editor/UIComponents/MyBlueprint/MyBlueprint.INT.udn
@@ -18,6 +18,23 @@ For example, a normal Blueprint will always have a **ConstructionScript** and an
 
 ## Creation Buttons
 
+[INCLUDE:#MyBlueprintPanel]
+
+## My Blueprint Sections
+
+The **My Blueprint** tab is divided up into 6 sections: The **Add New**, Graphs, Functions, Macros, Variables, and Event Dispatchers sections.
+
+![My Blueprint Pane](MyBlueprintPane.png)
+
+The 5 lower sections will organize elements of your Blueprint into their respective groups. While the top section gives you quick access to create new Graphs, Variables, etc... and search the entire **My Blueprint** panel.
+
+## Searching in My Blueprint
+
+The **My Blueprint** tab contains a text box used for searching the graphs belonging to the Blueprint. This works the same way as the contextual action menu that you use for adding new nodes, but is limited to searching items found within **My Blueprint**.  You can search based on name, comment, and other data.  So, if you know you have a node that is **SetActorHidden**, you can type that in the box, and the explorer will show you all the **SetActorHidden** nodes in your graph.
+
+[/EXCERPT:Main]
+
+[COMMENT:none]
 [EXCERPT:MyBlueprintPanel]
 The **My Blueprint** tab has shortcut buttons (the ![](plus_Button.png) at the end of each heading) for creating new variables, functions, macros, event graphs, and event dispatchers.
 
@@ -38,23 +55,4 @@ Finally, you can use the **Add New** button at the top of the **My Blueprint** P
 ![](addNew.png)
 
 [/EXCERPT:MyBlueprintPanel]
-
-
-## My Blueprint Sections
-
-The **My Blueprint** tab is divided up into 6 sections: The **Add New**, Graphs, Functions, Macros, Variables, and Event Dispatchers sections.
-
-![My Blueprint Pane](MyBlueprintPane.png)
-
-The 5 lower sections will organize elements of your Blueprint into their respective groups. While the top section gives you quick access to create new Graphs, Variables, etc... and search the entire **My Blueprint** panel.
-
-## Searching in My Blueprint
-
-The **My Blueprint** tab contains a text box used for searching the graphs belonging to the Blueprint. This works the same way as the contextual action menu that you use for adding new nodes, but is limited to searching items found within **My Blueprint**.  You can search based on name, comment, and other data.  So, if you know you have a node that is **SetActorHidden**, you can type that in the box, and the explorer will show you all the **SetActorHidden** nodes in your graph.
-
-
-
-
-
-[/EXCERPT:Main]
-
+[/COMMENT]
\ No newline at end of file
diff --git a/Engine/Documentation/Source/Engine/Blueprints/UserGuide/Debugging/KismetDebugging.INT.udn b/Engine/Documentation/Source/Engine/Blueprints/UserGuide/Debugging/KismetDebugging.INT.udn
--- a/Engine/Documentation/Source/Engine/Blueprints/UserGuide/Debugging/KismetDebugging.INT.udn
+++ b/Engine/Documentation/Source/Engine/Blueprints/UserGuide/Debugging/KismetDebugging.INT.udn
@@ -62,12 +62,9 @@ In addition, the viewport will display an overlay to alert the player that the g
 
 ## Debug Tab and Blueprint Debugger
 
-
 [EXCERPT:DebugTab]
 
-[EXCERPT:DebugPanel]
-The **Debug** tab shows information designated as important by the designer in the form of Breakpoints and watch values as well as a trace stack of all nodes belonging to the Blueprint that have been executed. This window also shows controls for controlling execution of the game when using Breakpoints.
-[/EXCERPT:DebugPanel]
+[INCLUDE:#DebugPanel]
 
 ### Watch Values
 
@@ -90,3 +87,9 @@ The **Execution Trace** stack shows a list of the nodes executed with the most r
 This list updates as you step through the graph when debugging.
 
 [/EXCERPT:DebugTab]
+
+[COMMENT:none]
+[EXCERPT:DebugPanel]
+The **Debug** tab shows information designated as important by the designer in the form of Breakpoints and watch values as well as a trace stack of all nodes belonging to the Blueprint that have been executed. This window also shows controls for controlling execution of the game when using Breakpoints.
+[/EXCERPT:DebugPanel]
+[/COMMENT]
\ No newline at end of file
diff --git a/Engine/Documentation/Source/Engine/Composure/HowTo/ComposureHowTo.INT.udn b/Engine/Documentation/Source/Engine/Composure/HowTo/ComposureHowTo.INT.udn
--- a/Engine/Documentation/Source/Engine/Composure/HowTo/ComposureHowTo.INT.udn
+++ b/Engine/Documentation/Source/Engine/Composure/HowTo/ComposureHowTo.INT.udn
@@ -14,8 +14,6 @@ ADD Overview Text
 The Overview text should describe the system, features, or concepts being covered in the How-Tos.
 [/REGION]
 
-[](FANCY:Engine\Composure\HowTo\FakeAo "%Engine\Composure\HowTo\FakeAo:description%")
-
 ADD Beginner Topic List
 [REGION:note]
 The Beginner Topic List should link to all beginner How-To guides using fancy links.
diff --git a/Engine/Documentation/Source/Engine/Landscape/Editing/Splines/Splines.INT.udn b/Engine/Documentation/Source/Engine/Landscape/Editing/Splines/Splines.INT.udn
--- a/Engine/Documentation/Source/Engine/Landscape/Editing/Splines/Splines.INT.udn
+++ b/Engine/Documentation/Source/Engine/Landscape/Editing/Splines/Splines.INT.udn
@@ -30,14 +30,18 @@ The **Spline** tool creates a spline Actor that deforms a terrain mesh along it,
 
 [EXCERPT:Splines001]
 
-[EXCERPT:LandscapeSplines]
-**Landscape Splines** are a flexible system for creating any linear feature that needs to conform to a Landscape, and can even push and pull the terrain to better facilitate building these features. They are created and edited using the spline tool in the Landscape tool.
-[/EXCERPT:LandscapeSplines]
+[INCLUDE:#LandscapeSplines]
 
 The Spline tool is best used to make streets or paths in the environment, but in general, it can be used for any mesh that must conform to the terrain, such as cobblestone walls or long straight rows of corn. 
 
 [/EXCERPT:Splines001]
 
+[COMMENT:none]
+[EXCERPT:LandscapeSplines]
+**Landscape Splines** are a flexible system for creating any linear feature that needs to conform to a Landscape, and can even push and pull the terrain to better facilitate building these features. They are created and edited using the spline tool in the Landscape tool.
+[/EXCERPT:LandscapeSplines]
+[/COMMENT]
+
 ## Creating a Spline
 
 **To create a spline:**
diff --git a/Engine/Documentation/Source/Engine/Landscape/Materials/Materials.INT.udn b/Engine/Documentation/Source/Engine/Landscape/Materials/Materials.INT.udn
--- a/Engine/Documentation/Source/Engine/Landscape/Materials/Materials.INT.udn
+++ b/Engine/Documentation/Source/Engine/Landscape/Materials/Materials.INT.udn
@@ -122,7 +122,7 @@ You can set the properties of the LandscapeLayerSwitch node in the **Details** p
 |2:| **Preview Used**   | If checked, uses a preview. |
 |3:| **LayerUsed**      | The result to use when the layer specified in the node's properties is in use by the current region of the Landscape.|
 |4:| **LayerNotUsed**   | The result to use when the layer is not used by the current region of the Landscape and has a weight of zero.|
-|5:| ** Unlabeled Output**        | Either the **LayerUsed** or **LayerNotUsed** inputs, depending on whether or not the layer contributes to the particular region of the Landscape.|  
+|5:| **Unlabeled Output**        | Either the **LayerUsed** or **LayerNotUsed** inputs, depending on whether or not the layer contributes to the particular region of the Landscape.|  
 
 [/EXCERPT:TerrainLayerSwitch]
 
@@ -295,14 +295,6 @@ You can use the same Material for your Landscape Material and your Landscape Hol
 
 	![](Landscape_BaseNodeOpacity.png)
 
-	[COMMENT:none]
-
-	[REGION:tip]
-	If you are using a separate Material for your Landscape hole you need to make sure that you also have a texture feeding into the Base Color slot or when you try and make parts of your Landscape invisible you will end up turning parts of your Landscape black.
-	[/REGION]
-	
-[/COMMENT]
-
 1. Next, search for and add a **LandscapeVisibilityMask** node to the Material Graph and connect its output to the **Opacity Mask** input.
 
 	![](Landscape_VisibilityMaskNode.png)
diff --git a/Engine/Documentation/Source/Engine/Matinee/HowTo/MHT_1/MHT_1.INT.udn b/Engine/Documentation/Source/Engine/Matinee/HowTo/MHT_1/MHT_1.INT.udn
--- a/Engine/Documentation/Source/Engine/Matinee/HowTo/MHT_1/MHT_1.INT.udn
+++ b/Engine/Documentation/Source/Engine/Matinee/HowTo/MHT_1/MHT_1.INT.udn
@@ -8,6 +8,7 @@ SkillLevel: Beginner
 
 In the following steps, we are going to create a simple door that opens/closes when entering/exiting a trigger volume. 
 
+[EXCERPT:ProjectStart]
 ## Creating the Project
 
 In this section, we will set up the project. 
diff --git a/Engine/Documentation/Source/Engine/MediaFramework/HowTo/MF_HowTo.INT.udn b/Engine/Documentation/Source/Engine/MediaFramework/HowTo/MF_HowTo.INT.udn
--- a/Engine/Documentation/Source/Engine/MediaFramework/HowTo/MF_HowTo.INT.udn
+++ b/Engine/Documentation/Source/Engine/MediaFramework/HowTo/MF_HowTo.INT.udn
@@ -15,31 +15,15 @@ If you are looking to quickly get a video asset up and running inside your level
 Refer to the table below for more information and click each link to view its contents. 
 
 ## Beginner
-[REGION:wrapper transparent]
-[](FANCY:Engine\MediaFramework\HowTo\FileMediaSource "%Engine\MediaFramework\HowTo\FileMediaSource:description%")
-[](FANCY:Engine\MediaFramework\HowTo\StreamMediaSource "%Engine\MediaFramework\HowTo\StreamMediaSource:description%")
-[](FANCY:Engine\MediaFramework\HowTo\ImgMediaSource "%Engine\MediaFramework\HowTo\ImgMediaSource:description%")
-[/REGION]
 
+[DIR(output:"fancy" type:"how-to" skill:"Beginner" parent:"Engine/MediaFramework/HowTo" org:"hierarchy")]
 
-## Intermediate
-[REGION:wrapper transparent]
-[](FANCY:Engine\MediaFramework\HowTo\MediaPlaylists "%Engine\MediaFramework\HowTo\MediaPlaylists:description%")
-[](FANCY:Engine\MediaFramework\HowTo\PlatformMedia "%Engine\MediaFramework\HowTo\PlatformMedia:description%")
-[](FANCY:Engine\MediaFramework\HowTo\ControllingPlayback "%Engine\MediaFramework\HowTo\ControllingPlayback:description%")
-[/REGION]
 
-## Advanced
-[REGION:wrapper transparent]
-[](FANCY:Engine\MediaFramework\HowTo\UsingWebCams "%Engine\MediaFramework\HowTo\UsingWebCams:description%")
-[/REGION]
+## Intermediate
 
+[DIR(output:"fancy" type:"how-to" skill:"Intermediate" parent:"Engine/MediaFramework/HowTo" org:"hierarchy")]
 
-[COMMENT:none]
- To Do
 ## Advanced
-[REGION:wrapper transparent]
-[/REGION]
 
-[/COMMENT]
+[DIR(output:"fancy" type:"how-to" skill:"Advanced" parent:"Engine/MediaFramework/HowTo" org:"hierarchy")]
 
diff --git a/Engine/Documentation/Source/Engine/Physics/Cloth/Reference/ClothReference.INT.udn b/Engine/Documentation/Source/Engine/Physics/Cloth/Reference/ClothReference.INT.udn
--- a/Engine/Documentation/Source/Engine/Physics/Cloth/Reference/ClothReference.INT.udn
+++ b/Engine/Documentation/Source/Engine/Physics/Cloth/Reference/ClothReference.INT.udn
@@ -41,7 +41,6 @@ mesh that you've selected, create a cloth asset for your render mesh and its LOD
 
 
 [COMMENT:none]
--
 
 [EXCERPT:createclothsection]
 ![](CreateClothSection.png)(w:350)
@@ -51,17 +50,6 @@ mesh that you've selected, create a cloth asset for your render mesh and its LOD
 | **Remove from Mesh** | Whether or not to leave this section behind (if driving a mesh with itself). Enable this is driving a high-poly mesh with a low-poly mesh. |
 | [REGION:tablesection] Collision [/REGION] ||
 | **Physics Asset** | Physics Asset to extract collisions from. Note that this will export Sphere and Sphylls, but does support up to 32 convex faces (or 5 boxes). |
-<!---| [REGION:tablesection] Auto Fix [/REGION] ||-
-[/COMMENT]
-[COMMENT:none]
--| **Try Auto Fix** |  |-
-[/COMMENT]
-[COMMENT:none]
--| **Auto Fix Threshold** |  |-
-[/COMMENT]
-[COMMENT:none]
--| **Simulated Particle Max Distance** |  |-
-[/COMMENT]
 [/EXCERPT:createclothsection]
 
 [EXCERPT:createclothsectionLOD]
@@ -82,8 +70,7 @@ mesh that you've selected, create a cloth asset for your render mesh and its LOD
 | **Add LOD** | Use the selected section to add as a new LOD. |
 [/EXCERPT:lodindex]
 
-
---->
+[/COMMENT]
 
 
 ## Cloth Paint Panel
@@ -116,7 +103,7 @@ The **Masks** window displays any created parameter sets for painted cloth value
 
 [REGION:raw]
 | Property | Description |
-| --- | --- | --- |
+| --- | --- |
 | **Name** | The name given to the mask along with its target setting for this parameter set. [INCLUDE:#masksettings]  |
 | **Add Mask** | Adds a new Mask to the available list of Mask parameters. |
 [/REGION]
@@ -124,7 +111,6 @@ The **Masks** window displays any created parameter sets for painted cloth value
 
 
 [COMMENT:none]
--
 [INCLUDE:#masksettings]
 
 [EXCERPT:masksettings]
@@ -141,8 +127,6 @@ The **Masks** window displays any created parameter sets for painted cloth value
 | **Delete** | Deletes the Mask from the list. |
 | **Apply** | Applies the Mask to the physical mesh. |
 [/EXCERPT:masksettings]
-
--
 [/COMMENT]
 
 
@@ -186,7 +170,6 @@ The **Config** properties enables you to adjust how your cloth react so that it
 
 
 [COMMENT:none]
--
 [INCLUDE:#constraintgconfig]
 
 [EXCERPT:constraintgconfig]
@@ -200,11 +183,8 @@ The **Config** properties enables you to adjust how your cloth react so that it
 [EXCERPT:windmethod]
 | --- | --- |
 | **Legacy** | Use the Legacy wind mode, where accelerations are modified directly by the simulation with no regard to drag or lift. |
-<!--- | **Accurate** |  | -
-[/COMMENT]
 [/EXCERPT:windmethod]
-
---->
+[/COMMENT]
 
 
 
diff --git a/Engine/Documentation/Source/Engine/Physics/Constraints/ConstraintsReference/Constraints.INT.udn b/Engine/Documentation/Source/Engine/Physics/Constraints/ConstraintsReference/Constraints.INT.udn
--- a/Engine/Documentation/Source/Engine/Physics/Constraints/ConstraintsReference/Constraints.INT.udn
+++ b/Engine/Documentation/Source/Engine/Physics/Constraints/ConstraintsReference/Constraints.INT.udn
@@ -229,7 +229,7 @@ Whether this angular motor uses SLERP (Spherical Linear Interpolation) or decomp
 
 | Property | Description |
 | -------- | ----------- |
-|[REGION:tablesection]Constraint[/REGION]|
+|[REGION:tablesection]Constraint[/REGION]||
 | **Component Name 1** |  |
 | **Component Name 2** |   |
 | **Joint Name** |  |
@@ -259,7 +259,7 @@ Whether this angular motor uses SLERP (Spherical Linear Interpolation) or decomp
 
 | Property | Description |
 | -------- | ----------- |
-|[REGION:tablesection]Linear Limits[/REGION]|
+|[REGION:tablesection]Linear Limits[/REGION]||
 | **Linear XMotion** |  |
 | **Linear YMotion** |  |
 | **Linear ZMotion** |  |
@@ -270,11 +270,11 @@ Whether this angular motor uses SLERP (Spherical Linear Interpolation) or decomp
 | **Linear Limit Damping** |  |
 | **Linear Breakable** |  |
 | **Linear Break Threshold** |  |
-|[REGION:tablesection]Projection[/REGION]|
+|[REGION:tablesection]Projection[/REGION]||
 | **Enable Projection** |  |
 | **Projection Linear Tolerance** |  |
 | **Projection Angular Tolerance** |  |
-|[REGION:tablesection]Angular Limits[/REGION]|
+|[REGION:tablesection]Angular Limits[/REGION]||
 | **Angular Swing 1Motion** |  |
 | **Angular Swing 2Motion** |  |
 | **Angular Twist Motion** |  |
@@ -290,7 +290,7 @@ Whether this angular motor uses SLERP (Spherical Linear Interpolation) or decomp
 | **Twist Limit Damping** |  |
 | **Angular Breakable** |  |
 | **Angular Break Threshold** |  |
-|[REGION:tablesection]Linear Motor[/REGION]|
+|[REGION:tablesection]Linear Motor[/REGION]||
 | **Linear Position Drive** |  |
 | **Linear XPosition Drive** |  |
 | **Linear YPosition Drive** |  |
@@ -301,17 +301,14 @@ Whether this angular motor uses SLERP (Spherical Linear Interpolation) or decomp
 | **Linear Velocity Target** |  |
 | **Linear Velocity Strength** | Damping to apply to the linear drive. |
 | **Max Linear Force** | Limit to the force the linear drive can apply. |
-|[REGION:tablesection]Angular Motor[/REGION]|
+|[REGION:tablesection]Angular Motor[/REGION]||
 | **Angular Orientation Drive** |  |
 | **Angular Orientation Target** |  |
 | **Angular Position Strength** | Spring value to apply for the angular drive. |
 | **Angular Velocity Drive** |  |
 | **Angular Velocity Target** |  |
 | **Angular Velocity Strength** | Damping value to apply for the angular drive. |
-| **Max Angular Force ** | Limit to the force the angular drive can apply. |
-
-
-<!--
+| **Max Angular Force** | Limit to the force the angular drive can apply. |
 
 [EXCERPT:linearMotionEnum]
 | **Type** | **Description** |
@@ -330,8 +327,6 @@ Whether this angular motor uses SLERP (Spherical Linear Interpolation) or decomp
 [/EXCERPT:angularMotionEnum]
 
 
-[/COMMENT]
-
 ## Further Information
 
 ### Joint Setup
@@ -360,10 +355,6 @@ Under most circumstances, the Linear Motions' freedoms will remain locked. Howev
 
 _From left to right (Physics Constraint scaled up 4x): Linear Motions limited to 15 Units, and Linear Motions set to free._
 
-[COMMENT:none]
- For information on all the properties relating to Linear Freedom and Limits see: [%Engine/Physics\Constraints\ConstraintsReferencetitle%](Engine/Physics\Constraints\ConstraintsReference 
-[/COMMENT]
-
 ### Angular Freedom and Limits
 
 ![](PhysCon_AngularLimitsProps.png)
@@ -384,10 +375,6 @@ Further, since the joint is located where the Physics Constraint Actor is locate
 [![](PhysCon_AngPos1.png)(h:312)](PhysCon_AngPos1.png) [![](PhysCon_AngPos2.png)(h:312)](PhysCon_AngPos2.png) 
 [/REGION]
 
-[COMMENT:none]
- For information on all the properties relating to Angular Freedom and Limits see: [%Engine/Physics\Constraints\ConstraintsReferencetitle%](Engine/Physics\Constraints\ConstraintsReference 
-[/COMMENT]
-
 ### Orientation
 
 The axes of a Physics Constraint Actor are in line with the rotation of the Physics Constraint Actor, and not aligned between the two objects. Given you desired effect, you will need to be aware of how the constraint is oriented.
@@ -404,9 +391,7 @@ In either setup, the Physics Constraint has its Linear XMotion set to free, and
 [![](PhysCon_SimFlat.png)(w:453)](PhysCon_SimFlat.png) [![](PhysCon_SimAngled.png)(w:453)](PhysCon_SimAngled.png)
 [/REGION]
 
-When using Angular limits, it is assumed that the Twist Axis (positive X) will be aligned pointing towards the blue constrained Actor (Constrained Actor 2). [COMMENT:none]
--This is really important because the -
-[/COMMENT] 
+When using Angular limits, it is assumed that the Twist Axis (positive X) will be aligned pointing towards the blue constrained Actor (Constrained Actor 2). 
 
 [REGION:lightbox]
 [![](PhysCon_Angular.png)(w:453)](PhysCon_Angular.png) [![](PhysCon_AngularSim.png)(w:453)](PhysCon_AngularSim.png)
@@ -416,32 +401,7 @@ Also, while simulating, you will notice that the cone that represents the angula
 
 ![](PhysCon_AngularCone_2.png)(h:300 convert:false) ![](PhysCon_Arc.png)(h:300 convert:false)
 
-[COMMENT:none]
-
-
-### Ball And Socket
-A ball-and-socket joint is simply a joint that has all its linear parts locked and all its angular parts free.
-
-
-### Hinge
-A hinge is like a ball-and-socket in that all its linear parts are locked, but its swing parts are also locked, so it can only rotate around its twist axis.
-
-
-### Prismatic
-A prismatic joint has 2 out of its 3 linear degrees of freedom locked, and also has all its angular degrees of freedom locked. This allows is to slide down a 'rail' but cannot rotate at all.
-
-
-### Joint Friction
-You can simulate joint friction using the drive settings. Turn on velocity drive with a target velocity of zero, and use the LinearDriveForceLimit and AngularDriveForceLimit settings to control how strongly the joint will resist motion.
-
-
-### Punch bag
-You could simulate something like a punch-bag (a ball on the top of a vertical pole that can be knocked around) with a joint by first creating a ball-and-socket joint. Then turn on swing position drive and adjust the **AngularDriveSpring** and **AngularDriveDamping** until it springs nicely around the upright position when pushed. You could turn on swing limits as well to ensure that it cannot bend over too far.
-
-
 [/COMMENT]
--->
-
 
 
 
diff --git a/Engine/Documentation/Source/Engine/Physics/PhAT/HowTo/ApplyConstraintProfile/ApplyConstraintProfile.INT.udn b/Engine/Documentation/Source/Engine/Physics/PhAT/HowTo/ApplyConstraintProfile/ApplyConstraintProfile.INT.udn
--- a/Engine/Documentation/Source/Engine/Physics/PhAT/HowTo/ApplyConstraintProfile/ApplyConstraintProfile.INT.udn
+++ b/Engine/Documentation/Source/Engine/Physics/PhAT/HowTo/ApplyConstraintProfile/ApplyConstraintProfile.INT.udn
@@ -87,7 +87,60 @@ Below, you will find the steps for creating a simple graph to enable a **Constra
 			![](fullGraph.png)
 		[/PARAM]
 		[PARAM:snippet]
-			%BPSnip_001%
+				Begin Object Class=K2Node_Event Name="K2Node_Event_8"
+				   EventReference=(MemberParent=Class'/Script/Engine.Actor',MemberName="ReceiveBeginPlay")
+				   bOverrideFunction=True
+				   NodePosX=32
+				   NodePosY=-16
+				   NodeGuid=48E7E2C9458A1028A689B3A249855163
+				   CustomProperties Pin (PinId=D1792C7A4CBF979351A63AB775A61034,PinName="OutputDelegate",Direction="EGPD_Output",PinType.PinCategory="delegate",PinType.PinSubCategoryObject=None,PinType.PinSubCategoryMemberReference=(MemberParent=Class'/Script/Engine.Actor',MemberName="ReceiveBeginPlay"),PinType.bIsArray=False,PinType.bIsReference=False,PinType.bIsConst=False,PinType.bIsWeakPointer=False,PersistentGuid=00000000000000000000000000000000,bHidden=False,bNotConnectable=False,bDefaultValueIsReadOnly=False,bDefaultValueIsIgnored=False,bAdvancedView=False,)
+				   CustomProperties Pin (PinId=3880C4144B19D9B0F602638392FAD6FA,PinName="then",Direction="EGPD_Output",PinType.PinCategory="exec",PinType.PinSubCategoryObject=None,PinType.PinSubCategoryMemberReference=(),PinType.bIsArray=False,PinType.bIsReference=False,PinType.bIsConst=False,PinType.bIsWeakPointer=False,LinkedTo=(K2Node_CallFunction_132 7ED8710C4A46FF83B175828EFBF77DAE,),PersistentGuid=00000000000000000000000000000000,bHidden=False,bNotConnectable=False,bDefaultValueIsReadOnly=False,bDefaultValueIsIgnored=False,bAdvancedView=False,)
+				End Object
+				Begin Object Class=K2Node_VariableGet Name="K2Node_VariableGet_155"
+				   VariableReference=(MemberName="SkeletalMeshComponent",bSelfContext=True)
+				   NodePosX=32
+				   NodePosY=112
+				   NodeGuid=42ADF0644E0964ECBC298DB6863ADA89
+				   CustomProperties Pin (PinId=788C7F5B4B4B5A7362C8AEA404935EEB,PinName="SkeletalMeshComponent",Direction="EGPD_Output",PinType.PinCategory="object",PinType.PinSubCategoryObject=Class'/Script/Engine.SkeletalMeshComponent',PinType.PinSubCategoryMemberReference=(),PinType.bIsArray=False,PinType.bIsReference=False,PinType.bIsConst=False,PinType.bIsWeakPointer=False,LinkedTo=(K2Node_CallFunction_132 E1392084422728C657D6C3A54A4116F2,),PersistentGuid=00000000000000000000000000000000,bHidden=False,bNotConnectable=False,bDefaultValueIsReadOnly=False,bDefaultValueIsIgnored=False,bAdvancedView=False,)
+				   CustomProperties Pin (PinId=D7B4F9F54069CB07496DEE843FDBBA6B,PinName="self",PinFriendlyName=NSLOCTEXT("K2Node", "Target", "Target"),PinType.PinCategory="object",PinType.PinSubCategoryObject=BlueprintGeneratedClass'/Game/PhysAnimPawn.PhysAnimPawn_C',PinType.PinSubCategoryMemberReference=(),PinType.bIsArray=False,PinType.bIsReference=False,PinType.bIsConst=False,PinType.bIsWeakPointer=False,PersistentGuid=00000000000000000000000000000000,bHidden=True,bNotConnectable=False,bDefaultValueIsReadOnly=False,bDefaultValueIsIgnored=False,bAdvancedView=False,)
+				End Object
+				Begin Object Class=K2Node_CallFunction Name="K2Node_CallFunction_132"
+				   FunctionReference=(MemberParent=Class'/Script/Engine.SkeletalMeshComponent',MemberName="SetConstraintProfileForAll")
+				   NodePosX=320
+				   NodeGuid=19366CE04ECF7D2D1B3F3FAB826015E3
+				   CustomProperties Pin (PinId=7ED8710C4A46FF83B175828EFBF77DAE,PinName="execute",PinToolTip="\nExec",PinType.PinCategory="exec",PinType.PinSubCategoryObject=None,PinType.PinSubCategoryMemberReference=(),PinType.bIsArray=False,PinType.bIsReference=False,PinType.bIsConst=False,PinType.bIsWeakPointer=False,LinkedTo=(K2Node_Event_8 3880C4144B19D9B0F602638392FAD6FA,),PersistentGuid=00000000000000000000000000000000,bHidden=False,bNotConnectable=False,bDefaultValueIsReadOnly=False,bDefaultValueIsIgnored=False,bAdvancedView=False,)
+				   CustomProperties Pin (PinId=94ECC7884A2DC860DBB45E8573B644E7,PinName="then",PinToolTip="\nExec",Direction="EGPD_Output",PinType.PinCategory="exec",PinType.PinSubCategoryObject=None,PinType.PinSubCategoryMemberReference=(),PinType.bIsArray=False,PinType.bIsReference=False,PinType.bIsConst=False,PinType.bIsWeakPointer=False,LinkedTo=(K2Node_CallFunction_8 399238E643A83FB3CFC85CB926B5CF9F,),PersistentGuid=00000000000000000000000000000000,bHidden=False,bNotConnectable=False,bDefaultValueIsReadOnly=False,bDefaultValueIsIgnored=False,bAdvancedView=False,)
+				   CustomProperties Pin (PinId=E1392084422728C657D6C3A54A4116F2,PinName="self",PinFriendlyName=NSLOCTEXT("K2Node", "Target", "Target"),PinToolTip="Target\nSkeletal Mesh Component Reference",PinType.PinCategory="object",PinType.PinSubCategoryObject=Class'/Script/Engine.SkeletalMeshComponent',PinType.PinSubCategoryMemberReference=(),PinType.bIsArray=False,PinType.bIsReference=False,PinType.bIsConst=False,PinType.bIsWeakPointer=False,LinkedTo=(K2Node_VariableGet_155 788C7F5B4B4B5A7362C8AEA404935EEB,),PersistentGuid=00000000000000000000000000000000,bHidden=False,bNotConnectable=False,bDefaultValueIsReadOnly=False,bDefaultValueIsIgnored=False,bAdvancedView=False,)
+				   CustomProperties Pin (PinId=28663C674854E90FE78C718CD73FB3DF,PinName="ProfileName",PinToolTip="Profile Name\nName",PinType.PinCategory="name",PinType.PinSubCategoryObject=None,PinType.PinSubCategoryMemberReference=(),PinType.bIsArray=False,PinType.bIsReference=False,PinType.bIsConst=False,PinType.bIsWeakPointer=False,DefaultValue="DocsConstraintProfile",AutogeneratedDefaultValue="None",PersistentGuid=00000000000000000000000000000000,bHidden=False,bNotConnectable=False,bDefaultValueIsReadOnly=False,bDefaultValueIsIgnored=False,bAdvancedView=False,)
+				   CustomProperties Pin (PinId=D960208D49BE62D6FDE20698BB7D27CC,PinName="bDefaultIfNotFound",PinToolTip="Default if Not Found\nBoolean",PinType.PinCategory="bool",PinType.PinSubCategoryObject=None,PinType.PinSubCategoryMemberReference=(),PinType.bIsArray=False,PinType.bIsReference=False,PinType.bIsConst=False,PinType.bIsWeakPointer=False,DefaultValue="true",AutogeneratedDefaultValue="false",PersistentGuid=00000000000000000000000000000000,bHidden=False,bNotConnectable=False,bDefaultValueIsReadOnly=False,bDefaultValueIsIgnored=False,bAdvancedView=False,)
+				End Object
+				Begin Object Class=K2Node_VariableGet Name="K2Node_VariableGet_1"
+				   VariableReference=(MemberName="SkeletalMeshComponent",bSelfContext=True)
+				   NodePosX=624
+				   NodePosY=80
+				   NodeGuid=211663904BCAEE00C6164CBF0871A8B9
+				   CustomProperties Pin (PinId=1EEABDD6494CBF38903C6691F5F5B45C,PinName="SkeletalMeshComponent",Direction="EGPD_Output",PinType.PinCategory="object",PinType.PinSubCategoryObject=Class'/Script/Engine.SkeletalMeshComponent',PinType.PinSubCategoryMemberReference=(),PinType.bIsArray=False,PinType.bIsReference=False,PinType.bIsConst=False,PinType.bIsWeakPointer=False,LinkedTo=(K2Node_CallFunction_8 4DD6FE2846E695E0F3F569B06CED0A46,),PersistentGuid=00000000000000000000000000000000,bHidden=False,bNotConnectable=False,bDefaultValueIsReadOnly=False,bDefaultValueIsIgnored=False,bAdvancedView=False,)
+				   CustomProperties Pin (PinId=D68B011B4E4D641D541E8CA766972993,PinName="self",PinFriendlyName=NSLOCTEXT("K2Node", "Target", "Target"),PinType.PinCategory="object",PinType.PinSubCategoryObject=BlueprintGeneratedClass'/Game/PhysAnimPawn.PhysAnimPawn_C',PinType.PinSubCategoryMemberReference=(),PinType.bIsArray=False,PinType.bIsReference=False,PinType.bIsConst=False,PinType.bIsWeakPointer=False,PersistentGuid=00000000000000000000000000000000,bHidden=True,bNotConnectable=False,bDefaultValueIsReadOnly=False,bDefaultValueIsIgnored=False,bAdvancedView=False,)
+				End Object
+				Begin Object Class=K2Node_CallFunction Name="K2Node_CallFunction_8"
+				   FunctionReference=(MemberParent=Class'/Script/Engine.SkeletalMeshComponent',MemberName="SetAllBodiesBelowSimulatePhysics")
+				   NodePosX=848
+				   NodeGuid=3116B9174BA9B4F48CE4CAA1E47000EC
+				   CustomProperties Pin (PinId=399238E643A83FB3CFC85CB926B5CF9F,PinName="execute",PinToolTip="\nExec",PinType.PinCategory="exec",PinType.PinSubCategoryObject=None,PinType.PinSubCategoryMemberReference=(),PinType.bIsArray=False,PinType.bIsReference=False,PinType.bIsConst=False,PinType.bIsWeakPointer=False,LinkedTo=(K2Node_CallFunction_132 94ECC7884A2DC860DBB45E8573B644E7,),PersistentGuid=00000000000000000000000000000000,bHidden=False,bNotConnectable=False,bDefaultValueIsReadOnly=False,bDefaultValueIsIgnored=False,bAdvancedView=False,)
+				   CustomProperties Pin (PinId=C132D8014BC6A4DE84C437AF07CF838E,PinName="then",PinToolTip="\nExec",Direction="EGPD_Output",PinType.PinCategory="exec",PinType.PinSubCategoryObject=None,PinType.PinSubCategoryMemberReference=(),PinType.bIsArray=False,PinType.bIsReference=False,PinType.bIsConst=False,PinType.bIsWeakPointer=False,PersistentGuid=00000000000000000000000000000000,bHidden=False,bNotConnectable=False,bDefaultValueIsReadOnly=False,bDefaultValueIsIgnored=False,bAdvancedView=False,)
+				   CustomProperties Pin (PinId=4DD6FE2846E695E0F3F569B06CED0A46,PinName="self",PinFriendlyName=NSLOCTEXT("K2Node", "Target", "Target"),PinToolTip="Target\nSkeletal Mesh Component Reference",PinType.PinCategory="object",PinType.PinSubCategoryObject=Class'/Script/Engine.SkeletalMeshComponent',PinType.PinSubCategoryMemberReference=(),PinType.bIsArray=False,PinType.bIsReference=False,PinType.bIsConst=False,PinType.bIsWeakPointer=False,LinkedTo=(K2Node_VariableGet_1 1EEABDD6494CBF38903C6691F5F5B45C,),PersistentGuid=00000000000000000000000000000000,bHidden=False,bNotConnectable=False,bDefaultValueIsReadOnly=False,bDefaultValueIsIgnored=False,bAdvancedView=False,)
+				   CustomProperties Pin (PinId=7D29E35244DEB831C54B1CAEF4A5FE5C,PinName="InBoneName",PinToolTip="In Bone Name\nName (by ref)",PinType.PinCategory="name",PinType.PinSubCategoryObject=None,PinType.PinSubCategoryMemberReference=(),PinType.bIsArray=False,PinType.bIsReference=True,PinType.bIsConst=True,PinType.bIsWeakPointer=False,DefaultValue="None",AutogeneratedDefaultValue="None",LinkedTo=(K2Node_VariableGet_181 4DF93E5141ACE87768E5D8A7543C7B07,),PersistentGuid=00000000000000000000000000000000,bHidden=False,bNotConnectable=False,bDefaultValueIsReadOnly=False,bDefaultValueIsIgnored=True,bAdvancedView=False,)
+				   CustomProperties Pin (PinId=AE4340174DD534FEA5B812AE6DAEA1EB,PinName="bNewSimulate",PinToolTip="New Simulate\nBoolean",PinType.PinCategory="bool",PinType.PinSubCategoryObject=None,PinType.PinSubCategoryMemberReference=(),PinType.bIsArray=False,PinType.bIsReference=False,PinType.bIsConst=False,PinType.bIsWeakPointer=False,DefaultValue="true",AutogeneratedDefaultValue="false",PersistentGuid=00000000000000000000000000000000,bHidden=False,bNotConnectable=False,bDefaultValueIsReadOnly=False,bDefaultValueIsIgnored=False,bAdvancedView=False,)
+				   CustomProperties Pin (PinId=80EC7602465D797C7F0D4D93D6BA9BD5,PinName="bIncludeSelf",PinToolTip="Include Self\nBoolean",PinType.PinCategory="bool",PinType.PinSubCategoryObject=None,PinType.PinSubCategoryMemberReference=(),PinType.bIsArray=False,PinType.bIsReference=False,PinType.bIsConst=False,PinType.bIsWeakPointer=False,DefaultValue="true",AutogeneratedDefaultValue="true",PersistentGuid=00000000000000000000000000000000,bHidden=False,bNotConnectable=False,bDefaultValueIsReadOnly=False,bDefaultValueIsIgnored=False,bAdvancedView=False,)
+				End Object
+				Begin Object Class=K2Node_VariableGet Name="K2Node_VariableGet_181"
+				   VariableReference=(MemberName="BoneName",MemberGuid=A5E270C141CEE32A0269309215CEBAC4,bSelfContext=True)
+				   NodePosX=640
+				   NodePosY=128
+				   NodeGuid=B116D28D4BF4E6625FF02F92D2EA8DFD
+				   CustomProperties Pin (PinId=4DF93E5141ACE87768E5D8A7543C7B07,PinName="BoneName",Direction="EGPD_Output",PinType.PinCategory="name",PinType.PinSubCategoryObject=None,PinType.PinSubCategoryMemberReference=(),PinType.bIsArray=False,PinType.bIsReference=False,PinType.bIsConst=False,PinType.bIsWeakPointer=False,AutogeneratedDefaultValue="None",LinkedTo=(K2Node_CallFunction_8 7D29E35244DEB831C54B1CAEF4A5FE5C,),PersistentGuid=00000000000000000000000000000000,bHidden=False,bNotConnectable=False,bDefaultValueIsReadOnly=False,bDefaultValueIsIgnored=False,bAdvancedView=False,)
+				   CustomProperties Pin (PinId=F259489241A7853EE35268954C561767,PinName="self",PinFriendlyName=NSLOCTEXT("K2Node", "Target", "Target"),PinType.PinCategory="object",PinType.PinSubCategoryObject=BlueprintGeneratedClass'/Game/PhysAnimPawn.PhysAnimPawn_C',PinType.PinSubCategoryMemberReference=(),PinType.bIsArray=False,PinType.bIsReference=False,PinType.bIsConst=False,PinType.bIsWeakPointer=False,PersistentGuid=00000000000000000000000000000000,bHidden=True,bNotConnectable=False,bDefaultValueIsReadOnly=False,bDefaultValueIsIgnored=False,bAdvancedView=False,)
+				End Object
 		[/PARAM]
 	[/OBJECT]
 
@@ -113,67 +166,6 @@ This effect was achieved by setting the **Linear Limits** on the upper arm const
 
 ![](achievedBy.png)
 
-[COMMENT:none]
- Variables 
-[/COMMENT]
-
-[VAR:BPSnip_001]
-	Begin Object Class=K2Node_Event Name="K2Node_Event_8"
-	   EventReference=(MemberParent=Class'/Script/Engine.Actor',MemberName="ReceiveBeginPlay")
-	   bOverrideFunction=True
-	   NodePosX=32
-	   NodePosY=-16
-	   NodeGuid=48E7E2C9458A1028A689B3A249855163
-	   CustomProperties Pin (PinId=D1792C7A4CBF979351A63AB775A61034,PinName="OutputDelegate",Direction="EGPD_Output",PinType.PinCategory="delegate",PinType.PinSubCategoryObject=None,PinType.PinSubCategoryMemberReference=(MemberParent=Class'/Script/Engine.Actor',MemberName="ReceiveBeginPlay"),PinType.bIsArray=False,PinType.bIsReference=False,PinType.bIsConst=False,PinType.bIsWeakPointer=False,PersistentGuid=00000000000000000000000000000000,bHidden=False,bNotConnectable=False,bDefaultValueIsReadOnly=False,bDefaultValueIsIgnored=False,bAdvancedView=False,)
-	   CustomProperties Pin (PinId=3880C4144B19D9B0F602638392FAD6FA,PinName="then",Direction="EGPD_Output",PinType.PinCategory="exec",PinType.PinSubCategoryObject=None,PinType.PinSubCategoryMemberReference=(),PinType.bIsArray=False,PinType.bIsReference=False,PinType.bIsConst=False,PinType.bIsWeakPointer=False,LinkedTo=(K2Node_CallFunction_132 7ED8710C4A46FF83B175828EFBF77DAE,),PersistentGuid=00000000000000000000000000000000,bHidden=False,bNotConnectable=False,bDefaultValueIsReadOnly=False,bDefaultValueIsIgnored=False,bAdvancedView=False,)
-	End Object
-	Begin Object Class=K2Node_VariableGet Name="K2Node_VariableGet_155"
-	   VariableReference=(MemberName="SkeletalMeshComponent",bSelfContext=True)
-	   NodePosX=32
-	   NodePosY=112
-	   NodeGuid=42ADF0644E0964ECBC298DB6863ADA89
-	   CustomProperties Pin (PinId=788C7F5B4B4B5A7362C8AEA404935EEB,PinName="SkeletalMeshComponent",Direction="EGPD_Output",PinType.PinCategory="object",PinType.PinSubCategoryObject=Class'/Script/Engine.SkeletalMeshComponent',PinType.PinSubCategoryMemberReference=(),PinType.bIsArray=False,PinType.bIsReference=False,PinType.bIsConst=False,PinType.bIsWeakPointer=False,LinkedTo=(K2Node_CallFunction_132 E1392084422728C657D6C3A54A4116F2,),PersistentGuid=00000000000000000000000000000000,bHidden=False,bNotConnectable=False,bDefaultValueIsReadOnly=False,bDefaultValueIsIgnored=False,bAdvancedView=False,)
-	   CustomProperties Pin (PinId=D7B4F9F54069CB07496DEE843FDBBA6B,PinName="self",PinFriendlyName=NSLOCTEXT("K2Node", "Target", "Target"),PinType.PinCategory="object",PinType.PinSubCategoryObject=BlueprintGeneratedClass'/Game/PhysAnimPawn.PhysAnimPawn_C',PinType.PinSubCategoryMemberReference=(),PinType.bIsArray=False,PinType.bIsReference=False,PinType.bIsConst=False,PinType.bIsWeakPointer=False,PersistentGuid=00000000000000000000000000000000,bHidden=True,bNotConnectable=False,bDefaultValueIsReadOnly=False,bDefaultValueIsIgnored=False,bAdvancedView=False,)
-	End Object
-	Begin Object Class=K2Node_CallFunction Name="K2Node_CallFunction_132"
-	   FunctionReference=(MemberParent=Class'/Script/Engine.SkeletalMeshComponent',MemberName="SetConstraintProfileForAll")
-	   NodePosX=320
-	   NodeGuid=19366CE04ECF7D2D1B3F3FAB826015E3
-	   CustomProperties Pin (PinId=7ED8710C4A46FF83B175828EFBF77DAE,PinName="execute",PinToolTip="\nExec",PinType.PinCategory="exec",PinType.PinSubCategoryObject=None,PinType.PinSubCategoryMemberReference=(),PinType.bIsArray=False,PinType.bIsReference=False,PinType.bIsConst=False,PinType.bIsWeakPointer=False,LinkedTo=(K2Node_Event_8 3880C4144B19D9B0F602638392FAD6FA,),PersistentGuid=00000000000000000000000000000000,bHidden=False,bNotConnectable=False,bDefaultValueIsReadOnly=False,bDefaultValueIsIgnored=False,bAdvancedView=False,)
-	   CustomProperties Pin (PinId=94ECC7884A2DC860DBB45E8573B644E7,PinName="then",PinToolTip="\nExec",Direction="EGPD_Output",PinType.PinCategory="exec",PinType.PinSubCategoryObject=None,PinType.PinSubCategoryMemberReference=(),PinType.bIsArray=False,PinType.bIsReference=False,PinType.bIsConst=False,PinType.bIsWeakPointer=False,LinkedTo=(K2Node_CallFunction_8 399238E643A83FB3CFC85CB926B5CF9F,),PersistentGuid=00000000000000000000000000000000,bHidden=False,bNotConnectable=False,bDefaultValueIsReadOnly=False,bDefaultValueIsIgnored=False,bAdvancedView=False,)
-	   CustomProperties Pin (PinId=E1392084422728C657D6C3A54A4116F2,PinName="self",PinFriendlyName=NSLOCTEXT("K2Node", "Target", "Target"),PinToolTip="Target\nSkeletal Mesh Component Reference",PinType.PinCategory="object",PinType.PinSubCategoryObject=Class'/Script/Engine.SkeletalMeshComponent',PinType.PinSubCategoryMemberReference=(),PinType.bIsArray=False,PinType.bIsReference=False,PinType.bIsConst=False,PinType.bIsWeakPointer=False,LinkedTo=(K2Node_VariableGet_155 788C7F5B4B4B5A7362C8AEA404935EEB,),PersistentGuid=00000000000000000000000000000000,bHidden=False,bNotConnectable=False,bDefaultValueIsReadOnly=False,bDefaultValueIsIgnored=False,bAdvancedView=False,)
-	   CustomProperties Pin (PinId=28663C674854E90FE78C718CD73FB3DF,PinName="ProfileName",PinToolTip="Profile Name\nName",PinType.PinCategory="name",PinType.PinSubCategoryObject=None,PinType.PinSubCategoryMemberReference=(),PinType.bIsArray=False,PinType.bIsReference=False,PinType.bIsConst=False,PinType.bIsWeakPointer=False,DefaultValue="DocsConstraintProfile",AutogeneratedDefaultValue="None",PersistentGuid=00000000000000000000000000000000,bHidden=False,bNotConnectable=False,bDefaultValueIsReadOnly=False,bDefaultValueIsIgnored=False,bAdvancedView=False,)
-	   CustomProperties Pin (PinId=D960208D49BE62D6FDE20698BB7D27CC,PinName="bDefaultIfNotFound",PinToolTip="Default if Not Found\nBoolean",PinType.PinCategory="bool",PinType.PinSubCategoryObject=None,PinType.PinSubCategoryMemberReference=(),PinType.bIsArray=False,PinType.bIsReference=False,PinType.bIsConst=False,PinType.bIsWeakPointer=False,DefaultValue="true",AutogeneratedDefaultValue="false",PersistentGuid=00000000000000000000000000000000,bHidden=False,bNotConnectable=False,bDefaultValueIsReadOnly=False,bDefaultValueIsIgnored=False,bAdvancedView=False,)
-	End Object
-	Begin Object Class=K2Node_VariableGet Name="K2Node_VariableGet_1"
-	   VariableReference=(MemberName="SkeletalMeshComponent",bSelfContext=True)
-	   NodePosX=624
-	   NodePosY=80
-	   NodeGuid=211663904BCAEE00C6164CBF0871A8B9
-	   CustomProperties Pin (PinId=1EEABDD6494CBF38903C6691F5F5B45C,PinName="SkeletalMeshComponent",Direction="EGPD_Output",PinType.PinCategory="object",PinType.PinSubCategoryObject=Class'/Script/Engine.SkeletalMeshComponent',PinType.PinSubCategoryMemberReference=(),PinType.bIsArray=False,PinType.bIsReference=False,PinType.bIsConst=False,PinType.bIsWeakPointer=False,LinkedTo=(K2Node_CallFunction_8 4DD6FE2846E695E0F3F569B06CED0A46,),PersistentGuid=00000000000000000000000000000000,bHidden=False,bNotConnectable=False,bDefaultValueIsReadOnly=False,bDefaultValueIsIgnored=False,bAdvancedView=False,)
-	   CustomProperties Pin (PinId=D68B011B4E4D641D541E8CA766972993,PinName="self",PinFriendlyName=NSLOCTEXT("K2Node", "Target", "Target"),PinType.PinCategory="object",PinType.PinSubCategoryObject=BlueprintGeneratedClass'/Game/PhysAnimPawn.PhysAnimPawn_C',PinType.PinSubCategoryMemberReference=(),PinType.bIsArray=False,PinType.bIsReference=False,PinType.bIsConst=False,PinType.bIsWeakPointer=False,PersistentGuid=00000000000000000000000000000000,bHidden=True,bNotConnectable=False,bDefaultValueIsReadOnly=False,bDefaultValueIsIgnored=False,bAdvancedView=False,)
-	End Object
-	Begin Object Class=K2Node_CallFunction Name="K2Node_CallFunction_8"
-	   FunctionReference=(MemberParent=Class'/Script/Engine.SkeletalMeshComponent',MemberName="SetAllBodiesBelowSimulatePhysics")
-	   NodePosX=848
-	   NodeGuid=3116B9174BA9B4F48CE4CAA1E47000EC
-	   CustomProperties Pin (PinId=399238E643A83FB3CFC85CB926B5CF9F,PinName="execute",PinToolTip="\nExec",PinType.PinCategory="exec",PinType.PinSubCategoryObject=None,PinType.PinSubCategoryMemberReference=(),PinType.bIsArray=False,PinType.bIsReference=False,PinType.bIsConst=False,PinType.bIsWeakPointer=False,LinkedTo=(K2Node_CallFunction_132 94ECC7884A2DC860DBB45E8573B644E7,),PersistentGuid=00000000000000000000000000000000,bHidden=False,bNotConnectable=False,bDefaultValueIsReadOnly=False,bDefaultValueIsIgnored=False,bAdvancedView=False,)
-	   CustomProperties Pin (PinId=C132D8014BC6A4DE84C437AF07CF838E,PinName="then",PinToolTip="\nExec",Direction="EGPD_Output",PinType.PinCategory="exec",PinType.PinSubCategoryObject=None,PinType.PinSubCategoryMemberReference=(),PinType.bIsArray=False,PinType.bIsReference=False,PinType.bIsConst=False,PinType.bIsWeakPointer=False,PersistentGuid=00000000000000000000000000000000,bHidden=False,bNotConnectable=False,bDefaultValueIsReadOnly=False,bDefaultValueIsIgnored=False,bAdvancedView=False,)
-	   CustomProperties Pin (PinId=4DD6FE2846E695E0F3F569B06CED0A46,PinName="self",PinFriendlyName=NSLOCTEXT("K2Node", "Target", "Target"),PinToolTip="Target\nSkeletal Mesh Component Reference",PinType.PinCategory="object",PinType.PinSubCategoryObject=Class'/Script/Engine.SkeletalMeshComponent',PinType.PinSubCategoryMemberReference=(),PinType.bIsArray=False,PinType.bIsReference=False,PinType.bIsConst=False,PinType.bIsWeakPointer=False,LinkedTo=(K2Node_VariableGet_1 1EEABDD6494CBF38903C6691F5F5B45C,),PersistentGuid=00000000000000000000000000000000,bHidden=False,bNotConnectable=False,bDefaultValueIsReadOnly=False,bDefaultValueIsIgnored=False,bAdvancedView=False,)
-	   CustomProperties Pin (PinId=7D29E35244DEB831C54B1CAEF4A5FE5C,PinName="InBoneName",PinToolTip="In Bone Name\nName (by ref)",PinType.PinCategory="name",PinType.PinSubCategoryObject=None,PinType.PinSubCategoryMemberReference=(),PinType.bIsArray=False,PinType.bIsReference=True,PinType.bIsConst=True,PinType.bIsWeakPointer=False,DefaultValue="None",AutogeneratedDefaultValue="None",LinkedTo=(K2Node_VariableGet_181 4DF93E5141ACE87768E5D8A7543C7B07,),PersistentGuid=00000000000000000000000000000000,bHidden=False,bNotConnectable=False,bDefaultValueIsReadOnly=False,bDefaultValueIsIgnored=True,bAdvancedView=False,)
-	   CustomProperties Pin (PinId=AE4340174DD534FEA5B812AE6DAEA1EB,PinName="bNewSimulate",PinToolTip="New Simulate\nBoolean",PinType.PinCategory="bool",PinType.PinSubCategoryObject=None,PinType.PinSubCategoryMemberReference=(),PinType.bIsArray=False,PinType.bIsReference=False,PinType.bIsConst=False,PinType.bIsWeakPointer=False,DefaultValue="true",AutogeneratedDefaultValue="false",PersistentGuid=00000000000000000000000000000000,bHidden=False,bNotConnectable=False,bDefaultValueIsReadOnly=False,bDefaultValueIsIgnored=False,bAdvancedView=False,)
-	   CustomProperties Pin (PinId=80EC7602465D797C7F0D4D93D6BA9BD5,PinName="bIncludeSelf",PinToolTip="Include Self\nBoolean",PinType.PinCategory="bool",PinType.PinSubCategoryObject=None,PinType.PinSubCategoryMemberReference=(),PinType.bIsArray=False,PinType.bIsReference=False,PinType.bIsConst=False,PinType.bIsWeakPointer=False,DefaultValue="true",AutogeneratedDefaultValue="true",PersistentGuid=00000000000000000000000000000000,bHidden=False,bNotConnectable=False,bDefaultValueIsReadOnly=False,bDefaultValueIsIgnored=False,bAdvancedView=False,)
-	End Object
-	Begin Object Class=K2Node_VariableGet Name="K2Node_VariableGet_181"
-	   VariableReference=(MemberName="BoneName",MemberGuid=A5E270C141CEE32A0269309215CEBAC4,bSelfContext=True)
-	   NodePosX=640
-	   NodePosY=128
-	   NodeGuid=B116D28D4BF4E6625FF02F92D2EA8DFD
-	   CustomProperties Pin (PinId=4DF93E5141ACE87768E5D8A7543C7B07,PinName="BoneName",Direction="EGPD_Output",PinType.PinCategory="name",PinType.PinSubCategoryObject=None,PinType.PinSubCategoryMemberReference=(),PinType.bIsArray=False,PinType.bIsReference=False,PinType.bIsConst=False,PinType.bIsWeakPointer=False,AutogeneratedDefaultValue="None",LinkedTo=(K2Node_CallFunction_8 7D29E35244DEB831C54B1CAEF4A5FE5C,),PersistentGuid=00000000000000000000000000000000,bHidden=False,bNotConnectable=False,bDefaultValueIsReadOnly=False,bDefaultValueIsIgnored=False,bAdvancedView=False,)
-	   CustomProperties Pin (PinId=F259489241A7853EE35268954C561767,PinName="self",PinFriendlyName=NSLOCTEXT("K2Node", "Target", "Target"),PinType.PinCategory="object",PinType.PinSubCategoryObject=BlueprintGeneratedClass'/Game/PhysAnimPawn.PhysAnimPawn_C',PinType.PinSubCategoryMemberReference=(),PinType.bIsArray=False,PinType.bIsReference=False,PinType.bIsConst=False,PinType.bIsWeakPointer=False,PersistentGuid=00000000000000000000000000000000,bHidden=True,bNotConnectable=False,bDefaultValueIsReadOnly=False,bDefaultValueIsIgnored=False,bAdvancedView=False,)
-	End Object
-[/VAR]
-
 
 ## Additional Resources
 
