Markdown conversion of files for UE releases 4.20 and 4.21.

diff --git a/Engine/Documentation/Source/Engine/ProVideoIO/AjaIOReference/AjaIOReference.INT.udn b/Engine/Documentation/Source/Engine/ProVideoIO/AjaIOReference/AjaIOReference.INT.udn
new file mode 100644
--- /dev/null
+++ b/Engine/Documentation/Source/Engine/ProVideoIO/AjaIOReference/AjaIOReference.INT.udn
@@ -0,0 +1,199 @@
+Availability: Public
+Crumbs: 
+Title: AJA Media Reference
+Description: Describes the options and settings exposed by the AJA Media Framework components.
+Type: 
+Version: 4.20
+Parent: Engine/ProVideoIO
+Order: 
+Tags: Reference
+Tags: Early Access
+Tags: professional video
+
+
+This page describes the options and settings exposed on AJA Media Framework objects.
+
+## Supported AJA Cards and Software
+
+The AJA Media Source and AJA Media Output have been tested with the following cards, using **Version 14.X** of the AJA Desktop Software:
+
+*   **Kona 4** (in 4K mode, not UFC mode)
+*   **Corvid 44**
+
+Other devices may or may not work as expected.
+
+**Version 15** of the AJA Desktop Software is not yet supported.
+
+## AJA Media Source Settings
+
+Each AJA Media Source object that you create exposes the following configuration settings.
+
+  
+| **Property** | **Description** |
+| --- | --- |
+| **Configuration Options**| | 
+| [INCLUDE:#excerpt_0] | |
+| **Device** | Sets the AJA device that this Media Source will use to bring video into the Unreal Engine. If you have multiple cards or devices attached to your computer, you can choose which one to use here. |
+| **Source** | Sets the port or SDI connection on the specified device that this Media Source will use to bring video into the Unreal Engine. Note that only single-link connections are supported. |
+| **Resolution** | Sets the resolution of the incoming video feed. Note that this must match the actual video feed exactly. |
+| **Standard** | Sets whether the incoming video feed is progressive or interlaced. Note that this must match the actual video feed exactly. |
+| **Frame Rate** | Sets the number of video frames per second in the incoming feed. Note that this must match the actual video feed exactly. |
+| **Other AJA Options** | |
+| **Timecode Format** | Specifies the type of timecode that accompanies the video signal. |
+| **Capture with Auto Circulating** | [INCLUDE:#excerpt_1] |
+| **Ancillary** | |
+| **Capture Ancillary** | Determines whether the Unreal Engine captures ancillary metadata that accompanies the video signal. |
+| **Max Num Ancillary Frame Buffer** | Sets the maximum number of frames of ancillary data the Unreal Engine will store in memory at any given time. If the input video jumps or hitches, you can try raising this value. |
+| **Audio** | |
+| **Capture Audio** | Determines whether the Unreal Engine captures audio from the Media Port. |
+| **Audio Channel** | Specifies the audio channel that contains the signal you want the Unreal Engine to capture. |
+| **Max Num Audio Frame Buffer** | Sets the maximum number of frames of audio data the Unreal Engine will store in memory at any given time. If the input video jumps or hitches, you can try raising this value. |
+| **Video** | |
+| **Capture Video** | Determines whether the Unreal Engine captures video from the Media Port. |
+| **Color Format** | Determines the order of the color channels that make up each pixel in the input video, and the number of bits in each channel. |
+| **Max Num Video Frame Buffer** | Sets the maximum number of frames of video data the Unreal Engine will store in memory at any given time. If the input video jumps or hitches, you can try raising this value. |
+| **Debug** | |
+| **Log Drop Frame** | When enabled, the Unreal Engine prints a message to its output log every time it detects a dropped frame in the input feed. |
+| **Encode Timecode in Texel** | When enabled, the engine embeds the timecode of each frame into the captured video. You can use this to check that the timecode for each frame of input matches the values you're expecting. See [Timecode Texel Encoding](https://docs.unrealengine.com/en-us/Engine/ProVideoIO/TimecodeGenlock#timecodetexelencoding). |
+| **Synchronization** | |
+| **Synchronize with Engine's Timecode** | By default, the Unreal Engine attempts to play the video frames as they come in. Enable this setting to buffer incoming frames, and try to align them with the Unreal Engine's internal timecode. If the timecode of any buffered frame matches the Unreal Engine's internal timecode for any frame, the video input will be synchronized with from that frame forward. |
+| **Platforms** | |
+| **Player Overrides** | Leave these options at their default values for AJA Media Sources. |
+
+
+[COMMENT:none]
+-
+[EXCERPT:excerpt_1]
+Captures video, audio, and ancillary data at the same time. This may decrease your transfer performance, but it guarantees that all data for each frame will be synchronized together.  
+If you experience problems with latency, you can try disabling this option.
+[/EXCERPT:excerpt_1]
+-
+[/COMMENT]
+
+[COMMENT:none]
+-
+[EXCERPT:excerpt_0]
+
+![Aja Media Source Configuration](aja-source-settings.png "Aja Media Source Configuration")
+
+[/EXCERPT:excerpt_0]
+-
+[/COMMENT]
+
+## AJA Media Output Settings
+
+Each AJA Media Output object that you create exposes the following configuration settings.
+
+  
+| **Property** | **Description** |
+| --- | --- |
+| **Configuration Settings** | |
+| [INCLUDE:#excerpt_2]| |
+| **Output Type** | [INCLUDE:#excerpt_3] |
+| **Device** | Sets the AJA device that this Media Source will send its video feed to. If you have multiple cards or devices attached to your computer, you can choose which one to use here |
+| **Source** | Sets the port or SDI connection on the specified **Device **that this Media Source will send video to. Note that only single-link connections are supported. |
+| **Resolution** | Sets the resolution of the video feed produced by this Media Output. |
+| **Standard** | Sets whether the output feed produced by this Media Output is progressive or interlaced. |
+| **Frame Rate** | Sets the number of frames per second in the video feed produced by this Media Output. |
+| **Key Source** | Sets the port that will receive the key images from the Unreal Engine, when the **Output Type** is set to **Fill and Key**. |
+| **Reference** | [INCLUDE:#excerpt_4] |
+| **Sync Source** | When the **Output Reference** is set to **Input**, this setting specifies which input port the AJA card should use as its output reference. The sync port must be on the same device as the fill port. |
+| **Output** | |
+| **Output with Auto Circulating** | [INCLUDE:#excerpt_5] |
+| **Timecode Format** | Determines whether the Unreal Engine should embed timecode in the output feed, and which timecode format it should use. |
+| **Pixel Format** | [INCLUDE:#excerpt_6] |
+| **Number of AJA Buffers** | [INCLUDE:#excerpt_7] |
+| **Interlaced Fields Timecode Need to Match** | When producing an interlaced video feed, this setting determines whether the timecode values for both fields in a single interlaced frame need to match. |
+| **Number of Texture Buffers** | [INCLUDE:#excerpt_8] |
+| **Synchronization** | |
+| **Wait for Sync Event** | [INCLUDE:#excerpt_9] |
+| **Debug** | |
+| **Encode Timecode in Pixel** | When enabled, the engine embeds the timecode of each frame in the output signal. See [Timecode Texel Encoding](Engine/ProVideoIO/TimecodeGenlock#timecodetexelencoding). |
+
+[COMMENT:none]
+-
+[EXCERPT:excerpt_9]
+
+When this option is disabled, and you don't already have the Unreal Engine genlocked to an input signal, the engine runs at the fastest frame rate it can manage and provides all the frames it generates to the AJA card. Each time the card is ready to output a new frame, it selects one of the frames generated by the Engine.
+
+When this option is enabled, the Unreal Engine does not generate any new frames of output until the AJA card is ready to accept the new frame. The effect is similar to genlock, but instead of locking the Unreal Engine's frame rate to an input signal, it locks the Engine's frame rate to the output timing of the AJA card.  
+
+This option is most useful when you don't already have an input signal that you can lock the Unreal Engine's frame rate to, but you want to ensure that the Engine is producing only one output frame for every frame in the output video feed.  
+
+Do not enable this option if you already have the Unreal Engine genlocked to an input feed using a custom time step.
+
+[/EXCERPT:excerpt_9]
+-
+[/COMMENT]
+
+[COMMENT:none]
+-
+[EXCERPT:excerpt_8]
+Sets the number of buffers used to transfer each frame image from the GPU to main thread memory.  
+Lower values are more likely to cause a bottleneck on the GPU side as it waits for each transfer to complete; larger numbers are more likely to increase latency.
+[/EXCERPT:excerpt_8]
+-
+[/COMMENT]
+
+[COMMENT:none]
+-
+[EXCERPT:excerpt_7]
+Sets the number of buffers used to transfer each frame image from the main thread memory to the AJA card.  
+Lower values are more likely to cause missed frames as it waits for each transfer to complete; larger numbers are more likely to increase latency.
+[/EXCERPT:excerpt_7]
+-
+[/COMMENT]
+[COMMENT:none]
+-
+[EXCERPT:excerpt_6]
+
+Determines the order of the color channels that make up each pixel, and the number of bits in each channel.
+
+[REGION:tip]
+If you want to output the alpha, set the **Output Type **setting to **Fill and Key**, and use the **Key Source **to send the alpha to an output port on your AJA card.
+[/REGION]
+[/EXCERPT:excerpt_6]
+-
+[/COMMENT]
+[COMMENT:none]
+-
+[EXCERPT:excerpt_5]
+
+When enabled, the Unreal Engine buffers its output frames before sending them to the AJA card. This may improve the smoothness of the video signal, at the cost of some latency.  
+Leave this option disabled to minimize latency, at the risk of seeing interruptions in the output signal.
+
+[/EXCERPT:excerpt_5]
+-
+[/COMMENT]
+
+[COMMENT:none]
+-
+[EXCERPT:excerpt_4]
+
+Configures the source of the timing for the internal clock on the AJA card. The card uses this to determine when it should send each frame of video output.  
+
+*   **Free Run -** Uses the card's internal clock.
+*   **External -** Synchronizes the card's internal clock with the genlock signal that arrives on its reference pin from an external source.
+*   **Input -** Synchronizes with the video signal from the input port that you specify in the **Sync Source** setting below. 
+
+[/EXCERPT:excerpt_4]
+-
+[/COMMENT]
+[COMMENT:none]
+-
+[EXCERPT:excerpt_3]
+Determines whether the Unreal Engine outputs only the fill image, or both the fill and key images.  
+When you set this to **Fill Only**, only the fill image is output to the **Source** set below.  
+When you set this to **Fill and Key**, the fill image is output to the **Source**, and the key is output to the **Key Source**.
+[/EXCERPT:excerpt_3]
+-
+[/COMMENT]
+[COMMENT:none]
+-
+[EXCERPT:excerpt_2]
+
+![Aja Media Output Configuration](aja-output-settings.png "Aja Media Output Configuration")
+
+[/EXCERPT:excerpt_2]
+-
+[/COMMENT]
\ No newline at end of file
diff --git a/Engine/Documentation/Source/Engine/ProVideoIO/AjaQuickStart/AjaQuickStart.INT.udn b/Engine/Documentation/Source/Engine/ProVideoIO/AjaQuickStart/AjaQuickStart.INT.udn
new file mode 100644
--- /dev/null
+++ b/Engine/Documentation/Source/Engine/ProVideoIO/AjaQuickStart/AjaQuickStart.INT.udn
@@ -0,0 +1,212 @@
+Availability: Public
+Crumbs: 
+Title: AJA Video I/O Quick Start
+Description: A step-by-step guide to getting video from a supported AJA media card into the Unreal Engine, and sending captured video output from the Unreal Engine out to a port on the AJA card.
+Type: 
+Version: 4.20
+Parent: Engine/ProVideoIO
+Order: 
+Tags: Getting Started
+Tags: Early Access
+Tags: professional video
+
+
+![](aja-qs-banner.png)
+
+In this Quick Start guide, we walk through the process of setting up an Unreal Engine Project to work with a professional video card from AJA Video Systems. At the end of this guide:
+
+*   you'll have video input from your AJA card playing inside your Unreal Engine Project.
+*   you'll be able to capture camera viewpoints both from the Editor and from your runtime application, and send them out to an SDI port on your AJA card.
+*   you'll know where to go when you want to set up more advanced adjustments to your video inputs, such as correcting lens deformation and applying chroma-keying effects.
+
+[REGION:tip]
+For a working example that shows many of the elements described below put into practice, see the [Virtual Studio](Resources/Showcases/VirtualStudio) showcase, available on the Learn tab of the Epic Games Launcher.
+[/REGION]
+
+[REGION:note]
+**Prerequisites:**
+
+*   Make sure that you have a supported card from AJA Video Systems, and that you've installed the necessary drivers and software. For details, see the [AJA Media Reference](Engine/ProVideoIO/AjaIOReference) page.
+*   Make sure that your card is working correctly, and that you have some video input feeding in to at least one of the card's SDI ports.
+*   Open an Unreal Engine Project that you want to integrate with your video feeds. This page shows the steps in the **Third Person** Blueprint template, but the same steps will work equally well in any Project.
+
+The AJA Media components used in this guide are built on top of the [Media Framework](Engine/MediaFramework), and we'll use [Blueprints](Engine/Blueprints) to script the video capturing process at runtime. Some prior knowledge of these topics is recommended but not required.
+[/REGION]
+
+## 1 - Set Up the Project
+
+Before you can get video input from your AJA card into your Unreal Engine Level, and send output from the Unreal Engine through one of your AJA card's SDI ports, you'll need to do some basic setup to acquire the Aja Media Player Plugin and enable it in your Project.
+
+### Steps
+
+1.  The AJA Media Player Plugin is distributed through the Marketplace. You'll need to install it from there.  
+    Open the Epic Games Launcher, go to the **Marketplace** tab, and search for **AJA Media Player**. Follow the instructions in the Launcher to install the Plugin to your version of the Unreal Engine.
+2.  Open the Project that you want to use with AJA video I/O in the Unreal Editor.
+3.  In the **Plugins** window, find the **Aja Media Player** plugin under the **Installed** list. Check its **Enabled** checkbox.  
+    ![Enable the Aja Media Player Plugin](aja-qs-enable-plugin.png "Enable the Aja Media Player Plugin")  
+    
+4.  Find the **Media Framework Utilities** Plugin under the **Media** category. Check its**Enabled** checkbox, if it's not already checked.  
+    ![Enable the Media Framework Utilities Plugin](media-qs-utilities-enabled.png "Enable the Media Framework Utilities Plugin")
+5.  Click **Restart Now** to restart the Unreal Editor and reopen your Project.  
+    ![Restart Now](media-qs-utilities-restart.png "Restart Now")
+
+### End Result
+
+Your Project is now ready to accept video from the AJA card, and to send rendered output to the card. In the next sections, we'll hook it up and start playing video in and out.
+
+## 2 - Rendering Video Input in the Unreal Engine
+
+In this process, we'll make video input from the AJA card visible in the current Level in the Unreal Editor. This process uses a Media Bundle: a kind of Asset that packages together several different types of Assets involved in the Media Framework, and that offers you control over some advanced features like lens deformation, chroma-keying, color correction, and more. 
+
+### Steps
+
+1.  In your **Content Browser**, expand the **Sources** panel (1). Right-click, and choose **New Folder** from the context menu (2).  
+    ![New Folder](aja-qs-2b-1-newfolder.png "New Folder")  
+    Rename your new folder **AJA**.  
+    
+2.  Open your new folder, right-click in the **Content Browser** and choose **Media > Media Bundle**.  
+    ![New Media Bundle](aja-qs-2b-2-newbundle.png "New Media Bundle")
+3.  Your new Asset's name is automatically selected in the Content Browser, so you can give it a descriptive name:  
+    ![Name the Media Bundle](aja-qs-2b-3-filename.png "Name the Media Bundle")Type a new name, like **AjaMediaBundle**, and press **Enter**.  
+    A new folder of Media Framework Assets is automatically created next to your Media Bundle, named with the suffix **_InnerAssets**. 
+4.  Save your new Assets by clicking the **Save All** button in the **Content Browser**.  
+    ![Save All Assets](aja-qs-2b-4-saveall.png "Save All Assets")
+5.  Double-click your new Media Bundle to edit its properties. The Media Bundle is capable of playing video from any kind of media source the Engine supports, so you'll need to tell it that you want to get the video from your AJA card.  
+    In the **Media Source** property, select **Aja Media Source** from the drop-down list:  
+    ![Set the AJA Media Source](aja-qs-2b-4-sourcetype.png "Set the AJA Media Source")
+6.  Once you've identified the type of Media Source you want the Media Bundle to handle, you can then set up any configuration properties offered by that type of source.  
+    The most important thing to set here for the Aja Media Source is the **Configuration** setting, to make sure that the bundle is set up to capture video from the right device and input port, using the same resolution and frame rate as the actual video feed. Click the arrow to open the settings submenu, select the options that match your setup, then click **Apply** in the submenu.  
+    ![Aja Media Source Configuration](aja-qs-2b-5-configuration.png "Aja Media Source Configuration") The options you see may vary depending on the devices you have installed. For details on all of the properties you can set for an AJA Media Source, see the [AJA Media Reference](https://docs.unrealengine.com/en-us/Engine/ProVideoIO/AjaIOReference) page.
+7.  If you want to apply any compensation to the incoming video to account for lens distortion, you can set up the physical properties of the lens in the **Lens Parameters** section.   
+    ![Lens undistortion parameters](aja-qs-2b-6-lensdistortion.png "Lens undistortion parameters")  
+    These **Lens Parameters** just set up the physical properties of the lens. You'll actually activate the lens compensation later, when you edit the Material Instance used by the Media Bundle.  
+    Save your Media Bundle when you're done setting up its properties, and return to the **AJA** folder in the Content Browser.
+8.  Drag your **AjaMediaBundle** Asset from the Content Browser into the Level Viewport.  
+    ![Drag and drop the Media Bundle](aja-qs-2b-9-dragdrop.png "Drag and drop the Media Bundle")  
+    You'll see a new plane appear, showing the video currently being played over the port configured for your Media Bundle. Use the transform tools in the Viewport toolbar to move, rotate, and resize it.  
+    If your Media Bundle doesn't start playing automatically, select it, then click the **Media Bundle > Request Play Media** button in the **Details** panel.  
+    ![Request Play Media](aja-request-play-media.png "Request Play Media")
+9.  Now, we'll see how to apply keying and compositing effects to the video stream.  
+    Back in the Media Bundle Editor, click the **Open Material Editor** button in the Toolbar to edit the Material Instance that this Media Bundle uses to draw its incoming video feed on to an object in the Level.  
+    ![Open Material Editor](aja-open-material-editr.png "Open Material Editor")  
+    [REGION:tip]
+    This Material Instance is saved inside the **AjaMediaBundle_InnerAssets** folder that was created automatically with your Media Bundle.  
+    ![Material Instance](aja-qs-2b-7-materialinstance.png "Material Instance")
+    [/REGION]
+10.  In the Material Instance Editor, you'll see a number of properties exposed for you to configure keying, cropping, and color correction, and to activate the correction for the lens distortion that you set up in the Media Bundle.  
+    ![Material Instance Editor](aja-qs-2b-8-materialinstanceeditor.png "Material Instance Editor")  
+    While you adjust the settings in the Material Instance Editor, you can see the effect of your changes on the video feed playing back in the main Level Viewport.
+    [REGION:tip]
+    You may find it more convenient to see the effects of the changes you make in the preview panel of the Material Instance Editor instead. To do this, temporarily enable the **IsValid** setting, and set its value to `1.0`.  
+    ![IsValid](video-qs-preview-isvalid.png "IsValid")  
+    Click the arrow at the top left of the viewport toolbar, and enable the **Realtime** option in the menu.  
+    ![Realtime viewport](aja-qs-2b-8-realtime.png "Realtime viewport")  
+    You'll be able to judge the effect of your changes more easily by changing the preview mesh to a plane or a cube. Use the controls at the bottom of the viewport:  
+    ![Preview mesh](aja-qs-2b-8-displayshape.png "Preview mesh")  
+    When you're done, return the **IsValid** setting to its previous value.
+    [/REGION]
+11.  When you're done changing the Material Instance properties, click the **Save** button in the Toolbar.
+
+### End Result
+
+At this point, you should have video playing over an SDI port showing up inside your Unreal Engine Level, and you should understand where to set up more advanced features like lens deformation and chroma-keying.
+
+If you're already familiar with the Media Framework, another way you could get video into your Level is to create a new **AjaMediaSource** Asset in your Project, and set it up with the same source properties you set up inside the Media Bundle in the procedure above. Then, create your own **MediaPlayer** and **MediaTexture** Assets to handle the playback of that source in your Level. For details, see the [Media Framework](https://docs.unrealengine.com/en-us/Engine/MediaFramework) documentation. However, we recommend using the Media Bundle, as shown above, to get the best balance between ease of use and professional quality video features.
+
+## 3 - Output Captures from the Unreal Editor
+
+In this process, you'll set up an AJA Media Output object, and use the **Media Captures** panel in the Unreal Editor to output the view from selected cameras in the Level to your AJA card.
+
+### Steps
+
+1.  Right-click in the Content Browser, and select **Media > Aja Media Output**.   
+    ![New AJA Media Output](aja-qs-3-1-newoutput.png "New AJA Media Output")  
+    Name your new Asset **AjaMediaOutput**.
+2.  Double-click your new Asset to open it up for editing. Just like when you created your Aja Media Source, you have to set up the **Configuration** property to control the properties of the video feed that the Unreal Engine sends to your AJA card. Click the arrow to open the submenu, select the options that match your video setup, then click **Apply** in the submenu.  
+    ![Aja Media Output Configuration](aja-qs-3-2-configuration.png "Aja Media Output Configuration")  
+    For details on all of the properties you can set in the AJA Media Output, see the [AJA Media Reference](Engine/ProVideoIO/AjaIOReference) page.Save and close your Media Output when you're done.
+3.  Now we'll place two cameras in the Level, to give us viewpoints for the output we'll send to the AJA card. In the **Modes** panel, open the **Cinematic** tab, and drag two instances of the **Cine Camera Actor** into the viewport.  
+    ![Drag and drop Cine Camera Actors](aja-qs-3-3-cinecameras.png "Drag and drop Cine Camera Actors")  
+    Place the cameras where you want them in the Level, so that they're showing different viewpoints on the scene.
+    [REGION:tip]
+    **Piloting** the camera is a fast and easy way to set its viewpoint exactly the way you want it. See [Pilot Actors in the Viewport](Engine/UI/LevelEditor/Viewports/ActorLocking).
+    [/REGION]
+4.  From the main menu, choose **Window > Media Capture**. You'll use the **Media Capture** window to control when the Editor sends output to your AJA port, and what camera it should use in the Level.  
+    ![Media Capture window](aja-qs-3-4-mediacapture.png "Media Capture window")
+5.  Under the **Media Viewport Capture** area, find the **Viewport Captures** control. Click the **+** button to add a new capture to this list.  
+    ![Add a viewport capture](aja-qs-3-5-newcapture.png "Add a viewport capture")
+6.  Expand the new entry. First, we'll add the cameras that we want to capture from. In the **Locked Camera Actors** control, click the **+** button to add a new entry.  
+    ![Add a camera actor](aja-qs-3-6-newcamera.png "Add a camera actor")  
+    Then, use the drop-down list to choose one of the cameras you placed in the Level.  
+    ![Select the camera actor](aja-qs-3-6-selectcamera.png "Select the camera actor")  
+    Repeat the same steps to add the other camera to the list.
+7.  Now, set up the output that you want to capture these cameras to. Set the **Media Output** control to point to the new AJA Media Output Asset that you created above. You can do this by selecting it in the drop-down list, or drag your AJA Media Output Asset from the Content Browser and drop it into this slot.  
+    ![Set the AJA Media Output](aja-qs-3-7-selectoutput.png "Set the AJA Media Output")
+8.  At the top of the window, click the **Capture** button.  
+    ![Start capturing](aja-qs-3-8-capture.png "Start capturing")  
+    You'll see a new frame at the bottom of the window that shows a preview of the output being sent to the AJA card. If you have this port hooked up to another downstream device, you should start to see the output coming through.  
+    ![Active Media Capture](aja-qs-output-8.png "Active Media Capture")
+9.  Each camera that you added to the Locked Camera Actors list for this viewport capture is represented by a corresponding button above the video preview. Click the buttons to switch the capture back and forth between the two views.  
+    ![Switch Cameras](aja-qs-output-9.png "Switch Cameras")
+
+### End Result
+
+Now you've set up the Unreal Editor to stream output from cameras in your Level to a port on your AJA card. Next, we'll see how to use Blueprint scripting to do the same thing in a running Unreal Engine Project.
+
+## 4 - Output Captures at Runtime
+
+The **Media Capture** window that you used in the last section is a practical and easy way to send captures to the AJA card. However, it only works inside the Unreal Editor. To do the same thing when you're running your Project as a standalone application, you'll need to use the Blueprint API provided by the Media Output. In this procedure, we'll set up a simple toggle switch in the Level Blueprint that starts and stops capturing when the player presses a key on the keyboard.
+
+[REGION:tip]
+The [Virtual Studio](Resources/Showcases/VirtualStudio)** showcase, available on the **Learn** tab of the Epic Games Launcher, contains a UMG interface widget that demonstrates how you could control capturing from an on-screen user interface.
+[/REGION]
+
+### Steps
+
+1.  From the main toolbar in the Unreal Editor, choose **Blueprints > Open Level Blueprint**.  
+    ![Open Level Blueprint](aja-qs-4-1-levelbp.png "Open Level Blueprint")
+2.  We'll need to start from the AJA Media Output Asset that you've created, where you identify the port you want to output to. In the **Variables** list in the **My Blueprint** panel, click the **+** button to add a new variable.  
+    ![New Variable](aja-qs-4-2-newvar.png "New Variable")
+3.  In the **Details** panel, set the **Variable Name** to **AjaMediaOutput**, and use the **Variable Type** drop-down list to make it an **Aja Media Output Object Reference**.  
+    ![Aja Media Output Object Reference](aja-qs-4-3-vartype.png "Aja Media Output Object Reference")
+4.  Enable the **Instance Editable** setting (1), and compile the Blueprint. Then, in the **Default Value** section, set the variable to point to the AJA Media Output Asset that you created in your Content Browser (2).  
+    ![Set the default value](aja-qs-4-4-default.png "Set the default value")
+5.  Press **Ctrl**, and drag the **AjaMediaOutput** from the Variables list in the **My Blueprint** panel into the **Event Graph**.  
+    ![Control+drag the AjaMediaOutput](aja-qs-4-5-dropvariablet.png "Control+drag the AjaMediaOutput")
+6.  Click and drag from the output port of the **AjaMediaOutput** variable node, and choose **Media > Output > Create Media Capture**.   
+    ![Create Media Capture](aja-qs-4-6-createcapture.png "Create Media Capture")  
+    Hook up your nodes to the **Event BeginPlay** node as shown below:  
+    ![Event Begin Play](aja-qs-4-6-eventplay.png "Event Begin Play")  
+    This creates a new Media Capture object from your Aja Media Output. The Media Capture offers two main Blueprint functions that we'll use to control the capturing: **Capture Active Scene Viewport** and **Stop Capture**.
+7.  First, we'll save the new Media Capture object into its own variable, so we can get access to it again elsewhere. Click and drag from the output port of the **Create Media Capture** node, and choose **Promote to Variable**.  
+    ![Promote to variable](aja-qs-4-7-promote.png "Promote to variable")  
+    Rename the new variable **MediaCapture** in the Variables list in the **My Blueprint** panel.
+    [REGION:tip]
+    It's important to save the Media Capture to a variable here. If you don't, the Unreal Engine's garbage collector may destroy it automatically before you're done with it.
+    [/REGION]
+8.  Press **Ctrl** and drag the **MediaCapture** variable into the **Event Graph**.  
+    ![Control+drag the MediaCapture](aja-qs-4-8-mediacapture.png "Control+drag the MediaCapture")  
+    
+9.  Click and drag from the output port of the **MediaCapture** variable node, and choose **Media > Output > Capture Active Scene Viewport**. Do it again, and choose **Media > Output > Stop Capture**.  
+    ![Start and stop the capture](aja-qs-4-9-capturestop-1.png "Start and stop the capture")
+10.  Right-click in the **Event Graph** and choose **Input > Keyboard Events > P**.   Click and drag from the **Pressed** output of the **P** node and choose **Flow Control > FlipFlop**.   
+    ![FlipFlop](aja-qs-4-10-pflipflop.png "FlipFlop")
+11.  Connect the **A** output of the **FlipFlop** node to the input event of the **Capture Active Scene Viewport** node, and connect the **B** output of the **FlipFlop** node to the input event of the **Stop Capture** node, as shown below:  
+    ![Connect the nodes](aja-qs-4-11-done.png "Connect the nodes")
+12.  Compile and save the Blueprint, and try playing your Project. Click the arrow next to the Play button in the main Toolbar, and choose either the **New Editor Window (PIE)** or **Standalone Game** option.  
+    ![Launch the Project](aja-qs-4-12-pie.png "Launch the Project")  
+    [REGION:note]
+    Video capture from the Editor will only work when you play your Project in a **New Editor Window (PIE)** or a **Standalone Game**. It won't work in the default **Selected Viewport** mode, or in **Simulate** mode.  
+    In addition, the viewport resolution of your Project (that is, the size of the rendered image the Unreal Engine generates each frame) must match the output resolution set in the active Media Profile, so that it will be the right size for the output video feed.
+    [/REGION]After your project starts up, you should be able to press the **P** button on your keyboard to toggle sending the output from the Engine to the AJA card.
+
+### End Result
+
+At this point, you should have a basic idea of how to work with Aja Media Sources, Media Bundles, and the Media Capture system, and you should understand how all of these elements work together to get professional video in and out of the Unreal Engine.
+
+## On Your Own
+
+Now that you've seen the basics of how to get a new Project exchanging video input and output with an AJA card, you can continue learning on your own:
+
+*   Explore the in-engine keying solution in the Material Instance created by your Media Bundle. Try passing some green-screen video into your card's input port, and use the keying controls in the Material Instance to remove the background.
+*   Explore the [Virtual Studio](Resources/Showcases/VirtualStudio) showcase to see what it adds to this basic setup, like its on-screen UI that switches cameras and controls video capture at runtime.
\ No newline at end of file
diff --git a/Engine/Documentation/Source/Engine/ProVideoIO/BlackmagicIOReference/BlackmagicIOReference.INT.udn b/Engine/Documentation/Source/Engine/ProVideoIO/BlackmagicIOReference/BlackmagicIOReference.INT.udn
new file mode 100644
--- /dev/null
+++ b/Engine/Documentation/Source/Engine/ProVideoIO/BlackmagicIOReference/BlackmagicIOReference.INT.udn
@@ -0,0 +1,155 @@
+Availability: Public
+Crumbs: 
+Title: Blackmagic Media Reference
+Description: Describes the options and settings exposed by the Blackmagic Design Media Framework components.
+Type: 
+Version: 4.21
+Parent: Engine/ProVideoIO
+Order: 
+Tags: Reference
+Tags: Early Access
+Tags: professional video
+
+
+This page describes the options and settings exposed on Blackmagic Media Framework objects.
+
+## Supported Blackmagic Cards
+
+The Blackmagic Media Source and Blackmagic Media Output have been tested with the following cards:
+
+*   **DeckLink 4K Extreme 12G**
+*   **DeckLink Duo 2**
+*   **DeckLink 8K Pro**
+
+Other devices may or may not work as expected.
+
+## Blackmagic Media Source Settings
+
+Each Blackmagic Media Source object that you create exposes the following configuration settings.
+
+  
+| **Property** | **Description** |
+| --- | --- |
+| **Configuration Options**  | |
+| [INCLUDE:#excerpt_0]| |
+| **Device** | Sets the Blackmagic device and SDI connection that this Media Source will use to bring video into the Unreal Engine. If you have multiple cards or devices attached to your computer, you can choose which one to use here. |
+| **Resolution** | Sets the resolution of the incoming video feed. Note that this must match the actual video feed exactly. |
+| **Standard** | Sets whether the incoming video feed is progressive or interlaced. Note that this must match the actual video feed exactly. |
+| **Frame Rate** | Sets the number of video frames per second in the incoming feed. Note that this must match the actual video feed exactly. |
+| **Other Blackmagic Options** | |
+| **Timecode Format** | Specifies the type of timecode that accompanies the video signal. |
+| **Audio** | |
+| **Capture Audio** | Determines whether the Unreal Engine captures audio from the Media Port. |
+| **Audio Channels** | Specifies the audio channel that contains the signal you want the Unreal Engine to capture. |
+| **Max Num Audio Frame Buffer** | Sets the maximum number of frames of audio data the Unreal Engine will store in memory at any given time. If the input video jumps or hitches, you can try raising this value. |
+| **Video** | |
+| **Capture Video** | Determines whether the Unreal Engine captures video from the Media Port. |
+| **Color Format** | Determines the order of the color channels that make up each pixel in the input video, and the number of bits in each channel. |
+| **Max Num Video Frame Buffer** | Sets the maximum number of frames of video data the Unreal Engine will store in memory at any given time. If the input video jumps or hitches, you can try raising this value. |
+| **Debug** | |
+| **Log Drop Frame** | When enabled, the Unreal Engine prints a message to its output log every time it detects a dropped frame in the input feed. |
+| **Encode Timecode in Texel** | When enabled, the engine embeds the timecode of each frame into the captured video. You can use this to check that the timecode for each frame of input matches the values you're expecting. See [Timecode Texel Encoding](https://docs.unrealengine.com/en-us/Engine/ProVideoIO/TimecodeGenlock#timecodetexelencoding). |
+| **Synchronization** | |
+| **Synchronize with Engine's Timecode** | By default, the Unreal Engine attempts to play the video frames as they come in. Enable this setting to buffer incoming frames, and try to align them with the Unreal Engine's internal timecode. If the timecode of any buffered frame matches the Unreal Engine's internal timecode for any frame, the video input will be synchronized with from that frame forward. |
+| **Platforms** | |
+| **Player Overrides** | Leave these options at their default values for Blackmagic Media Sources. |
+
+
+[COMMENT:none]
+-
+[EXCERPT:excerpt_0] 
+
+![Blackmagic Media Source Configuration](bm-source-settings.png "Blackmagic Media Source Configuration")
+
+[/EXCERPT:excerpt_0]
+-
+[/COMMENT]
+
+## Blackmagic Media Output Settings
+
+Each Blackmagic Media Output object that you create exposes the following configuration settings.
+
+  
+| **Property** | **Description** |
+| --- | --- |
+| **Configuration Settings**  | |
+| [INCLUDE:#excerpt_1]| |
+| **Output Type** | [INCLUDE:#excerpt_2] |
+| **Device** | Sets the Blackmagic device and SDI connection that this Media Source will send its video feed to. If you have multiple cards or devices attached to your computer, you can choose which one to use here. |
+| **Resolution** | Sets the resolution of the video feed produced by this Media Output. |
+| **Standard** | Sets whether the output feed produced by this Media Output is progressive or interlaced. |
+| **Frame Rate** | Sets the number of frames per second in the video feed produced by this Media Output. |
+| **Key Source** | Sets the port that will receive the key images from the Unreal Engine, when the **Output Type** is set to **Fill and Key**. |
+| **Output** | |
+| **Timecode Format** | Determines whether the Unreal Engine should embed timecode in the output feed, and which timecode format it should use. |
+| **Pixel Format** | [INCLUDE:#excerpt_3] |
+| **Number of Blackmagic Buffers** | [INCLUDE:#excerpt_4] |
+| **Interlaced Fields Timecode Need to Match** | When producing an interlaced video feed, this setting determines whether the timecode values for both fields in a single interlaced frame need to match. |
+| **Number of Texture Buffers** | [INCLUDE:#excerpt_5] |
+| **Synchronization** | |
+| **Wait for Sync Event** | [INCLUDE:#excerpt_6] |
+| **Debug** | |
+| **Encode Timecode in Pixel** | When enabled, the engine embeds the timecode of each frame in the output signal. See [Timecode Texel Encoding](Engine/ProVideoIO/TimecodeGenlock#timecodetexelencoding). |
+
+[COMMENT:none]
+-
+[EXCERPT:excerpt_6]
+
+When this option is disabled, and you don't already have the Unreal Engine genlocked to an input signal, the engine runs at the fastest frame rate it can manage and provides all the frames it generates to the Blackmagic card. Each time the card is ready to output a new frame, it selects one of the frames generated by the Engine.
+
+When this option is enabled, the Unreal Engine does not generate any new frames of output until the Blackmagic card is ready to accept the new frame. The effect is similar to genlock, but instead of locking the Unreal Engine's frame rate to an input signal, it locks the Engine's frame rate to the output timing of the Blackmagic card.  
+
+This option is most useful when you don't already have an input signal that you can lock the Unreal Engine's frame rate to, but you want to ensure that the Engine is producing only one output frame for every frame in the output video feed.  
+
+Do not enable this option if you already have the Unreal Engine genlocked to an input feed using a custom time step.
+
+[/EXCERPT:excerpt_6]
+-
+[/COMMENT]
+
+[COMMENT:none]
+-
+[EXCERPT:excerpt_5]
+Sets the number of buffers used to transfer each frame image from the GPU to main thread memory.  
+Lower values are more likely to cause a bottleneck on the GPU side as it waits for each transfer to complete; larger numbers are more likely to increase latency.
+[/EXCERPT:excerpt_5]
+-
+[/COMMENT]
+[COMMENT:none]
+-
+[EXCERPT:excerpt_4]
+Sets the number of buffers used to transfer each frame image from the main thread memory to the Blackmagic card.  
+Lower values are more likely to cause missed frames as it waits for each transfer to complete; larger numbers are more likely to increase latency.
+[/EXCERPT:excerpt_4]
+-
+[/COMMENT]
+[COMMENT:none]
+-
+[EXCERPT:excerpt_3]
+
+Determines the order of the color channels that make up each pixel, and the number of bits in each channel.
+
+[REGION:tip]
+If you want to output the alpha, set the **Output Type **setting to **Fill and Key**, and use the **Key Source **to send the alpha to an output port on your Blackmagic card.
+[/REGION]
+[/EXCERPT:excerpt_3]
+-
+[/COMMENT]
+[COMMENT:none]
+-
+[EXCERPT:excerpt_2]
+Determines whether the Unreal Engine outputs only the fill image, or both the fill and key images.  
+When you set this to **Fill Only**, only the fill image is output to the **Source** set below.  
+When you set this to **Fill and Key**, the fill image is output to the **Source**, and the key is output to the **Key Source**.
+[/EXCERPT:excerpt_2]
+-
+[/COMMENT]
+[COMMENT:none]
+-
+[EXCERPT:excerpt_1] 
+
+![Blackmagic Media Output Configuration](bm-output-settings.png "Blackmagic Media Output Configuration")
+
+[/EXCERPT:excerpt_1]
+-
+[/COMMENT]
\ No newline at end of file
diff --git a/Engine/Documentation/Source/Engine/ProVideoIO/BlackmagicQuickStart/BlackmagicQuickStart.INT.udn b/Engine/Documentation/Source/Engine/ProVideoIO/BlackmagicQuickStart/BlackmagicQuickStart.INT.udn
new file mode 100644
--- /dev/null
+++ b/Engine/Documentation/Source/Engine/ProVideoIO/BlackmagicQuickStart/BlackmagicQuickStart.INT.udn
@@ -0,0 +1,222 @@
+Availability: Public
+Crumbs: 
+Title: Blackmagic Video I/O Quick Start
+Description: A step-by-step guide to getting video from a supported Blackmagic media card into the Unreal Engine, and sending captured video output from the Unreal Engine out to a port on the Blackmagic card.
+Type: 
+Version: 4.21
+Parent: Engine/ProVideoIO
+Order: 
+Tags: Getting Started
+Tags: Early Access
+Tags: professional video
+
+
+![](blackmagic-qs-banner.png)
+
+In this Quick Start guide, we walk through the process of setting up an Unreal Engine Project to work with a professional video card from Blackmagic Design. At the end of this guide:
+
+*   you'll have video input from your Blackmagic card playing inside your Unreal Engine Project.
+*   you'll be able to capture camera viewpoints both from the Editor and from your runtime application, and send them out to an SDI port on your Blackmagic card.
+*   you'll know where to go when you want to set up more advanced adjustments to your video inputs, such as correcting lens deformation and applying chroma-keying effects.
+
+[REGION:tip]
+For a working example that shows many of the elements described below put into practice, see the [Virtual Studio](Resources/Showcases/VirtualStudio) showcase, available on the Learn tab of the Epic Games Launcher.
+[/REGION]
+
+[REGION:note]
+**Prerequisites:**
+
+*   Make sure that you have a supported card from Blackmagic Design, and that you've installed the necessary drivers and software. For details, see the [Blackmagic Media Reference](Engine/ProVideoIO/BlackmagicIOReference) page.
+*   Make sure that your card is working correctly, and that you have some video input feeding in to at least one of the card's SDI ports.
+*   Open an Unreal Engine Project that you want to integrate with your video feeds. This page shows the steps in the **Third Person** Blueprint template, but the same steps will work equally well in any Project.
+
+The Blackmagic Design components used in this guide are built on top of the [Media Framework](Engine/MediaFramework), and we'll use [Blueprints](Engine/Blueprints) to script the video capturing process at runtime. Some prior knowledge of these topics is recommended but not required.
+[/REGION]
+
+## 1 - Set Up the Project
+
+Before you can get video input from your Blackmagic card into your Unreal Engine Level, and send output from the Unreal Engine through one of your Blackmagic card's SDI ports, you'll need to do some basic setup to acquire the Blackmagic Media Player Plugin and enable it in your Project.
+
+### Steps
+
+1.  The Blackmagic Media Player Plugin is distributed through the Marketplace. You'll need to install it from there.  
+    Open the Epic Games Launcher, go to the **Marketplace** tab, and search for **Blackmagic Media Player**. Follow the instructions in the Launcher to install the Plugin to your version of the Unreal Engine.
+2.  Open the Project that you want to use with Blackmagic video I/O in the Unreal Editor.
+3.  In the **Plugins** window, find the **Blackmagic Media Player** plugin. Check its **Enabled** checkbox.  
+    ![Enable the Blackmagic Media Player Plugin](blackmagic-qs-enable-plugin-1.png "Enable the Blackmagic Media Player Plugin")  
+    
+4.  Find the **Media Framework Utilities** Plugin under the **Media** category. Check its **Enabled** checkbox, if it's not already checked.  
+    ![Enable the Media Framework Utilities Plugin](media-qs-utilities-enabled.png "Enable the Media Framework Utilities Plugin")
+5.  Click **Restart Now** to restart the Unreal Editor and reopen your Project.  
+    ![Restart Now](media-qs-utilities-restart.png "Restart Now")
+
+### End Result
+
+Your Project is now ready to accept video from the Blackmagic card, and to send rendered output to the card. In the next sections, we'll hook it up and start playing video in and out.
+
+## 2 - Rendering Video Input in the Unreal Engine
+
+In this process, we'll make video input from the Blackmagic card visible in the current Level in the Unreal Editor. This process uses a Media Bundle: a kind of Asset that packages together several different types of Assets involved in the Media Framework, and that offers you control over some advanced features like lens deformation, chroma-keying, color correction, and more. 
+
+### Steps
+
+1.  In your **Content Browser**, expand the **Sources** panel (1). Right-click, and choose **New Folder** from the context menu (2).  
+    ![New Folder](aja-qs-2b-1-newfolder.png "New Folder")  
+    Rename your new folder **Blackmagic**.  
+    
+2.  Open your new folder, right-click in the **Content Browser** and choose **Media > Media Bundle**.  
+    ![New Media Bundle](blackmagic-qs-2-2.png "New Media Bundle")  
+    
+3.  Your new Asset's name is automatically selected in the Content Browser, so you can give it a descriptive name:  
+    ![Name the Media Bundle](blackmagic-qs-2-3.png "Name the Media Bundle") Type a new name, like **BlackmagicMediaBundle**, and press **Enter**.  
+    A new folder of Media Framework Assets is automatically created next to your Media Bundle, named with the suffix **_InnerAssets**. We'll take a look at these Assets later.
+4.  Save your new Assets by clicking the **Save All** button in the **Content Browser**.  
+    ![Save All Assets](blackmagic-qs-2-4-saveall.png "Save All Assets")
+5.  Double-click your new Media Bundle to edit its properties. The Media Bundle is capable of playing video from any kind of media source the Engine supports, so you'll need to tell it that you want to get the video from your Blackmagic card.  
+    In the **Media Source** property, select **Blackmagic Media Source** from the drop-down list:  
+    ![Set the Blackmagic Media Source](blackmagic-qs-2-4.png "Set the Blackmagic Media Source")
+6.  Once you've identified the type of Media Source you want the Media Bundle to handle, you can then set up any configuration properties offered by that type of source.   
+    The most important thing to set here for the Blackmagic Media Source is the **Configuration** setting, to make sure that the bundle is set up to capture video from the right device and input port, using the same resolution and frame rate as the actual video feed. Click the arrow to open the settings submenu, select the options that match your setup, then click **Apply** in the submenu.  
+    ![Blackmagic Media Source Configuration](blackmagic-qs-2-5.png "Blackmagic Media Source Configuration")  
+    The options you see may vary depending on the devices you have installed. For details on all of the properties you can set for a Blackmagic Media Source, see the [Blackmagic Media Reference](Engine/ProVideoIO/BlackmagicIOReference) page.
+7.  If you want to apply any compensation to the incoming video to account for lens distortion, you can set up the physical properties of the lens in the **Lens Parameters** section.   
+    ![Lens undistortion parameters](aja-qs-2b-6-lensdistortion.png "Lens undistortion parameters")  
+    These Lens Parameters just set up the physical properties of the lens. You'll actually activate the lens compensation later, when you edit the Material Instance used by the Media Bundle.  
+    Save and close your Media Bundle when you're done setting up its properties.
+8.  Drag your **BlackmagicMediaBundle** Asset from the **Content Browser** into the Level Viewport.   
+    ![Drag and drop the Media Bundle](blackmagic-qs-2-9.png "Drag and drop the media bundle data-verified=")  
+    You'll see a new plane appear, showing the video currently being played over the port configured for your Media Bundle. Use the transform tools in the Viewport toolbar to move, rotate, and resize it.  
+    If your Media Bundle doesn't start playing automatically, select it, then click the **Media Bundle > Request Play Media** button in the **Details** panel.  
+    ![Request Play Media](blackmagic-request-play-media.png "Request Play Media")
+9.  Now, we'll see how to apply keying and compositing effects to the video stream.  
+    Back in the Media Bundle Editor, click the **Open Material Editor** button in the Toolbar to edit the Material Instance that this Media Bundle uses to draw its incoming video feed on to an object in the Level.  
+    ![Open Material Editor](blackmagic-open-material-editor.png "Open Material Editor")
+    [REGION:tip]
+    This Material Instance is saved inside the **BlackmagicMediaBundle_InnerAssets** folder that was created automatically with your Media Bundle.  
+    ![Material Instance](blackmagic-qs-2-7.png "Material Instance")
+    [/REGION]
+10.  In the Material Instance Editor, you'll see a number of properties exposed for you to configure keying, cropping, and color correction, and to activate correcting for the lens distortion that you set up in the Media Bundle.  
+    ![Material Instance Editor](blackmagic-qs-2-8.png "Material Instance Editor")   
+    While you adjust the settings in the Material Instance Editor, you can see the effect of your changes on the video feed playing back in the main Level Viewport.
+    [REGION:tip]
+    You may find it more convenient to see the effects of the changes you make in the preview panel of the Material Instance Editor instead. To do this, temporarily enable the **IsValid** setting, and set its value to `1.0`.  
+    ![IsValid](video-qs-preview-isvalid.png "IsValid")  
+    Click the arrow at the top left of the viewport toolbar, and enable the **Realtime** option in the menu.  
+    ![Realtime viewport](blackmagic-qs-2-8a.png "Realtime viewport")  
+    You'll be able to judge the effect of your changes more easily by changing the preview mesh to a plane or a cube. Use the controls at the bottom of the viewport:  
+    ![Preview mesh](blackmagic-qs-2-8b.png "Preview mesh")  
+    When you're done, return the **IsValid** setting to its previous value.
+    [/REGION]
+11.  When you're done changing the Material Instance properties, click the **Save** button in the Toolbar.
+
+### End Result
+
+At this point, you should have video playing over an SDI port showing up inside your Unreal Engine Level, and you should understand where to set up more advanced features like lens deformation and chroma-keying.
+
+If you're already familiar with the Media Framework, another way you could get video into your Level is to create a new **BlackmagicMediaSource** Asset in your Project, and set it up with the same source properties you set up inside the Media Bundle in the procedure above. Then, create your own **MediaPlayer** and **MediaTexture** Assets to handle the playback of that source in your Level. For details, see the [Media Framework](Engine/MediaFramework) documentation. However, we recommend using the Media Bundle, as shown above, to get the best balance between ease of use and professional quality video features.  
+
+## 3 - Output Captures from the Unreal Editor
+
+In this process, you'll set up an Blackmagic Media Output object, and use the **Media Captures** panel in the Unreal Editor to output the view from selected cameras in the Level to your Blackmagic card.
+
+### Steps
+
+1.  Right-click in the Content Browser, and select **Media > Blackmagic Media Output**.   
+    ![New Blackmagic Media Output](blackmagic-qs-3-1.png "New Blackmagic Media Output")  
+    Name your new Asset **BlackmagicMediaOutput**.
+2.  Double-click your new Asset to open it up for editing. Just like when you created your Blackmagic Media Source, you have to set up the **Configuration** property to control the properties of the video feed that the Unreal Engine sends to your Blackmagic card. Click the arrow to open the submenu, select the options that match your video setup, then click **Apply** in the submenu.   
+    ![Blackmagic Media Output Configuration](blackmagic-qs-3-2.png "Blackmagic Media Output Configuration")  
+    The options you see may vary depending on the devices you have installed. For details on all of the properties you can set in the Blackmagic Media Output, see the [Blackmagic Media Reference](Engine/ProVideoIO/BlackmagicIOReference) page.  
+    Save and close your Media Output when you're done.
+3.  Now we'll place two cameras in the Level, to give us viewpoints for the output we'll send to the Blackmagic card. In the **Modes** panel, open the **Cinematic** tab, and drag two instances of the **Cine Camera Actor** into the viewport.  
+    ![Drag and drop Cine Camera Actors](aja-qs-3-3-cinecameras.png "Drag and drop Cine Camera Actors")  
+    Place the cameras where you want them in the Level, so that they're showing different viewpoints on the scene.
+    [REGION:tip]
+    **Piloting** the camera is a fast and easy way to set its viewpoint exactly the way you want it. See [Pilot Actors in the Viewport](Engine/UI/LevelEditor/Viewports/ActorLocking).
+    [/REGION]
+4.  From the main menu, choose **Window > Media Capture**. You'll use the **Media Capture** window to control when the Editor sends output to your Blackmagic device, and what camera it should use in the Level.  
+    ![Media Capture window](aja-qs-3-4-mediacapture.png "Media Capture window")
+5.  Under the **Media Viewport Capture** area, find the **Viewport Captures** control. Click the **+** button to add a new capture to this list.  
+    ![Add a viewport capture](aja-qs-3-5-newcapture.png "Add a viewport capture")
+6.  Expand the new entry. First, we'll add the cameras that we want to capture from. In the **Locked Camera Actors** control, click the **+** button to add a new entry.  
+    ![Add a camera actor](aja-qs-3-6-newcamera.png "Add a camera actor")  
+    Then, use the drop-down list to choose one of the cameras you placed in the Level.  
+    ![Select the camera actor](aja-qs-3-6-selectcamera.png "Select the camera actor")  
+    Repeat the same steps to add the other camera to the list.
+7.  Now, set up the output that you want to capture these cameras to. Set the **Media Output** control to point to the new Blackmagic Media Output Asset that you created above. You can do this by selecting it in the drop-down list, or drag your Blackmagic Media Output Asset from the Content Browser and drop it into this slot.  
+    ![Set the Blackmagic Media Output](blackmagic-qs-3-7.png "Set the Blackmagic Media Output")  
+    
+8.  At the top of the window, click the **Capture** button.  
+    ![Start capturing](aja-qs-3-8-capture.png "Start capturing")  
+    You'll see a new frame at the bottom of the window that shows a preview of the output being sent to the Blackmagic card. If you have this port hooked up to another downstream device, you should start to see the output coming through.  
+    ![Active Media Capture](blackmagic-qs-output-8.png "Active Media Capture")
+9.  Each camera that you added to the Locked Camera Actors list for this viewport capture is represented by a corresponding button above the video preview. Click the buttons to switch the capture back and forth between the two views.
+
+![Switch Cameras](blackmagic-qs-output-9.png "Switch Cameras")
+
+  
+
+### End Result
+
+Now you've set up the Unreal Editor to stream output from cameras in your Level to a port on your Blackmagic card. Next, we'll see how to use Blueprint scripting to do the same thing in a running Unreal Engine Project.
+
+## 4 - Output Captures at Runtime
+
+The **Media Capture** window that you used in the last section is a practical and easy way to send captures to the Blackmagic card. However, it only works inside the Unreal Editor. To do the same thing when you're running your Project as a standalone application, you'll need to use the Blueprint API provided by the Media Output. In this procedure, we'll set up a simple toggle switch in the Level Blueprint that starts and stops capturing when the player presses a key on the keyboard.
+
+[REGION:tip]
+The [Virtual Studio](Resources/Showcases/VirtualStudio) showcase, available on the **Learn** tab of the Epic Games Launcher, contains a UMG interface widget that demonstrates how you could control capturing from an on-screen user interface.
+[/REGION]
+
+### Steps
+
+1.  From the main toolbar in the Unreal Editor, choose **Blueprints > Open Level Blueprint**.  
+    ![Open Level Blueprint](aja-qs-4-1-levelbp.png "Open Level Blueprint")
+2.  We'll need to start from the Blackmagic Media Output Asset that you've created, where you identify the port you want to output to. In the **Variables** list in the **My Blueprint** panel, click the **+** button to add a new variable.  
+    ![New Variable](aja-qs-4-2-newvar.png "New Variable")
+3.  In the **Details** panel, set the **Variable Name** to **BlackmagicMediaOutput**, and use the **Variable Type** drop-down list to make it a **Blackmagic Media Output Object Reference**.  
+    ![Blackmagic Media Output Object Reference](blackmagic-qs-4-3.png "Blackmagic Media Output Object Reference")  
+    
+4.  Enable the **Instance Editable** setting (1), and compile the Blueprint. Then, in the **Default Value** section, set the variable to point to the Blackmagic Media Output Asset that you created in your Content Browser (2).  
+    ![Set the default value](blackmagic-qs-4-4.png "Set the default value")  
+    
+5.  Press **Ctrl**, and drag the **BlackmagicMediaOutput** from the Variables list in the **My Blueprint** panel into the **Event Graph**.  
+    ![Control+drag the BlackmagicMediaOutput](blackmagic-qs-4-5.png "Control+drag the BlackmagicMediaOutput")  
+    
+6.  Click and drag from the output port of the **BlackmagicMediaOutput** variable node, and choose **Media > Output > Create Media Capture**.   
+    ![Create Media Capture](blackmagic-qs-4-6.png "Create Media Capture")  
+    Hook up your nodes to the **Event BeginPlay** node as shown below:  
+    ![Event Begin Play](blackmagic-qs-4-6b.png "Event Begin Play")  
+    This creates a new Media Capture object from your Blackmagic Media Output. The Media Capture offers two main Blueprint functions that we'll use to control the capturing: **Capture Active Scene Viewport** and **Stop Capture**.
+7.  First, we'll save the new Media Capture object into its own variable, so we can get access to it again elsewhere. Click and drag from the output port of the **Create Media Capture** node, and choose **Promote to Variable**.  
+    ![Promote to variable](aja-qs-4-7-promote.png "Promote to variable")  
+    Rename the new variable **MediaCapture** in the Variables list in the **My Blueprint** panel.
+    [REGION:tip]
+    It's important to save the Media Capture to a variable here. If you don't, the Unreal Engine's garbage collector may destroy it automatically before you're done with it.
+    [/REGION]
+8.  Press **Ctrl** and drag the **MediaCapture** variable into the **Event Graph**.  
+    ![Control+drag the MediaCapture](blackmagic-qs-4-8-mediacapture.png "Control+drag the MediaCapture")  
+    
+9.  Click and drag from the output port of the **MediaCapture** variable node, and choose **Media > Output > Capture Active Scene Viewport**. Do it again, and choose **Media > Output > Stop Capture**.  
+    ![Start and stop capture](aja-qs-4-9-capturestop-1.png "Start and stop capture")
+10.  Right-click in the **Event Graph** and choose **Input > Keyboard Events > P**.   Click and drag from the **Pressed** output of the **P** node and choose **Flow Control > FlipFlop**.   
+    ![FlipFlop](aja-qs-4-10-pflipflop.png "FlipFlop")
+11.  Connect the **A** output of the **FlipFlop** node to the input event of the **Capture Active Scene Viewport** node, and connect the **B** output of the **FlipFlop** node to the input event of the **Stop Capture** node, as shown below:  
+    ![Connect the nodes](aja-qs-4-11-done.png "Connect the nodes")
+12.  Compile and save the Blueprint, and try playing your Project. Click the arrow next to the Play button in the main Toolbar, and choose either the **New Editor Window (PIE)** or **Standalone Game** option.  
+    ![Launch the Project](aja-qs-4-12-pie.png "Launch the Project")  
+    [REGION:note]
+    Video capture from the Editor will only work when you play your Project in a **New Editor Window (PIE)** or a **Standalone Game**. It won't work in the default **Selected Viewport** mode, or in **Simulate** mode.  
+    In addition, the viewport resolution of your Project (that is, the size of the rendered image the Unreal Engine generates each frame) must match the output resolution set in the active Media Profile, so that it will be the right size for the output video feed.
+    [/REGION]After your project starts up, you should be able to press the **P** button on your keyboard to toggle sending the output from the Engine to the Blackmagic card.
+
+### End Result
+
+At this point, you should have a basic idea of how to work with Blackmagic Media Sources, Media Bundles, and the Media Capture system, and you should understand how all of these elements work together to get professional video in and out of the Unreal Engine.
+
+## On Your Own
+
+Now that you've seen the basics of how to get a new Project exchanging video input and output with a Blackmagic card, you can continue learning on your own:
+
+*   Explore the in-engine keying solution in the Material Instance created by your Media Bundle. Try passing some green-screen video into your card's input port, and use the keying controls in the Material Instance to remove the background.
+*   Explore the [Virtual Studio](Resources/Showcases/VirtualStudio) showcase to see what it adds to this basic setup, like its on-screen UI that switches cameras and controls video capture at runtime.
\ No newline at end of file
diff --git a/Engine/Documentation/Source/Engine/ProVideoIO/MediaProfiles/MediaProfiles.INT.udn b/Engine/Documentation/Source/Engine/ProVideoIO/MediaProfiles/MediaProfiles.INT.udn
new file mode 100644
--- /dev/null
+++ b/Engine/Documentation/Source/Engine/ProVideoIO/MediaProfiles/MediaProfiles.INT.udn
@@ -0,0 +1,145 @@
+Availability: Public
+Crumbs: 
+Title: Supporting Multiple Media Configurations
+Description: The Media Profile gathers input, output, timecode and genlock settings into one place for easy configuration. Proxies help route input and output between your media content and the inputs and outputs you set in your Media Profile.
+Type: 
+Version: 4.21
+Parent: Engine/ProVideoIO
+Order: 
+Tags: Guide
+Tags: Early Access
+Tags: professional video
+
+
+If you've followed the Quick Start guide for working with [AJA Media](Engine/ProVideoIO/AjaQuickStart) or [Blackmagic Design](Engine/ProVideoIO/BlackmagicQuickStart) hardware, you've seen one possible way to set up Media Input and Media Capture Assets to get SDI video feeds in and out of your Project.
+
+However, you often need a single Project to work with multiple different media configurations. For example:
+
+*   You may need to work with the Project on different computers, each of which has different hardware or a different wiring setup. For example, one computer might have an AJA card while another has a Blackmagic card; or the video feed that is connected to port 1 on one computer might be wired to port 2 on the other computer.
+*   You may need to be able to switch your Project on one computer between different kinds of sources and outputs. For example, you may want to work with live feeds coming in over the SDI connections, but switch to using pre-recorded files on disk as inputs when the feed is offline.
+*   You may also need to change the way the Unreal Engine reads timecode and genlocks its frame rate, especially when you change to different sources or inputs. For example, you may want to lock the Unreal Engine's frame rate only when using input feeds coming in over SDI connections. 
+
+In cases like these, you don't want to have to make deep changes in your media setup and Project content every time you switch setups. Instead, you can set up multiple Media Profiles, each of which gathers together all configuration settings related to media inputs, outputs, timecode providers, and genlock providers. By simply switching to a different Media Profile, you can immediately change your entire media I/O setup.
+
+## How Media Profiles and Proxies Work
+
+Media Profiles rely on *proxy* Assets to make the connection between two different points:
+
+*   On one hand, the input and output configurations that you set up in the Media Profile.
+*   On the other hand, the places that those inputs and outputs are consumed or generated in your Project content.
+
+[REGION:lightbox]
+[![Media](media_profile_overview.png "Media Profile Overview")](media_profile_overview.png)
+
+*Click for full image.*
+[/REGION]
+
+The following list illustrates the numbered sections in the image above:
+
+1.  You set up multiple Media Profiles in your Project: one for each usage scenario that your Project needs to support. Each profile contains a list of media sources and a list of media outputs. You set up these lists to determine where your Project will get video input from, and where it will send captures to.  
+    For example, you could set up one Profile to get inputs from an AJA Media card, another to get inputs from a Blackmagic Design card, and a third to get inputs from files on disk.  
+    You can also mix and match within a single Media Profile. For example, a single Media Profile could get inputs from both files and SDI feeds.
+2.  When you have the **Media Framework Utilities Plugin** installed, the **Project Settings** panel for your Project contains similar lists of inputs and outputs. Each numbered entry in the input list of the Media Profile matches up with the same number in the input list of the Project Settings, and the same goes for the output lists.  
+    You set up each of these list entries in the Project Settings to point to a Proxy Media Source or Proxy Media Asset that you create in your Project.
+3.  You then refer to these same Proxy Media Assets anytime you need to set up a Media Source or a Media Output in your Project content. For example:
+    1.  When you set up a Media Bundle or Media Player, you point it to one of those Media Proxy Assets instead of pointing it to a Media Source that directly represents a file, stream, or SDI input.
+    2.  Similarly, when you set up a Media Capture, you send the output to a Proxy Media Output instead of directly to an AJA or Blackmagic Media Output. You can do this in the **Media Capture** panel in the Unreal Editor as shown above, or when you invoke the **Create Media Capture** Blueprint node at runtime. If you've set up the output list in your Project Settings to contain a reference to the same Proxy Media Output, the capture is then forwarded on to the output with the matching index that you have configured in the Media Profile.
+
+The Proxy Media Source and Proxy Media Output Assets don't contain any important configuration properties themselves. They only act as a channel — a way to connect the input and output configurations set in the Media Profile with the places that those inputs and outputs are actually used or generated by other Assets elsewhere in the Project.
+
+## Selecting and Saving the Active Media Profile
+
+Whenever your Project is loaded in the Unreal Engine or the Unreal Editor, it can have only one Media Profile active (or none at all).
+
+When you have your Project open in the Unreal Editor, you select the active Media Profile using the Media Profile menu in the Toolbar of the main Level Editor:
+
+![Selecting the active Media Profile](profiles-12-selectfileprofile.png "Selecting the active Media Profile")
+
+Media Profiles are Assets that are always saved in the Unreal Engine Project. However, the choice of which Media Profile is active is saved on each computer, not within the Project. That means that once you've set a Media Profile for a given Project on a given computer, you never have to change it again (unless your media hardware or configuration changes).
+
+On the other hand, when you run a packaged version of your Project, the choice of which Media Profile is active is always determined by the **Startup Media Profile** setting, which you'll find in the **Project Settings** window under **Plugins > Media Profile**.
+
+![](media-profile-setting.png "media-profile-setting.png")
+
+  
+
+## Getting Started with Media Profiles and Proxies
+
+In this procedure, we'll set up two Media Profiles: one that plays input from local movie files, and one that plays live video feeds from an SDI card. We'll set up two separate Media Bundles to play the videos in the Level. Finally, we'll use Proxy Media Sources to redirect the inputs from the sources configured in the Media Profiles on to the Media Bundles.
+
+[REGION:note]
+**Prerequisites**:
+
+*   Media Profiles and Proxy Assets are provided by the **Media Framework Utilities** Plugin. You'll need to have this Plugin installed.
+*   You'll also need to have followed the Quick Start Guide for either [AJA Media](Engine/ProVideoIO/AjaQuickStart) or [Blackmagic Design](Engine/ProVideoIO/BlackmagicQuickStart) hardware.
+[/REGION]
+[REGION:tip]
+To see a working example of a setup very similar to the one described below, see the [Virtual Studio showcase](Resources/Showcases/VirtualStudio), available in the **Learn** tab of the Epic Games Launcher.
+[/REGION]
+
+1.  Start by creating a Proxy Media Source Asset. Right-click in the **Content Browser** and select **Media > Proxy Media Source** from the contextual menu.  
+    ![Create a Proxy Media Source](profiles-1-createproxy.png "Create a Proxy Media Source")  
+    Name the new Asset **VideoProxyInA**.
+2.  Repeat the previous step to create another Proxy Media Source, but name it **VideoProxyInB**.
+3.  Create a new Media Bundle Asset. Right-click in the **Content Browser** and select **Media > Media Bundle**.  
+    ![Create a Media Bundle](profiles-3-createbundle.png "Create a Media Bundle")Name the new Asset **MediaBundleA**.
+4.  Double-click the new Media Bundle Asset to edit it.
+5.  In the **Details** panel, find the **Media Source** setting and select **Proxy Media Source** from the drop-down list.  
+    ![Select Proxy Media Source](profiles-5-setproxysource.png "Select Proxy Media Source")
+6.  Expand the **Source** category, and set the **Proxy** to refer to the **VideoProxyInA** Asset that you created earlier.  
+    ![Set the Proxy Media Source Asset](profiles-6-setproxyasset.png "Set the Proxy Media Source Asset")
+7.  Repeat the previous steps to create another Media Bundle. This time, name it **MediaBundleB**, and set its source proxy to refer to **VideoProxyInB**.
+    [REGION:tip]
+    The Asset names recommended here, like VideoProxyInA and MediaProfileA, are intended to help make the relationships between the various Assets clear. However, in your own Project, we recommend using names that are more descriptive of the kind of content that the Media Proxy and Media Bundle need to handle. These Asset names are visible in your configurations, and in places like the **Media Capture** window, so ideally they should be meaningful to the whole team working on your Project.
+    [/REGION]
+8.  From the main menu, select **Edit > Project Settings**. Find the **Plugins > Media Profile** section, and expand the advanced options.![Media Profile settings](profiles-8-profilesettings.png "Media Profile settings")
+9.  Add two entries to the **Media Source Proxy** list. Set one to refer to **VideoProxyInA**, and one to refer to **VideoProxyInB**.  
+    ![Add the Proxy Media Source Assets](profiles-9-sourceproxylist.png "Add the Proxy Media Source Assets")
+10.  Now, we'll create a new Media Profile named **FileProfile** that plays videos from files on disk. You can do this in either of two ways:
+    *   Right-click in the **Content Browser** and select **Media > Media Profile** from the contextual menu, then rename the Asset.  
+        ![Right-click to create the new Asset](profiles-10a-createasset.png "Right-click to create the new Asset")
+    *   From the profile selection button in the Toolbar, select **Create New Media Profile** and set the path and name for the new Asset.  
+        ![Create the new Asset from the Toolbar](profiles-10b-toolbar.png "Create the new Asset from the Toolbar")
+11.  Double-click the new Media Profile to edit it, if it doesn't open for editing automatically.  
+    Find the **Media Sources** setting, and add two new elements to the list. Set each entry in the list to be a **File Media Source**, and set a different **File Path** for each.   
+    ![File Media Source paths](profiles-11-filepaths.png "File Media Source paths")  
+    [REGION:tip]
+    If you need some sample videos to test with, you can use these: [Video1](https://docs.unrealengine.com/en-US/SampleVideo.mp4) | [Video2](https://dnnrz1gqa.blob.core.windows.net/portals/attachments/Engine/MediaFramework/HowTo/StreamMediaSource/Infiltrator%20Demo.mp4?sv=2017-04-17&sr=b&si=DNNFileManagerPolicy&sig=Op%2FlRJkIe3a2VBZTKj9actY8UKQLgPK20jMoki5xNaU%3D)
+    [/REGION]
+12.  Repeat the previous steps to create a second Media Profile, named **LiveFeedProfile**, that brings video in from the AJA or Blackmagic device connected your computer. For example:  
+    ![Media Profile for live feeds](profiles-14-livefeed.png "Media Profile for live feeds") Each Media Profile also offers you the ability to set a Timecode Provider and a Genlock Provider:  
+    ![Timecode and Genlock settings in the Media Profile](media-profile-timecode-genlock.png "Timecode and Genlock settings in the Media Profile")  
+    These have exactly the same effect as the **TimecodeProvider** and **Custom TimeStep** settings offered in the **Project Settings** panel. However, when you set these values in a Media Profile, they override the **Project Settings** only when that Media Profile is active. For more information, see [Timecode and Genlock](Engine/ProVideoIO/TimecodeGenlock).
+13.  Use the profile selection button in the Toolbar to select the **FileMediaProfile** that you created above.  
+    ![Activate the FileMediaProfile](profiles-12-selectfileprofile.png "Activate the FileMediaProfile")
+14.  Drag and drop your two Media Bundles into your Level Viewport. You should see them begin showing the videos loaded from the files or incoming over the SDI connections.  
+    ![Drag and drop the Media Bundles](profiles-13-dropbundles.png "Drag and drop the Media Bundles")
+15.  Now you can use the Media Profile selection tool in the Toolbar to switch between your two Media Profiles easily, changing your Media Sources (and, optionally, your timecode and genlock settings) in one click.  
+    ![Switch Media Profiles](profiles-15-swap.png "Switch Media Profiles")
+
+### Using Proxy Media Outputs
+
+In this procedure, we'll extend the Media Profile setup we created in the previous section. We'll capture a video feed in the Unreal Editor (or the Unreal Engine at runtime), and route that feed through a Proxy Media Output to an output configuration that we define in the Media Profile. The overall process is very similar to setting up media sources: you create a Proxy Asset, update the Project Settings to refer to the Proxy, and configure the actual output device in the Media Profile.
+
+1.  Start by creating a Proxy Media Output Asset. Right-click in the **Content Browser** and select **Media > Proxy Media Output** from the contextual menu.  
+    ![Create a new Proxy Media Output](proxies-output-1-createproxy.png "Create a new Proxy Media Output")  
+    Name the new Asset **VideoProxyOut**.
+2.  From the main menu, select **Edit > Project Settings**. Find the **Plugins > Media Profile** section, and expand the advanced options.  
+    ![Project Settings for Media Profiles](proxies-output-2-projectsettings.png "Project Settings for Media Profiles")
+3.  Add an entry to the **Media Output Proxy** list, and set its output proxy to point to **VideoProxyOut**.  
+    ![Set the Media Output Proxy](proxies-output-3-proxyoutputasset.png "Set the Media Output Proxy")  
+    
+4.  Double-click the **LiveFeedProfile** you created in the previous section, and add a new entry to the **Media Outputs** list. Set it up to send the video feed to a port on the AJA or Blackmagic device connected to your computer.  
+    ![](proxies-output-4-profile.png "proxies-output-4-profile.png")
+5.  From the main menu, choose **Window > Media Capture**.
+6.  Add a new entry to the **Viewport Captures** list if needed, and set its **Media Output** setting to point to your **VideoProxyOut** Asset.  
+    ![Set the Proxy in the Media Capture](proxies-output-6-mediacapture-1.png "Set the Proxy in the Media Capture")
+7.  Click **Capture** to start capturing the feed and sending it through the proxy to your AJA or Blackmagic device.![Start capturing from the Unreal Editor](aja-qs-3-8-capture.png "Start capturing from the Unreal Editor")   
+    If you switch the current Media Profile at this point to the **FileVideoProfile**, the capture stops, because the **FileVideoProfile** does not have an output configured for entry 0 in its **Media Outputs** list.
+8.  To capture the video feed at runtime and send it through the Proxy, you do it exactly as instructed in the Quick Start Guide for [AJA Media](Engine/ProVideoIO/AjaQuickStart) or [Blackmagic Design](Engine/ProVideoIO/BlackmagicQuickStart) hardware.  
+    The only difference is that instead of using a variable that is an Object Reference to an **AjaMediaOutput** or **BlackmagicMediaOutput**, you use an Object Reference to the **ProxyMediaOutput** that you want to send the captures to. You then use this **ProxyMediaOutput** variable as the input to the **Create Media Capture** node. For example:  
+    ![Create the Media Capture from the Proxy Media Output](proxies-output-8-runtime.png "Create the Media Capture from the Proxy Media Output")
+
+### End Result
+
+You now have two different Media Profiles set up, each with its own configuration of input and output feeds. Each uses Media Proxy Assets to map those input and output configurations to the other Media Framework Assets that consume or produce those video feeds. This example configuration might not match your exact usage scenario. However, having walked through the steps to set it up, you should have a good idea of how to extend the same basic functionality to meet your own needs.
\ No newline at end of file
diff --git a/Engine/Documentation/Source/Engine/ProVideoIO/TimecodeGenlock/TimecodeGenlock.INT.udn b/Engine/Documentation/Source/Engine/ProVideoIO/TimecodeGenlock/TimecodeGenlock.INT.udn
new file mode 100644
--- /dev/null
+++ b/Engine/Documentation/Source/Engine/ProVideoIO/TimecodeGenlock/TimecodeGenlock.INT.udn
@@ -0,0 +1,94 @@
+Availability: Public
+Crumbs: 
+Title: Timecode and Genlock
+Description: Describes how to make the Unreal Engine adopt timecode from an AJA video input, and how to lock the Engine's frame rate to match the source video.
+Type: 
+Version: 4.20
+Parent: Engine/ProVideoIO
+Order: 
+Tags: Guide
+Tags: Early Access
+Tags: professional video
+
+
+By default, the Unreal Editor generates its own timecode, based on the current system time on your computer. This timecode is set by default to run at 30 frames per second, but the engine typically renders many more frames than 30 in any given second. So, typically multiple consecutive frames in the output will get assigned the same timecode value.
+
+When you're working with professional video, you often have to keep timecode in sync between multiple different video feeds and signal processing devices. To do this, you may need to make the Unreal Engine's timecode match the timecode of the video feed you're working with. In some cases, you may want to go even further, and lock the engine so that it only produces one single frame for each frame of video that comes in through a reference input — we refer to this as *genlock*.
+
+This page describes how to make the Unreal Engine adopt the timecode values from a port that you specify on your AJA card instead of generating its own, and how to genlock the Unreal Engine to that timecode signal.
+
+[REGION:note]
+Timecode and genlock management is under active development, and is likely to change in future versions.
+[/REGION]
+
+## Visualizing Timecode in the Unreal Editor
+
+To see the actual timecode values being used by the Unreal Engine while you're working in the Unreal Editor, use the **Timecode Provider** panel. You can find this panel in the main menu under **Window > Developer Tools > Timecode Provider**.
+
+![](timecode-panel.png "timecode-panel.png")
+
+It displays the current timecode, the timecode provider (the source of the timecode values), and the number of timecode frames that are being generated within each second. You can dock this panel anywhere in the UI of the main Level Editor.
+
+Alternatively, you can use the following console command:
+
+    stat timecode
+
+You'll see the values printed in the Unreal Editor viewport, in **HH:MM:SS:FF** format:
+
+![Timecode display in the viewport](tc-genlock-vis-1.png "Timecode display in the viewport")
+
+If you have the Unreal Engine set up to adopt timecode from an input signal, or to genlock to a video signal, you'll also see the source of the timecode and genlock. For example:
+
+![Genlocked and synchronized timecode in the viewport](tc-genlock-vis-2.png "Genlocked and synchronized timecode in the viewport")
+
+You can use this command to confirm that the Unreal Engine is generating frames at the rate you're expecting, and that it's using the same timecode values as your input video.
+
+## Setting up Timecode and Genlock
+
+The following steps show how to make the Unreal Engine adopt its timecode values from an input SDI video feed coming in from an AJA or Blackmagic device, and how to make the Unreal Engine lock its frame rate to that feed so that it generates only one frame of output for each frame of input.
+
+1.  Right-click in the Content Browser, and choose **Create Basic Asset > Blueprint Class**.  
+    ![Create new Blueprint class](tc-genlock-1.png "Create new Blueprint class")
+2.  In the **Pick Parent Class** window, expand the **All Classes** section (1). Search for and select the **AjaTimecodeProvider** class (pictured) or **BlackmagicTimecodeProvider** class (2), then click **Select** (3).  
+    ![TimecodeProvider parent classes](timecode-2-timecodeproviderparent.png "TimecodeProvider parent classes")  
+    Back in the Content Browser, give your new Asset a descriptive name, such as **AJA\_Timecode\_Port** or **Blackmagic\_Timecode\_Port**.
+3.  Double-click your new Asset to open it up in the Blueprint Editor, and set up the properties in the **Details** panel.![Timecode Provider settings](timecode-3-timecodeprovidersettings.png "Timecode Provider settings")
+    *   To read the timecode from an incoming video stream, set the **Video ****Configuration** or **Media Configuration** to point to the port on your AJA or Blackmagic card, and (if applicable) the type of timecode to read from the feed (LTC or VITC).
+    *   If you are using an AJA card, and you want to read the timecode from a reference port on the card rather than from a video stream, enable **Use Reference In** and set the **Reference Configuration** instead of the **Video Configuration**.
+4.  **Compile** and **Save** your Blueprint class, and close the Blueprint Editor.
+5.  If you want to genlock the Unreal Engine to your video source, repeat the previous instructions to create and set up another Blueprint class, but this time use **AjaCustomTimeStep** or **BlackmagicCustomTimeStep** as the parent class.  
+    ![Custom Time Step Provider parent classes](timecode-5-timestepparents.png "Custom Time Step Provider parent classes")
+6.  This class requires similar settings to the **AjaTimecodeProvider** or **BlackmagicTimecodeProvider**:  
+    ![Custom Time Step Provider settings](timecode-6-timestepsettings.png "Custom Time Step Provider settings")  
+    AJA cards also offer additional settings; see the tooltips for details.  
+    
+7.  From the main menu, choose **Edit > Project Settings**.  
+    In the **Project Settings** window, open the **Engine > General Settings** category (1), and find the **Timecode** section (2). From the **TimecodeProvider** drop-down list, select the timecode port Asset that you created.  
+    ![Set the TimecodeProvider in the Project Settings](tc-genlock-6.png "Set the TimecodeProvider in the Project Settings")
+8.  If you want to genlock the Unreal Engine to your video source, find the **Framerate** section above, and expand the advanced properties at the bottom of the section. From the **Custom TimeStep** drop-down list, select the genlock port Asset that you created.  
+    ![Set the Custom TimeStep in the Project Settings](tc-genlock-7.png "Set the Custom TimeStep in the Project Settings")
+9.  Close the Unreal Editor and restart your Project.
+
+[REGION:tip]
+If you use Media Profiles to set up multiple different media configurations for your Project, you can also override these default Project-level settings for Timecode and Timestep Providers in each Media Profile. For more information, see [Supporting Multiple Media Configurations](Engine/ProVideoIO/MediaProfiles).
+[/REGION]
+
+### End Result
+
+As long as you have video input coming in to the port you set up for the timecode provider and custom timestep classes, and that video has the resolution, frame rate and timecode format that you specified, the Unreal Engine should adopt your video's timecode and lock its frame update frequency so that it generates one output frame for each input frame.
+
+To verify that this is working, use one of the options described under [Visualizing Timecode in the Unreal Editor](#visualizingtimecodeintheunrealeditor).
+
+## Timecode Texel Encoding
+
+For debugging purposes, the Media Source and Media Output Assets for AJA and Blackmagic devices can encode the timecode of each frame into the video image by illuminating pixels in white. This can be useful if you need a visual way to verify that a given frame of video actually matches the timecode signal that accompanies it.
+
+Timecode is usually expressed in the format **HH:MM::SS:FF**, for a total of twelve digits. Each of these digits is assigned a row in the texture output, starting from the top of the image. Within each of those rows, count from the left of the image to value of that timecode digit to find the pixel that represents the timecode value. The leftmost pixel is pixel number 0.
+
+So, for example, given the timecode value 12:08:44:12:
+
+*   On the first row, the second pixel from the left is lit (pixel number 1).
+*   On the second row, the third pixel from the left is lit (pixel number 2).
+*   On the third row, the leftmost pixel from the left is lit (pixel number 0).
+*   On the fourth row, the ninth pixel from the left is lit (pixel number 8).
+*   and so on.
\ No newline at end of file
