Update the tooltip on the F, T and V Interp functions to be explict about what happens when you pass in 0 speed

diff --git a/Engine/Documentation/Source/Shared/GraphNodes/Blueprint/UKismetMathLibrary/UKismetMathLibrary.INT.udn b/Engine/Documentation/Source/Shared/GraphNodes/Blueprint/UKismetMathLibrary/UKismetMathLibrary.INT.udn
--- a/Engine/Documentation/Source/Shared/GraphNodes/Blueprint/UKismetMathLibrary/UKismetMathLibrary.INT.udn
+++ b/Engine/Documentation/Source/Shared/GraphNodes/Blueprint/UKismetMathLibrary/UKismetMathLibrary.INT.udn
@@ -224,7 +224,7 @@ Return Value - This outputs A's string if "Pick A" was true, otherwise it output
 
 The "FInterp to" (float interpolation) node allows you to smoothly transition a float value to a different target value. When utilized, this node will take your current float value and (based on a time-step and speed) output a new float approaching (or equal to) the target value.
 
-NOTE: That the output value should be used as the "Current" input for the following frame. If the input value sits unchanged, then this node will make the same step each frame (resulting in the same output every time, failing to ever reach the "Target" value).
+NOTE: That the output value should be used as the "Current" input for the following frame. If the input value sits unchanged, then this node will make the same step each frame (resulting in the same output every time, failing to ever reach the "Target" value). If a Speed of 0 is given, then the result will be the target. 
 
 
 ---
@@ -252,7 +252,7 @@ Return Value - This outputs an interpolated float value approaching (or equal to
 
 The "RInterp to" (Rotator interpolation) node allows you to smoothly transition from one Rotator to another. When utilized, this node will take your current Rotator and (based on a time-step and speed) output a new Rotator approaching (or equal to) the target one.
 
-NOTE: That the output value should be used as the "Current" input for the following frame. If the input value sits unchanged, then this node will make the same step each frame (resulting in the same output every time, failing to ever reach the "Target" value).
+NOTE: That the output value should be used as the "Current" input for the following frame. If the input value sits unchanged, then this node will make the same step each frame (resulting in the same output every time, failing to ever reach the "Target" value). If a Speed of 0 is given, then the result will be the target. 
 
 
 ---
@@ -280,7 +280,7 @@ Return Value - This outputs an interpolated Rotator approaching (or equal to) th
 
 The "TInterp to" (Transform interpolation) node allows you to smoothly transition from one Transform to another. When utilized, this node will take your current Transform and (based on a time-step and speed) output a new Transform approaching (or equal to) the target one.
 
-NOTE: That the output value should be used as the "Current" input for the following frame. If the input value sits unchanged, then this node will make the same step each frame (resulting in the same output every time, failing to ever reach the "Target" value).
+NOTE: That the output value should be used as the "Current" input for the following frame. If the input value sits unchanged, then this node will make the same step each frame (resulting in the same output every time, failing to ever reach the "Target" value). If a Speed of 0 is given, then the result will be the target. 
 
 
 ---
@@ -308,7 +308,7 @@ Return Value - This outputs an interpolated Transform approaching (or equal to)
 
 The "VInterp to" (Vector interpolation) node allows you to smoothly transition from one Vector to another. When utilized, this node will take your current Vector and (based on a time-step and speed) output a new Vector approaching (or equal to) the target one.
 
-NOTE: That the output value should be used as the "Current" input for the following frame. If the input value sits unchanged, then this node will make the same step each frame (resulting in the same output every time, failing to ever reach the "Target" value).
+NOTE: That the output value should be used as the "Current" input for the following frame. If the input value sits unchanged, then this node will make the same step each frame (resulting in the same output every time, failing to ever reach the "Target" value). If a Speed of 0 is given, then the result will be the target.  If a Speed of 0 is given, then the result will be the target. 
 
 
 ---
diff --git a/Engine/Source/Runtime/Engine/Classes/Kismet/KismetMathLibrary.h b/Engine/Source/Runtime/Engine/Classes/Kismet/KismetMathLibrary.h
--- a/Engine/Source/Runtime/Engine/Classes/Kismet/KismetMathLibrary.h
+++ b/Engine/Source/Runtime/Engine/Classes/Kismet/KismetMathLibrary.h
@@ -3265,11 +3265,11 @@ class ENGINE_API UKismetMathLibrary : public UBlueprintFunctionLibrary
 
 	/**
 	 * Tries to reach Target based on distance from Current position, giving a nice smooth feeling when tracking a position.
-	 *
+	 * 
 	 * @param		Current			Actual position
 	 * @param		Target			Target position
 	 * @param		DeltaTime		Time since last tick
-	 * @param		InterpSpeed		Interpolation speed
+	 * @param		InterpSpeed		Interpolation speed, if the speed given is 0, then jump to the target.
 	 * @return		New interpolated position
 	 */
 	UFUNCTION(BlueprintPure, Category="Math|Interpolation")
@@ -3293,7 +3293,7 @@ class ENGINE_API UKismetMathLibrary : public UBlueprintFunctionLibrary
 	 * @param		Current			Actual rotation
 	 * @param		Target			Target rotation
 	 * @param		DeltaTime		Time since last tick
-	 * @param		InterpSpeed		Interpolation speed
+	 * @param		InterpSpeed		Interpolation speed, if the speed given is 0, then jump to the target.
 	 * @return		New interpolated position
 	 */
 	UFUNCTION(BlueprintPure, Category="Math|Interpolation", meta=(Keywords="rotation rotate"))
