Documentation for build tools

diff --git a/Engine/Documentation/Source/Programming/Development/BuildGraph/ScriptAnatomy/Tasks/PredefinedTasks/BuildGraphScriptPredefinedTasks.INT.udn b/Engine/Documentation/Source/Programming/Development/BuildGraph/ScriptAnatomy/Tasks/PredefinedTasks/BuildGraphScriptPredefinedTasks.INT.udn
new file mode 100644
--- /dev/null
+++ b/Engine/Documentation/Source/Programming/Development/BuildGraph/ScriptAnatomy/Tasks/PredefinedTasks/BuildGraphScriptPredefinedTasks.INT.udn
@@ -0,0 +1,349 @@
+Availability: NoPublish
+Title: BuildGraph Predefined Tasks
+Crumbs: %ROOT%, Programming, Programming/Development, Programming/Development/BuildGraph, Programming/Development/BuildGraph/BuildGraphScriptTasks
+Description: This is a procedurally generated markdown page.
+version: 4.24
+parent:Programming/Development/BuildGraph/BuildGraphScriptTasks
+
+### AgeStore
+
+Task that strips symbols from a set of files. This task is named after the AGESTORE utility that comes with the Microsoft debugger tools SDK, but is actually a separate implementation. The main difference is that it uses the last modified time rather than last access time to determine which files to delete.
+
+|          |                      |          |                                                                                                                                                                                                                                                                                                                                     |
+| -------- | -------------------- | -------- | ----------------------------------------------------------------------------------------------------------------------------------------------------------------------------------------------------------------------------------------------------------------------------------------------------------------------------------- |
+| Platform | UnrealTargetPlatform | Required | The target platform to age symbols for.                                                                                                                                                                                                                                                                                             |
+| StoreDir | String               | Required | The symbol server directory.                                                                                                                                                                                                                                                                                                        |
+| Days     | Int32                | Required | Number of days worth of symbols to keep.                                                                                                                                                                                                                                                                                            |
+| BuildDir | String               | Optional | The root of the build directory to check for existing buildversion named directories.                                                                                                                                                                                                                                               |
+| Filter   | String               | Optional | A substring to match in directory file names before deleting symbols. This allows the "age store" task to avoid deleting symbols from other builds in the case where multiple builds share the same symbol server. Specific use of the filter value is determined by the symbol server structure defined by the platform toolchain. |
+| If       | Condition            | Optional | Whether to execute this task. It is ignored if this condition evaluates to false.                                                                                                                                                                                                                                                   |
+
+### Command
+
+Invokes an AutomationTool child process to run the given command.
+
+|                          |           |          |                                                                                                                                                                   |
+| ------------------------ | --------- | -------- | ----------------------------------------------------------------------------------------------------------------------------------------------------------------- |
+| Name                     | String    | Required | The command name to execute.                                                                                                                                      |
+| Arguments                | String    | Optional | Arguments to be passed to the command.                                                                                                                            |
+| MergeTelemetryWithPrefix | String    | Optional | If non-null, instructs telemetry from the command to be merged into the telemetry for this UAT instance with the given prefix. May be an empty (non-null) string. |
+| If                       | Condition | Optional | Whether to execute this task. It is ignored if this condition evaluates to false.                                                                                 |
+
+### Commandlet
+
+Spawns the editor to run a commandlet.
+
+|           |               |          |                                                                                                          |
+| --------- | ------------- | -------- | -------------------------------------------------------------------------------------------------------- |
+| Name      | String        | Required | The commandlet name to execute.                                                                          |
+| Project   | File Spec     | Optional | The project to run the editor with.                                                                      |
+| Arguments | String        | Optional | Arguments to be passed to the commandlet.                                                                |
+| EditorExe | FileReference | Optional | The editor executable to use. Defaults to the development UE4Editor executable for the current platform. |
+| If        | Condition     | Optional | Whether to execute this task. It is ignored if this condition evaluates to false.                        |
+
+### Compile
+
+Compiles a target with UnrealBuildTool.
+
+|                       |                           |          |                                                                                                                                  |
+| --------------------- | ------------------------- | -------- | -------------------------------------------------------------------------------------------------------------------------------- |
+| Target                | String                    | Required | The target to compile.                                                                                                           |
+| Configuration         | UnrealTargetConfiguration | Required | The configuration to compile.                                                                                                    |
+| Platform              | UnrealTargetPlatform      | Required | The platform to compile for.                                                                                                     |
+| Project               | String                    | Optional | The project to compile with.                                                                                                     |
+| Arguments             | String                    | Optional | Additional arguments for UnrealBuildTool.                                                                                        |
+| AllowXGE              | Boolean                   | Optional | Whether to allow using XGE for compilation.                                                                                      |
+| AllowParallelExecutor | Boolean                   | Optional | Whether to allow using the parallel executor for this compile.                                                                   |
+| Clean                 | Nullable`1                | Optional | Whether to allow cleaning this target. If unspecified, targets are cleaned if the -Clean argument is passed on the command line. |
+| Tag                   | Tag List                  | Optional | Tag to be applied to build products of this task.                                                                                |
+| If                    | Condition                 | Optional | Whether to execute this task. It is ignored if this condition evaluates to false.                                                |
+
+### Cook
+
+Cook a selection of maps for a certain platform
+
+|           |           |          |                                                                                   |
+| --------- | --------- | -------- | --------------------------------------------------------------------------------- |
+| Project   | String    | Required | Project file to be cooked.                                                        |
+| Platform  | String    | Required | The cook platform to target (for example, WindowsNoEditor).                       |
+| Maps      | String    | Optional | List of maps to be cooked, separated by '+' characters.                           |
+| Versioned | Boolean   | Optional | Additional arguments to be passed to the cooker.                                  |
+| Arguments | String    | Optional | Additional arguments to be passed to the cooker.                                  |
+| Tag       | Tag List  | Optional | Tag to be applied to build products of this task.                                 |
+| If        | Condition | Optional | Whether to execute this task. It is ignored if this condition evaluates to false. |
+
+### Copy
+
+Copies files from one directory to another.
+
+|           |           |          |                                                                                   |
+| --------- | --------- | -------- | --------------------------------------------------------------------------------- |
+| Files     | File Spec | Optional | Optional filter to be applied to the list of input files.                         |
+| From      | File Spec | Required | The pattern(s) to copy from (for example, Engine/*.txt).                          |
+| To        | File Spec | Required | The directory to copy to.                                                         |
+| Overwrite | Boolean   | Optional | Whether or not to overwrite existing files.                                       |
+| Tag       | Tag List  | Optional | Tag to be applied to build products of this task.                                 |
+| If        | Condition | Optional | Whether to execute this task. It is ignored if this condition evaluates to false. |
+
+### CsCompile
+
+Compiles C# project files, and their dependencies.
+
+|               |           |          |                                                                                                                                                     |
+| ------------- | --------- | -------- | --------------------------------------------------------------------------------------------------------------------------------------------------- |
+| Project       | String    | Required | The C# project file to compile. Using semicolons, more than one project file can be specified.                                                      |
+| Configuration | String    | Optional | The configuration to compile.                                                                                                                       |
+| Platform      | String    | Optional | The platform to compile.                                                                                                                            |
+| Target        | String    | Optional | The target to build.                                                                                                                                |
+| Arguments     | String    | Optional | Additional options to pass to the compiler.                                                                                                         |
+| EnumerateOnly | Boolean   | Optional | Only enumerate build products -- do not actually compile the projects.                                                                              |
+| Tag           | Tag List  | Optional | Tag to be applied to build products of this task.                                                                                                   |
+| TagReferences | Tag List  | Optional | Tag to be applied to any non-private references the projects have. (for example, those that are external and not copied into the output directory). |
+| If            | Condition | Optional | Whether to execute this task. It is ignored if this condition evaluates to false.                                                                   |
+
+### Delete
+
+Delete a set of files.
+
+|                        |           |          |                                                                                                                     |
+| ---------------------- | --------- | -------- | ------------------------------------------------------------------------------------------------------------------- |
+| Files                  | File Spec | Required | List of file specifications separated by semicolons (for example, *.cpp;Engine/.../*.bat), or the name of a tag set |
+| DeleteEmptyDirectories | Boolean   | Optional | Whether to delete empty directories after deleting the files. Defaults to true.                                     |
+| If                     | Condition | Optional | Whether to execute this task. It is ignored if this condition evaluates to false.                                   |
+
+### Log
+
+Print a message (and other optional diagnostic information) to the output log.
+
+|                 |           |          |                                                                                     |
+| --------------- | --------- | -------- | ----------------------------------------------------------------------------------- |
+| Message         | String    | Optional | Message to print out.                                                               |
+| Files           | File Spec | Optional | If specified, causes the given list of files to be printed after the given message. |
+| IncludeContents | Boolean   | Optional | If specified, causes the contents of the given files to be printed out.             |
+| If              | Condition | Optional | Whether to execute this task. It is ignored if this condition evaluates to false.   |
+
+### ModifyConfig
+
+Modifies a config file
+
+|         |           |          |                                                                                   |
+| ------- | --------- | -------- | --------------------------------------------------------------------------------- |
+| File    | File Spec | Required | Path to the config file                                                           |
+| Section | String    | Required | The section name to modify                                                        |
+| Key     | String    | Required | The property name to set                                                          |
+| Value   | String    | Required | The property value to set                                                         |
+| Tag     | Tag List  | Optional | Tag to be applied to the extracted files                                          |
+| If      | Condition | Optional | Whether to execute this task. It is ignored if this condition evaluates to false. |
+
+### Move
+
+Moves files from one directory to another.
+
+|       |           |          |                                                                                   |
+| ----- | --------- | -------- | --------------------------------------------------------------------------------- |
+| Files | File Spec | Optional | Optional filter to be applied to the list of input files.                         |
+| From  | File Spec | Required | The pattern(s) to copy from (for example, Engine/*.txt).                          |
+| To    | File Spec | Required | The directory to copy to.                                                         |
+| Tag   | Tag List  | Optional | Tag to be applied to build products of this task.                                 |
+| If    | Condition | Optional | Whether to execute this task. It is ignored if this condition evaluates to false. |
+
+### MsBuild
+
+Executes MsBuild
+
+|               |           |          |                                                                                                |
+| ------------- | --------- | -------- | ---------------------------------------------------------------------------------------------- |
+| Project       | String    | Required | The C# project file to compile. Using semicolons, more than one project file can be specified. |
+| Configuration | String    | Optional | The configuration to compile.                                                                  |
+| Platform      | String    | Optional | The platform to compile.                                                                       |
+| Arguments     | String    | Optional | Additional options to pass to the compiler.                                                    |
+| Verbosity     | String    | Optional | The MSBuild output verbosity.                                                                  |
+| If            | Condition | Optional | Whether to execute this task. It is ignored if this condition evaluates to false.              |
+
+### PakFile
+
+Creates a PAK file from a given set of files.
+
+|              |               |          |                                                                                                                                   |
+| ------------ | ------------- | -------- | --------------------------------------------------------------------------------------------------------------------------------- |
+| Files        | File Spec     | Required | List of files, wildcards, and tag sets to add to the pak file, separated by ';' characters.                                       |
+| Output       | FileReference | Required | PAK file to output.                                                                                                               |
+| ResponseFile | FileReference | Optional | Path to a Response File that contains a list of files to add to the pak file -- instead of specifying them individually.          |
+| RebaseDir    | HashSet`1     | Optional | Directories to rebase the files relative to. If specified, the shortest path under a listed directory will be used for each file. |
+| Order        | FileReference | Optional | Script that gives the order of files.                                                                                             |
+| Sign         | FileReference | Optional | Encryption keys for this pak file.                                                                                                |
+| Compress     | Boolean       | Optional | Whether to compress files.                                                                                                        |
+| Arguments    | String        | Optional | Additional arguments to pass to UnrealPak.                                                                                        |
+| Tag          | Tag List      | Optional | Tag to be applied to build products of this task.                                                                                 |
+| If           | Condition     | Optional | Whether to execute this task. It is ignored if this condition evaluates to false.                                                 |
+
+### Rename
+
+Renames a file, or group of files.
+
+|       |           |          |                                                                                                          |
+| ----- | --------- | -------- | -------------------------------------------------------------------------------------------------------- |
+| Files | File Spec | Required | The file or files to rename.                                                                             |
+| From  | String    | Optional | The current file name, or pattern to match (for example, *.txt). Should not include any path separators. |
+| To    | String    | Required | The new name for the file(s). Should not include any path separators.                                    |
+| Tag   | Tag List  | Optional | Tag to be applied to the renamed files.                                                                  |
+| If    | Condition | Optional | Whether to execute this task. It is ignored if this condition evaluates to false.                        |
+
+### SanitizeReceipt
+
+Task that tags build products and/or runtime dependencies by reading from *.target files.
+
+|           |                    |          |                                                                                                                                                 |
+| --------- | ------------------ | -------- | ----------------------------------------------------------------------------------------------------------------------------------------------- |
+| Files     | File Spec          | Required | Set of receipt files (*.target) to read, including wildcards and tag names, separated by semicolons.                                            |
+| EngineDir | DirectoryReference | Optional | Path to the Engine folder, used to expand $(EngineDir) properties in receipt files. Defaults to the Engine directory for the current workspace. |
+| If        | Condition          | Optional | Whether to execute this task. It is ignored if this condition evaluates to false.                                                               |
+
+### SetVersion
+
+Updates the local version files (Engine/Source/Runtime/Launch/Resources/Version.h, Engine/Build/Build.version, and Engine/Source/Programs/DotNETCommon/Metadata.cs) with the given version information.
+
+|                  |           |          |                                                                                                                    |
+| ---------------- | --------- | -------- | ------------------------------------------------------------------------------------------------------------------ |
+| Change           | Int32     | Required | The changelist to set in the version files.                                                                        |
+| CompatibleChange | Int32     | Optional | The engine compatible changelist to set in the version files.                                                      |
+| Branch           | String    | Required | The branch string.                                                                                                 |
+| Build            | String    | Optional | The build version string.                                                                                          |
+| Licensee         | Boolean   | Optional | Whether to set the IS_LICENSEE_VERSION flag to true.                                                               |
+| Promoted         | Boolean   | Optional | Whether to set the ENGINE_IS_PROMOTED_BUILD flag to true.                                                          |
+| SkipWrite        | Boolean   | Optional | If set, do not write to the files -- just return the version files that would be updated. Useful for local builds. |
+| Tag              | Tag List  | Optional | Tag to be applied to build products of this task.                                                                  |
+| If               | Condition | Optional | Whether to execute this task. It is ignored if this condition evaluates to false.                                  |
+
+### Sign
+
+Signs a set of executable files with an installed certificate.
+
+|       |           |          |                                                                                                                      |
+| ----- | --------- | -------- | -------------------------------------------------------------------------------------------------------------------- |
+| Files | File Spec | Required | List of file specifications separated by semicolons (for example, *.cpp;Engine/.../*.bat), or the name of a tag set. |
+| Tag   | Tag List  | Optional | Tag to be applied to build products of this task.                                                                    |
+| If    | Condition | Optional | Whether to execute this task. It is ignored if this condition evaluates to false.                                    |
+
+### Spawn
+
+Spawns an external executable and waits for it to complete.
+
+|            |           |          |                                                                                   |
+| ---------- | --------- | -------- | --------------------------------------------------------------------------------- |
+| Exe        | String    | Required | Executable to spawn.                                                              |
+| Arguments  | String    | Optional | Arguments for the newly created process.                                          |
+| ErrorLevel | Int32     | Optional | The minimum exit code, which is treated as an error.                              |
+| If         | Condition | Optional | Whether to execute this task. It is ignored if this condition evaluates to false. |
+
+### Stage
+
+Stages files listed in a build receipt to an output directory.
+
+|               |                           |          |                                                                                   |
+| ------------- | ------------------------- | -------- | --------------------------------------------------------------------------------- |
+| Project       | FileReference             | Optional | The project that this target belongs to.                                          |
+| Target        | String                    | Required | Name of the target to stage.                                                      |
+| Platform      | UnrealTargetPlatform      | Required | Platform to stage.                                                                |
+| Configuration | UnrealTargetConfiguration | Required | Configuration to be staged.                                                       |
+| Architecture  | String                    | Optional | Architecture to be staged.                                                        |
+| ToDir         | DirectoryReference        | Required | Directory that the receipt files should be staged to.                             |
+| Overwrite     | Boolean                   | Optional | Whether to overwrite existing files.                                              |
+| Tag           | Tag List                  | Optional | Tag to be applied to build products of this task.                                 |
+| If            | Condition                 | Optional | Whether to execute this task. It is ignored if this condition evaluates to false. |
+
+### Strip
+
+Strips debugging information from a set of files.
+
+|           |                      |          |                                                                                                                |
+| --------- | -------------------- | -------- | -------------------------------------------------------------------------------------------------------------- |
+| Platform  | UnrealTargetPlatform | Required | The platform toolchain to strip binaries.                                                                      |
+| BaseDir   | DirectoryReference   | Optional | The directory to find files in.                                                                                |
+| Files     | File Spec            | Required | List of file specifications separated by semicolons (for example, Engine/.../*.pdb), or the name of a tag set. |
+| OutputDir | DirectoryReference   | Optional | Output directory for the stripped files. Defaults to the input path, overwriting the input files.              |
+| Tag       | Tag List             | Optional | Tag to be applied to build products of this task.                                                              |
+| If        | Condition            | Optional | Whether to execute this task. It is ignored if this condition evaluates to false.                              |
+
+### Submit
+
+Creates a new changelist and submits a set of files to a Perforce stream.
+
+|                 |                    |          |                                                                                                                                                                              |
+| --------------- | ------------------ | -------- | ---------------------------------------------------------------------------------------------------------------------------------------------------------------------------- |
+| Description     | String             | Required | The description for the submitted changelist.                                                                                                                                |
+| Files           | File Spec          | Required | The files to submit.                                                                                                                                                         |
+| FileType        | String             | Optional | The Perforce file type for the submitted files (for example, binary+FS32).                                                                                                   |
+| Workspace       | String             | Optional | The workspace name. If specified, a new workspace will be created using the given stream and root directory to submit the files. If not, the current workspace will be used. |
+| Stream          | String             | Optional | The stream for the workspace -- defaults to the current stream. Ignored unless the Workspace attribute is also specified.                                                    |
+| RootDir         | DirectoryReference | Optional | Root directory for the stream. If not specified, defaults to the current root directory.                                                                                     |
+| RevertUnchanged | Boolean            | Optional | Whether to revert unchanged files before attempting to submit.                                                                                                               |
+| Force           | Boolean            | Optional | Force the submit to happen -- even if a resolve is needed (always accept current version).                                                                                   |
+| If              | Condition          | Optional | Whether to execute this task. It is ignored if this condition evaluates to false.                                                                                            |
+
+### SymStore
+
+Task that strips symbols from a set of files.
+
+|              |                      |          |                                                                                                                                               |
+| ------------ | -------------------- | -------- | --------------------------------------------------------------------------------------------------------------------------------------------- |
+| Platform     | UnrealTargetPlatform | Required | The platform toolchain required to handle symbol files.                                                                                       |
+| Files        | String               | Required | List of output files. PDBs will be extracted from this list.                                                                                  |
+| StoreDir     | String               | Required | Output directory for the compressed symbols.                                                                                                  |
+| Product      | String               | Required | Name of the product for the symbol store records.                                                                                             |
+| BuildVersion | String               | Optional | BuildVersion associated with these symbols. Used for clean-up in AgeStore by matching this version against a directory name in a build share. |
+| If           | Condition            | Optional | Whether to execute this task. It is ignored if this condition evaluates to false.                                                             |
+
+### Tag
+
+Applies a tag to a given set of files. The list of files is found by enumerating the tags and file specifications given by the 'Files' parameter. From this list, any files not matched by the 'Filter' parameter are removed, followed by any files matched by the 'Except' parameter.
+
+|           |                    |          |                                                                                                                                                                                                                           |
+| --------- | ------------------ | -------- | ------------------------------------------------------------------------------------------------------------------------------------------------------------------------------------------------------------------------- |
+| BaseDir   | DirectoryReference | Optional | Set the base directory to resolve relative paths and patterns against. If set, any absolute patterns (for example, /Engine/Build/...) are taken to be relative to this path. If not, they are taken to be truly absolute. |
+| Files     | File Spec          | Optional | Set of files to work from, including wildcards and tag names, separated by semicolons. If set, resolved relative to BaseDir, otherwise resolved to the branch root directory.                                             |
+| FileLists | File Spec          | Optional | Set of text files to add additional files from. Each file list should have one file per line.                                                                                                                             |
+| Filter    | File Spec          | Optional | Patterns to filter the list of files by, including tag names or wildcards. If set, may include patterns that apply to the base directory. If not specified, defaults to all files.                                        |
+| Except    | File Spec          | Optional | Set of patterns to exclude from the matched list. May include tag names of patterns that apply to the base directory.                                                                                                     |
+| With      | Tag List           | Required | Name of the tag to apply.                                                                                                                                                                                                 |
+| If        | Condition          | Optional | Whether to execute this task. It is ignored if this condition evaluates to false.                                                                                                                                         |
+
+### TagReceipt
+
+Task that tags build products and/or runtime dependencies by reading from *.target files.
+
+|                     |                    |          |                                                                                                                                                                 |
+| ------------------- | ------------------ | -------- | --------------------------------------------------------------------------------------------------------------------------------------------------------------- |
+| Files               | File Spec          | Required | Set of receipt files (*.target) to read, including wildcards and tag names, separated by semicolons.                                                            |
+| EngineDir           | DirectoryReference | Optional | Path to the Engine folder, used to expand $(EngineDir) properties in receipt files. Defaults to the Engine directory for the current workspace.                 |
+| ProjectDir          | DirectoryReference | Optional | Path to the project folder, used to expand $(ProjectDir) properties in receipt files. Defaults to the Engine directory for the current workspace -- DEPRECATED. |
+| BuildProducts       | Boolean            | Optional | Whether to tag the Build Products listed in receipts.                                                                                                           |
+| BuildProductType    | String             | Optional | Which type of Build Products to tag (see TargetReceipt.cs - UnrealBuildTool.BuildProductType for valid values).                                                 |
+| RuntimeDependencies | Boolean            | Optional | Whether to tag the Runtime Dependencies listed in receipts.                                                                                                     |
+| StagedFileType      | String             | Optional | Which type of Runtime Dependencies to tag (see TargetReceipt.cs - UnrealBuildTool.StagedFileType for valid values).                                             |
+| With                | Tag List           | Required | Name of the tag to apply.                                                                                                                                       |
+| If                  | Condition          | Optional | Whether to execute this task. It is ignored if this condition evaluates to false.                                                                               |
+
+### Unzip
+
+Extract files from a zip archive.
+
+|         |                    |          |                                                                                   |
+| ------- | ------------------ | -------- | --------------------------------------------------------------------------------- |
+| ZipFile | File Spec          | Required | Path to the zip file to extract.                                                  |
+| ToDir   | DirectoryReference | Required | Output directory for the extracted files.                                         |
+| Tag     | Tag List           | Optional | Tag to be applied to the extracted files.                                         |
+| If      | Condition          | Optional | Whether to execute this task. It is ignored if this condition evaluates to false. |
+
+### Zip
+
+Compresses files into a zip archive.
+
+|         |                    |          |                                                                                                                                                             |
+| ------- | ------------------ | -------- | ----------------------------------------------------------------------------------------------------------------------------------------------------------- |
+| FromDir | DirectoryReference | Required | The directory to read compressed files from.                                                                                                                |
+| Files   | File Spec          | Optional | List of file specifications separated by semicolons (for example, *.cpp;Engine/.../*.bat), or the name of a tag set. Relative paths are taken from FromDir. |
+| ZipFile | FileReference      | Required | The zip file to create.                                                                                                                                     |
+| Tag     | Tag List           | Optional | Tag to be applied to the created zip file.                                                                                                                  |
+| If      | Condition          | Optional | Whether to execute this task. It is ignored if this condition evaluates to false.                                                                           |
+
diff --git a/Engine/Documentation/Source/Programming/UnrealBuildSystem/Configuration/BuildConfigProperties/BuildConfigProperties.INT.udn b/Engine/Documentation/Source/Programming/UnrealBuildSystem/Configuration/BuildConfigProperties/BuildConfigProperties.INT.udn
new file mode 100644
--- /dev/null
+++ b/Engine/Documentation/Source/Programming/UnrealBuildSystem/Configuration/BuildConfigProperties/BuildConfigProperties.INT.udn
@@ -0,0 +1,578 @@
+Availability: NoPublish
+Title: Build Configuration Properties Page
+Crumbs:
+Description: This is a procedurally generated markdown page.
+Version: 4.24
+
+### BuildConfiguration
+
+$ bGeneratedSYMFile : Whether to generate dSYM files. Lists Architectures that you want to build.
+
+$ bEnableAddressSanitizer : Enables address sanitizer (ASan)
+
+$ bEnableThreadSanitizer : Enables thread sanitizer (TSan)
+
+$ bEnableUndefinedBehaviorSanitizer : Enables undefined behavior sanitizer (UBSan)
+
+$ BaseLogFileName : Specifies the file to use for logging.
+
+$ bAllowDistccLocalFallback : When enabled, allows DMUCS/Distcc to fallback to local compilation when remote compiling fails. Defaults to true as separation of pre-process and compile stages can introduce non-fatal errors.
+
+$ bVerboseDistccOutput : When true enable verbose distcc output to aid debugging.
+
+$ DistccExecutablesPath : Path to the Distcc and DMUCS executables.
+
+$ DMUCSCoordinator : The DMUCS coordinator hostname or IP address.
+
+$ DMUCSDistProp : The DMUCS distinguishing property value to request for compile hosts.
+
+$ bXGENoWatchdogThread : Whether to use the no_watchdog_thread option to prevent VS2015 toolchain stalls.
+
+$ bShowXGEMonitor : Whether to display the XGE build monitor.
+
+$ bStopXGECompilationAfterErrors : When enabled, XGE will stop compiling targets after a compile error occurs.  Recommended, as it saves computing resources for others.
+
+$ bWriteSolutionOptionFile : Whether to write a solution option (suo) file for the sln.
+
+$ bAddFastPDBToProjects : Whether to add the -FastPDB option to build command lines by default.
+
+$ bUsePerFileIntellisense : Whether to generate per-file intellisense data.
+
+$ bEditorDependsOnShaderCompileWorker : Whether to include a dependency on ShaderCompileWorker when generating project files for the editor.
+
+$ bDebugBuildsActuallyUseDebugCRT : Enables the debug C++ runtime (CRT) for debug builds. By default we always use the release runtime, since the debug version isn't particularly useful when debugging Unreal Engine projects, and linking against the debug CRT libraries forces our third party library dependencies to also be compiled using the debug CRT (and often perform more slowly). Often it can be inconvenient to require a separate copy of the debug versions of third party static libraries simply so that you can debug your program's code.
+
+$ bUseInlining : Enable inlining for all modules.
+
+$ bUseDebugLiveCodingConsole : Whether to enable support for live coding
+
+$ bUseXGEController : Whether the XGE controller worker and modules should be included in the engine build. These are required for distributed shader compilation using the XGE interception interface.
+
+$ bUseUnityBuild : Whether to unify C++ code into larger files for faster compilation.
+
+$ bForceUnityBuild : Whether to force C++ source files to be combined into larger files for faster compilation.
+
+$ bUseAdaptiveUnityBuild : Use a heuristic to determine which files are currently being iterated on and exclude them from unity blobs, result in faster incremental compile times. The current implementation uses the read-only flag to distinguish the working set, assuming that files will be made writable by the source control system if they are being modified. This is true for Perforce, but not for Git.
+
+$ bAdaptiveUnityDisablesOptimizations : Disable optimization for files that are in the adaptive non-unity working set.
+
+$ bAdaptiveUnityDisablesPCH : Disables force-included PCHs for files that are in the adaptive non-unity working set.
+
+$ bAdaptiveUnityDisablesProjectPCHForProjectPrivate : Backing storage for bAdaptiveUnityDisablesProjectPCH.
+
+$ bAdaptiveUnityCreatesDedicatedPCH : Creates a dedicated PCH for each source file in the working set, allowing faster iteration on cpp-only changes.
+
+$ bAdaptiveUnityEnablesEditAndContinue : Creates a dedicated PCH for each source file in the working set, allowing faster iteration on cpp-only changes.
+
+$ MinGameModuleSourceFilesForUnityBuild : The number of source files in a game module before unity build will be activated for that module.  This allows small game modules to have faster iterative compile times for single files, at the expense of slower full rebuild times.  This setting can be overridden by the bFasterWithoutUnity option in a module's Build.cs file.
+
+$ bShadowVariableErrors : Forces shadow variable warnings to be treated as errors on platforms that support it.
+
+$ bUndefinedIdentifierErrors : Forces the use of undefined identifiers in conditional expressions to be treated as errors.
+
+$ bUseFastMonoCalls : New Monolithic Graphics drivers have optional "fast calls" replacing various D3d functions
+
+$ bUseFastSemanticsRenderContexts : New Xbox driver supports a "fast semantics" context type. This switches it on for the immediate and deferred contexts Try disabling this if you see rendering issues and/or crashes inthe Xbox RHI.
+
+$ NumIncludedBytesPerUnityCPP : An approximate number of bytes of C++ code to target for inclusion in a single unified C++ file.
+
+$ bStressTestUnity : Whether to stress test the C++ unity build robustness by including all C++ files files in a project from a single unified file.
+
+$ bDisableDebugInfo : Whether to globally disable debug info generation; see DebugInfoHeuristics.cs for per-config and per-platform options.
+
+$ bDisableDebugInfoForGeneratedCode : Whether to disable debug info generation for generated files. This improves link times for modules that have a lot of generated glue code.
+
+$ bOmitPCDebugInfoInDevelopment : Whether to disable debug info on PC in development builds (for faster developer iteration, as link times are extremely fast with debug info disabled).
+
+$ bUsePDBFiles : Whether PDB files should be used for Visual C++ builds.
+
+$ bUsePCHFiles : Whether PCH files should be used.
+
+$ MinFilesUsingPrecompiledHeader : The minimum number of files that must use a pre-compiled header before it will be created and used.
+
+$ bForcePrecompiledHeaderForGameModules : When enabled, a precompiled header is always generated for game modules, even if there are only a few source files in the module.  This greatly improves compile times for iterative changes on a few files in the project, at the expense of slower full rebuild times for small game projects.  This can be overridden by setting MinFilesUsingPrecompiledHeaderOverride in a module's Build.cs file.
+
+$ bUseIncrementalLinking : Whether to use incremental linking or not. Incremental linking can yield faster iteration times when making small changes. Currently disabled by default because it tends to behave a bit buggy on some computers (PDB-related compile errors).
+
+$ bAllowLTCG : Whether to allow the use of link time code generation (LTCG).
+
+$ bPGOProfile : Whether to enable Profile Guided Optimization (PGO) instrumentation in this build.
+
+$ bPGOOptimize : Whether to optimize this build with Profile Guided Optimization (PGO).
+
+$ bAllowASLRInShipping : Whether to allow the use of ASLR (address space layout randomization) if supported. Only applies to shipping builds.
+
+$ bSupportEditAndContinue : Whether to support edit and continue.  Only works on Microsoft compilers.
+
+$ bOmitFramePointers : Whether to omit frame pointers or not. Disabling is useful for e.g. memory profiling on the PC.
+
+$ bUseMallocProfiler : If true, then enable memory profiling in the build (defines USE_MALLOC_PROFILER=1 and forces bOmitFramePointers=false).
+
+$ bUseSharedPCHs : Enables "Shared PCHs", a feature which significantly speeds up compile times by attempting to share certain PCH files between modules that UBT detects is including those PCH's header files.
+
+$ bUseShippingPhysXLibraries : True if Development and Release builds should use the release configuration of PhysX/APEX.
+
+$ bUseCheckedPhysXLibraries : True if Development and Release builds should use the checked configuration of PhysX/APEX. if bUseShippingPhysXLibraries is true this is ignored.
+
+$ bCheckLicenseViolations : Tells the UBT to check if module currently being built is violating EULA.
+
+$ bBreakBuildOnLicenseViolation : Tells the UBT to break build if module currently being built is violating EULA.
+
+$ bUseFastPDBLinking : Whether to use the :FASTLINK option when building with /DEBUG to create local PDBs on Windows. Fast, but currently seems to have problems finding symbols in the debugger.
+
+$ bCreateMapFile : Outputs a map file as part of the build.
+
+$ bAllowRuntimeSymbolFiles : True if runtime symbols files should be generated as a post build step for some platforms. These files are used by the engine to resolve symbol names of callstack backtraces in logs.
+
+$ bCheckSystemHeadersForModification : Whether headers in system paths should be checked for modification when determining outdated actions.
+
+$ bFlushBuildDirOnRemoteMac : Whether to clean Builds directory on a remote Mac before building.
+
+$ bPrintToolChainTimingInfo : Whether to write detailed timing info from the compiler and linker.
+
+$ bParseTimingInfoForTracing : Whether to parse timing data into a tracing file compatible with chrome://tracing.
+
+$ bPublicSymbolsByDefault : Whether to expose all symbols as public by default on POSIX platforms
+
+$ CppStandard : Which C++ stanard to use for compiling this target
+
+$ bStripSymbols : Whether to strip iOS symbols or not (implied by Shipping config).
+
+$ bIgnoreOutdatedImportLibraries : Whether to ignore import library files that are out of date when building targets. Set this to true to improve iteration time. By default, we do not bother re-linking targets if only a dependent .lib has changed, as chances are that the import library was not actually different unless a dependent header file of this target was also changed, in which case the target would automatically be rebuilt.
+
+$ bPrintDebugInfo : Whether debug info should be written to the console.
+
+$ bLogDetailedActionStats : Whether to log detailed action stats. This forces local execution.
+
+$ bAllowHybridExecutor : Whether the hybrid executor will be used (a remote executor and local executor).
+
+$ bAllowXGE : Whether XGE may be used.
+
+$ bAllowSNDBS : Whether SN-DBS may be used.
+
+$ bUseUBTMakefiles : Enables support for very fast iterative builds by caching target data. Turning this on causes Unreal Build Tool to emit 'UBT Makefiles' for targets when they are built the first time. Subsequent builds will load these Makefiles and begin outdatedness checking and build invocation very quickly. The caveat is that if source files are added or removed to the project, UBT will need to gather information about those in order for your build to complete successfully. Currently, you must run the project file generator after adding/removing source files to tell UBT to re-gather this information. 
+    * Events that can invalidate the 'UBT Makefile':
+        * Adding/removing .cpp files
+        * Adding/removing .h files with UObjects
+        * Adding new UObject types to a file that did not previously have any
+        * Changing global build settings (most settings in this file qualify)
+        * Changed code that affects how Unreal Header Tool works 
+    * You can force regeneration of the 'UBT Makefile' by passing the '-gather' argument, or simply regenerating project files. 
+    * This also enables the fast include file dependency scanning and caching system that allows Unreal Build Tool to detect out of date dependencies very quickly. When enabled, a deep C++ include graph does not have to be generated, and instead, we only scan and cache indirect includes for after a dependent build product was already found to be out of date. During the next build, we will load those cached indirect includes and check for outdatedness.
+
+$ bAllowDistcc : Whether DMUCS/Distcc may be used. Distcc requires some setup -- so by default, disable it so that we do not break local or remote building.
+
+$ bAllowParallelExecutor : Whether to allow using parallel executor on Windows.
+
+### HybridExecutor
+
+$ MaxLocalActions : Maximum number of actions to execute locally.
+
+### BuildMode
+
+$ bIgnoreJunk : Whether to skip checking for files identified by the junk manifest.
+
+### ProjectFileGenerator
+
+$ DisablePlatformProjectGenerators : Disable native project file generators for platforms. Platforms with native project file generators typically require IDE extensions to be installed.
+
+$ Format : Default list of project file formats to generate.
+
+$ bGenerateIntelliSenseData : True if intellisense data should be generated (takes a while longer).
+
+$ bIncludeDocumentation : True if we should include documentation in the generated projects.
+
+$ bAllDocumentationLanguages : True if all documentation languages should be included in generated projects, otherwise only INT files will be included.
+
+$ bUsePrecompiled : True if build targets should pass the -useprecompiled argument.
+
+$ bIncludeEngineSource : True if we should include engine source in the generated solution.
+
+$ bIncludeEnterpriseSource : Whether to include enterprise source in the generated solution.
+
+$ bIncludeShaderSource : True if shader source files should be included in generated projects.
+
+$ bIncludeBuildSystemFiles : True if build system files should be included.
+
+$ bIncludeConfigFiles : True if we should include config (.ini) files in the generated project.
+
+$ bIncludeLocalizationFiles : True if we should include localization files in the generated project.
+
+$ bIncludeTemplateFiles : True if we should include template files in the generated project.
+
+$ IncludeEnginePrograms : True if we should include program projects in the generated solution.
+
+$ bIncludeTempTargets : Whether to include temporary targets generated by UAT to support content only projects with non-default settings.
+
+$ bIncludeDotNETCoreProjects : True if we should include .NET Core projects in the generated solution.
+
+$ bKeepSourceSubDirectories : True if we should reflect "Source" sub-directories on disk in the master project as master project directories. This (arguably) adds some visual clutter to the master project but it is truer to the on-disk file organization.
+
+$ Platforms : Names of platforms to include in the generated project files
+
+$ MasterProjectName : Name of the master project file -- for example, the base file name for the Visual Studio solution file, or the Xcode project file on Mac.
+
+$ bMasterProjectNameFromFolder : If true, sets the master project name according to the name of the folder it is in.
+
+$ bIncludeTestAndShippingConfigs : Whether we should include configurations for "Test" and "Shipping" in generated projects. Pass "-NoShippingConfigs" to disable this.
+
+### LocalExecutor
+
+$ ProcessorCountMultiplier : Processor count multiplier for local execution. Can be below 1 to reserve CPU for other tasks. When using the local executor (not XGE), run a single action on each CPU core. Note that you can set this to a larger value to get slightly faster build times in many cases, but your computer's responsiveness during compiling may be much worse.
+
+$ MaxProcessorCount : Maximum processor count for local execution.
+
+### ParallelExecutor
+
+$ ProcessorCountMultiplier : Processor count multiplier for local execution. Can be below 1 to reserve CPU for other tasks.
+
+$ MaxProcessorCount : Maximum processor count for local execution.
+
+$ bStopCompilationAfterErrors : When enabled, will stop compiling targets after a compile error occurs.
+
+### SNDBS
+
+$ ProcessorCountMultiplier : Processor count multiplier for local execution. Can be below 1 to reserve CPU for other tasks.
+
+$ MaxProcessorCount : Maximum processor count for local execution.
+
+### CLionGenerator
+
+$ bIncludeDocumentation : True if we should include documentation in the generated projects.
+
+$ bUsePrecompiled : True if build targets should pass the -useprecompiled argument.
+
+$ bIncludeEngineSource : True if we should include engine source in the generated solution.
+
+$ bIncludeEnterpriseSource : Whether to include enterprise source in the generated solution.
+
+$ bIncludeShaderSource : True if shader source files should be included in generated projects.
+
+$ bIncludeConfigFiles : True if we should include config (.ini) files in the generated project.
+
+$ bIncludeTemplateFiles : True if we should include template files in the generated project.
+
+$ IncludeEnginePrograms : True if we should include program projects in the generated solution.
+
+$ MasterProjectName : Name of the master project file -- for example, the base file name for the Visual Studio solution file, or the Xcode project file on Mac.
+
+$ bMasterProjectNameFromFolder : If true, sets the master project name according to the name of the folder it is in.
+
+### CMakefileGenerator
+
+$ bIncludeDocumentation : True if we should include documentation in the generated projects.
+
+$ bUsePrecompiled : True if build targets should pass the -useprecompiled argument.
+
+$ bIncludeEngineSource : True if we should include engine source in the generated solution.
+
+$ bIncludeEnterpriseSource : Whether to include enterprise source in the generated solution.
+
+$ bIncludeShaderSource : True if shader source files should be included in generated projects.
+
+$ bIncludeConfigFiles : True if we should include config (.ini) files in the generated project.
+
+$ bIncludeTemplateFiles : True if we should include template files in the generated project.
+
+$ IncludeEnginePrograms : True if we should include program projects in the generated solution.
+
+$ MasterProjectName : Name of the master project file -- for example, the base file name for the Visual Studio solution file, or the Xcode project file on Mac.
+
+$ bMasterProjectNameFromFolder : If true, sets the master project name according to the name of the folder it is in.
+
+### CodeLiteGenerator
+
+$ bIncludeDocumentation : True if we should include documentation in the generated projects.
+
+$ bUsePrecompiled : True if build targets should pass the -useprecompiled argument.
+
+$ bIncludeEngineSource : True if we should include engine source in the generated solution.
+
+$ bIncludeEnterpriseSource : Whether to include enterprise source in the generated solution.
+
+$ bIncludeShaderSource : True if shader source files should be included in generated projects.
+
+$ bIncludeConfigFiles : True if we should include config (.ini) files in the generated project.
+
+$ bIncludeTemplateFiles : True if we should include template files in the generated project.
+
+$ IncludeEnginePrograms : True if we should include program projects in the generated solution.
+
+$ MasterProjectName : Name of the master project file -- for example, the base file name for the Visual Studio solution file, or the Xcode project file on Mac.
+
+$ bMasterProjectNameFromFolder : If true, sets the master project name according to the name of the folder it is in.
+
+### EddieProjectFileGenerator
+
+$ bIncludeDocumentation : True if we should include documentation in the generated projects.
+
+$ bUsePrecompiled : True if build targets should pass the -useprecompiled argument.
+
+$ bIncludeEngineSource : True if we should include engine source in the generated solution.
+
+$ bIncludeEnterpriseSource : Whether to include enterprise source in the generated solution.
+
+$ bIncludeShaderSource : True if shader source files should be included in generated projects.
+
+$ bIncludeConfigFiles : True if we should include config (.ini) files in the generated project.
+
+$ bIncludeTemplateFiles : True if we should include template files in the generated project.
+
+$ IncludeEnginePrograms : True if we should include program projects in the generated solution.
+
+$ MasterProjectName : Name of the master project file -- for example, the base file name for the Visual Studio solution file, or the Xcode project file on Mac.
+
+$ bMasterProjectNameFromFolder : If true, sets the master project name according to the name of the folder it is in.
+
+### KDevelopGenerator
+
+$ bIncludeDocumentation : True if we should include documentation in the generated projects.
+
+$ bUsePrecompiled : True if build targets should pass the -useprecompiled argument.
+
+$ bIncludeEngineSource : True if we should include engine source in the generated solution.
+
+$ bIncludeEnterpriseSource : Whether to include enterprise source in the generated solution.
+
+$ bIncludeShaderSource : True if shader source files should be included in generated projects.
+
+$ bIncludeConfigFiles : True if we should include config (.ini) files in the generated project.
+
+$ bIncludeTemplateFiles : True if we should include template files in the generated project.
+
+$ IncludeEnginePrograms : True if we should include program projects in the generated solution.
+
+$ MasterProjectName : Name of the master project file -- for example, the base file name for the Visual Studio solution file, or the Xcode project file on Mac.
+
+$ bMasterProjectNameFromFolder : If true, sets the master project name according to the name of the folder it is in.
+
+### MakefileGenerator
+
+$ bIncludeDocumentation : True if we should include documentation in the generated projects.
+
+$ bUsePrecompiled : True if build targets should pass the -useprecompiled argument.
+
+$ bIncludeEngineSource : True if we should include engine source in the generated solution.
+
+$ bIncludeEnterpriseSource : Whether to include enterprise source in the generated solution.
+
+$ bIncludeShaderSource : True if shader source files should be included in generated projects.
+
+$ bIncludeConfigFiles : True if we should include config (.ini) files in the generated project.
+
+$ bIncludeTemplateFiles : True if we should include template files in the generated project.
+
+$ IncludeEnginePrograms : True if we should include program projects in the generated solution.
+
+$ MasterProjectName : Name of the master project file -- for example, the base file name for the Visual Studio solution file, or the Xcode project file on Mac.
+
+$ bMasterProjectNameFromFolder : If true, sets the master project name according to the name of the folder it is in.
+
+### QMakefileGenerator
+
+$ bIncludeDocumentation : True if we should include documentation in the generated projects.
+
+$ bUsePrecompiled : True if build targets should pass the -useprecompiled argument.
+
+$ bIncludeEngineSource : True if we should include engine source in the generated solution.
+
+$ bIncludeEnterpriseSource : Whether to include enterprise source in the generated solution.
+
+$ bIncludeShaderSource : True if shader source files should be included in generated projects.
+
+$ bIncludeConfigFiles : True if we should include config (.ini) files in the generated project.
+
+$ bIncludeTemplateFiles : True if we should include template files in the generated project.
+
+$ IncludeEnginePrograms : True if we should include program projects in the generated solution.
+
+$ MasterProjectName : Name of the master project file -- for example, the base file name for the Visual Studio solution file, or the Xcode project file on Mac.
+
+$ bMasterProjectNameFromFolder : If true, sets the master project name according to the name of the folder it is in.
+
+### RiderProjectFileGenerator
+
+$ bIncludeDocumentation : True if we should include documentation in the generated projects.
+
+$ bUsePrecompiled : True if build targets should pass the -useprecompiled argument.
+
+$ bIncludeEngineSource : True if we should include engine source in the generated solution.
+
+$ bIncludeEnterpriseSource : Whether to include enterprise source in the generated solution.
+
+$ bIncludeShaderSource : True if shader source files should be included in generated projects.
+
+$ bIncludeConfigFiles : True if we should include config (.ini) files in the generated project.
+
+$ bIncludeTemplateFiles : True if we should include template files in the generated project.
+
+$ IncludeEnginePrograms : True if we should include program projects in the generated solution.
+
+$ MasterProjectName : Name of the master project file -- for example, the base file name for the Visual Studio solution file, or the Xcode project file on Mac.
+
+$ bMasterProjectNameFromFolder : If true, sets the master project name according to the name of the folder it is in.
+
+### VSCodeProjectFileGenerator
+
+$ bIncludeDocumentation : True if we should include documentation in the generated projects.
+
+$ bUsePrecompiled : True if build targets should pass the -useprecompiled argument.
+
+$ bIncludeEngineSource : True if we should include engine source in the generated solution.
+
+$ bIncludeEnterpriseSource : Whether to include enterprise source in the generated solution.
+
+$ bIncludeShaderSource : True if shader source files should be included in generated projects.
+
+$ bIncludeConfigFiles : True if we should include config (.ini) files in the generated project.
+
+$ bIncludeTemplateFiles : True if we should include template files in the generated project.
+
+$ IncludeEnginePrograms : True if we should include program projects in the generated solution.
+
+$ MasterProjectName : Name of the master project file -- for example, the base file name for the Visual Studio solution file, or the Xcode project file on Mac.
+
+$ bMasterProjectNameFromFolder : If true, sets the master project name according to the name of the folder it is in.
+
+### VCProjectFileGenerator
+
+$ Version : The version of Visual Studio to generate project files for.
+
+$ bBuildLiveCodingConsole : Whether to include a dependency on LiveCodingConsole when building targets that support live coding.
+
+$ bIncludeDocumentation : True if we should include documentation in the generated projects.
+
+$ bUsePrecompiled : True if build targets should pass the -useprecompiled argument.
+
+$ bIncludeEngineSource : True if we should include engine source in the generated solution.
+
+$ bIncludeEnterpriseSource : Whether to include enterprise source in the generated solution.
+
+$ bIncludeShaderSource : True if shader source files should be included in generated projects.
+
+$ bIncludeConfigFiles : True if we should include config (.ini) files in the generated project.
+
+$ bIncludeTemplateFiles : True if we should include template files in the generated project.
+
+$ IncludeEnginePrograms : True if we should include program projects in the generated solution.
+
+$ MasterProjectName : Name of the master project file -- for example, the base file name for the Visual Studio solution file, or the Xcode project file on Mac.
+
+$ bMasterProjectNameFromFolder : If true, sets the master project name according to the name of the folder it is in.
+
+### XcodeProjectFileGenerator
+
+$ bIncludeDocumentation : True if we should include documentation in the generated projects.
+
+$ bUsePrecompiled : True if build targets should pass the -useprecompiled argument.
+
+$ bIncludeEngineSource : True if we should include engine source in the generated solution.
+
+$ bIncludeEnterpriseSource : Whether to include enterprise source in the generated solution.
+
+$ bIncludeShaderSource : True if shader source files should be included in generated projects.
+
+$ bIncludeConfigFiles : True if we should include config (.ini) files in the generated project.
+
+$ bIncludeTemplateFiles : True if we should include template files in the generated project.
+
+$ IncludeEnginePrograms : True if we should include program projects in the generated solution.
+
+$ MasterProjectName : Name of the master project file -- for example, the base file name for the Visual Studio solution file, or the Xcode project file on Mac.
+
+$ bMasterProjectNameFromFolder : If true, sets the master project name according to the name of the folder it is in.
+
+### ModuleConfiguration
+
+$ DisableUnityBuild : List of modules to disable unity builds for
+
+$ EnableOptimizeCode : List of modules to enable optimizations for
+
+$ DisableOptimizeCode : List of modules to disable optimizations for
+
+### SourceFileWorkingSet
+
+$ Provider : Sets the provider to use for determining the working set.
+
+$ RepositoryPath : Sets the path to use for the repository. Interpreted relative to the Unreal Engine root directory (the folder above the Engine folder) -- if relative.
+
+$ GitPath : Sets the path to use for the Git executable. Defaults to "git" (assuming it is in the PATH).
+
+### VCMacProjectFileGenerator
+
+$ Version : The version of Visual Studio to generate project files for.
+
+$ bIncludeDocumentation : True if we should include documentation in the generated projects.
+
+$ bUsePrecompiled : True if build targets should pass the -useprecompiled argument.
+
+$ bIncludeEngineSource : True if we should include engine source in the generated solution.
+
+$ bIncludeEnterpriseSource : Whether to include enterprise source in the generated solution.
+
+$ bIncludeShaderSource : True if shader source files should be included in generated projects.
+
+$ bIncludeConfigFiles : True if we should include config (.ini) files in the generated project.
+
+$ bIncludeTemplateFiles : True if we should include template files in the generated project.
+
+$ IncludeEnginePrograms : True if we should include program projects in the generated solution.
+
+$ MasterProjectName : Name of the master project file -- for example, the base file name for the Visual Studio solution file, or the Xcode project file on Mac.
+
+$ bMasterProjectNameFromFolder : If true, sets the master project name according to the name of the folder it is in.
+
+### IOSToolChain
+
+$ IOSSDKVersion : The version of the iOS SDK to target at build time.
+
+$ BuildIOSVersion : The version of the iOS to allow at build time.
+
+$ bUseDangerouslyFastMode : If this is set, then we do not do any post-compile steps -- except moving the executable into the proper spot on Mac.
+
+### RemoteMac
+
+$ ServerName : These two variables will be loaded from the XML config file in XmlConfigLoader.Init().
+
+$ UserName : The remote username.
+
+$ SshPrivateKey : If set, instead of looking for RemoteToolChainPrivate.key in the usual places (Documents/Unreal, Engine/UnrealBuildTool/SSHKeys or Engine/Build/SSHKeys), this private key will be used.
+
+$ RsyncAuthentication : The authentication used for Rsync (for the -e rsync flag).
+
+$ SshAuthentication : The authentication used for SSH (probably similar to RsyncAuthentication).
+
+### UEBuildConfiguration
+
+$ bForceHeaderGeneration : If true, force header regeneration. Intended for the build machine.
+
+$ bDoNotBuildUHT : If true, do not build UHT, assume it is already built.
+
+$ bFailIfGeneratedCodeChanges : If true, fail if any of the generated header files is out of date.
+
+$ bAllowHotReloadFromIDE : True if hot-reload from IDE is allowed.
+
+$ bForceDebugUnrealHeaderTool : If true, the Debug version of UnrealHeaderTool will be built and run instead of the Development version.
+
+### WindowsPlatform
+
+$ Compiler : Version of the compiler toolchain to use on Windows platform. A value of "default" will be changed to a specific version at UBT start up.
+
+$ CompilerVersion : The specific toolchain version to use. This may be a specific version number (for example, "14.13.26128"), or the string "Latest", to select the newest available version. By default, and if it is available, we use the toolchain version indicated by WindowsPlatform.DefaultToolChainVersion (otherwise, we use the latest version).
+
+$ WindowsSdkVersion : The specific Windows SDK version to use. This may be a specific version number (for example, "8.1", "10.0" or "10.0.10150.0"), or the string "Latest", to select the newest available version. By default, and if it is available, we use the Windows SDK version indicated by WindowsPlatform.DefaultWindowsSdkVersion (otherwise, we use the latest version).
+
+$ StaticAnalyzer : The static analyzer to use.
+
+$ bStrictConformanceMode : Enables strict standard conformance mode (/permissive-) in VS2017+.
+
+$ PCHMemoryAllocationFactor : Determines the amount of memory that the compiler allocates to construct precompiled headers (/Zm).
+
+$ bClangTimeTrace : (Experimental) Appends the -ftime-trace argument to the command line for Clang to output a JSON file containing a timeline for the compile. See http://aras-p.info/blog/2019/01/16/time-trace-timeline-flame-chart-profiler-for-Clang/ for more info.
+
+$ bCompilerTrace : Outputs compile timing information so that it can be analyzed.
+
+### WindowsTargetRules
+
+$ ObjSrcMapFile : Whether we should export a file containing .obj to source file mappings.
+
diff --git a/Engine/Documentation/Source/Programming/UnrealBuildSystem/ModuleFiles/ModuleFilesProperties/ModuleFilesProperties.INT.udn b/Engine/Documentation/Source/Programming/UnrealBuildSystem/ModuleFiles/ModuleFilesProperties/ModuleFilesProperties.INT.udn
new file mode 100644
--- /dev/null
+++ b/Engine/Documentation/Source/Programming/UnrealBuildSystem/ModuleFiles/ModuleFilesProperties/ModuleFilesProperties.INT.udn
@@ -0,0 +1,118 @@
+Availability: NoPublish
+Title: Build Configuration Properties Page
+Crumbs:
+Description: This is a procedurally generated markdown page.
+Version: 4.24
+
+### Read/Write Properties
+$ Type (ModuleType): Type of module
+
+$ BinariesSubFolder (String): Subfolder of Binaries/PLATFORM folder to put this module in when building DLLs. This should only be used by modules that are found via searching like the TargetPlatform or ShaderFormat modules. If FindModules is not used to track them down, the modules will not be found.
+
+$ PrivatePCHHeaderFile (String): Explicit private PCH for this module. Implies that this module will not use a shared PCH.
+
+$ SharedPCHHeaderFile (String): Header file name for a shared PCH provided by this module.  Must be a valid relative path to a public C++ header file. This should only be set for header files that are included by a significant number of other C++ modules.
+
+$ ShortName (String): Specifies an alternate name for intermediate directories and files for intermediates of this module. Useful when hitting path length limitations.
+
+$ bTreatAsEngineModule (Boolean): Whether this module should be treated as an engine module (eg. using engine definitions, PCHs, compiled with optimizations enabled in DebugGame configurations, etc...). Initialized to a default based on the rules assembly it was created from.
+
+$ bUseRTTI (Boolean): Use run time type information
+
+$ bUseAVX (Boolean): Use AVX instructions
+
+$ bEnableBufferSecurityChecks (Boolean): Enable buffer security checks.  This should usually be enabled as it prevents severe security risks.
+
+$ bEnableExceptions (Boolean): Enable exception handling
+
+$ bEnableObjCExceptions (Boolean): Enable objective C exception handling
+
+$ bEnableShadowVariableWarnings (Boolean): Enable warnings for shadowed variables
+
+$ bEnableUndefinedIdentifierWarnings (Boolean): Enable warnings for using undefined identifiers in #if expressions
+
+$ MinSourceFilesForUnityBuildOverride (Int32): The number of source files in this module before unity build will be activated for that module.  If set to anything besides -1, will override the default setting which is controlled by MinGameModuleSourceFilesForUnityBuild
+
+$ MinFilesUsingPrecompiledHeaderOverride (Int32): Overrides BuildConfiguration.MinFilesUsingPrecompiledHeader if non-zero.
+
+$ bBuildLocallyWithSNDBS (Boolean): Module uses a #import so must be built locally when compiling with SN-DBS
+
+$ IsRedistributableOverride (Nullable&lt;Boolean&gt;): Redistribution override flag for this module.
+
+$ bLegalToDistributeObjectCode (Boolean): Whether the output from this module can be publicly distributed, even if it has code/ dependencies on modules that are not (i.e. CarefullyRedist, NotForLicensees, NoRedist). This should be used when you plan to release binaries but not source.
+
+$ WhitelistRestrictedFolders (List&lt;String&gt;): List of folders which are whitelisted to be referenced when compiling this binary, without propagating restricted folder names
+
+$ bEnforceIWYU (Boolean): Enforce "include what you use" rules when PCHUsage is set to ExplicitOrSharedPCH; warns when monolithic headers (Engine.h, UnrealEd.h, etc...) are used, and checks that source files include their matching header first.
+
+$ bAddDefaultIncludePaths (Boolean): Whether to add all the default include paths to the module (eg. the Source/Classes folder, subfolders under Source/Public).
+
+$ bIgnoreUnresolvedSymbols (Boolean): Whether to ignore dangling (i.e. unresolved external) symbols in modules
+
+$ bPrecompile (Boolean): Whether this module should be precompiled. Defaults to the bPrecompile flag from the target. Clear this flag to prevent a module being precompiled.
+
+$ bUsePrecompiled (Boolean): Whether this module should use precompiled data. Always true for modules created from installed assemblies.
+
+$ bAllowConfidentialPlatformDefines (Boolean): Whether this module can use PLATFORM_XXXX style defines, where XXXX is a confidential platform name. This is used to ensure engine or other shared code does not reveal confidential information inside an #if PLATFORM_XXXX block. Licensee game code may want to allow for them, however. Note: this is future looking, and previous confidential platforms (like PS4) are unlikely to be restricted
+
+$ PublicIncludePathModuleNames (List&lt;String&gt;): List of modules names (no path needed) with header files that our module's public headers needs access to, but we don't need to "import" or link against.
+
+$ PublicDependencyModuleNames (List&lt;String&gt;): List of public dependency module names (no path needed) (automatically does the private/public include). These are modules that are required by our public source files.
+
+$ PrivateIncludePathModuleNames (List&lt;String&gt;): List of modules name (no path needed) with header files that our module's private code files needs access to, but we don't need to "import" or link against.
+
+$ PrivateDependencyModuleNames (List&lt;String&gt;): List of private dependency module names.  These are modules that our private code depends on but nothing in our public include files depend on.
+
+$ CircularlyReferencedDependentModules (List&lt;String&gt;): Only for legacy reason, should not be used in new code. List of module dependencies that should be treated as circular references.  This modules must have already been added to either the public or private dependent module list.
+
+$ PublicSystemIncludePaths (List&lt;String&gt;): List of system/library include paths - typically used for External (third party) modules.  These are public stable header file directories that are not checked when resolving header dependencies.
+
+$ PublicIncludePaths (List&lt;String&gt;): (This setting is currently not need as we discover all files from the 'Public' folder) List of all paths to include files that are exposed to other modules
+
+$ PrivateIncludePaths (List&lt;String&gt;): List of all paths to this module's internal include files, not exposed to other modules (at least one include to the 'Private' path, more if we want to avoid relative paths)
+
+$ PublicSystemLibraryPaths (List&lt;String&gt;): List of system library paths (directory of .lib files) - for External (third party) modules please use the PublicAdditionalLibaries instead
+
+$ PrivateRuntimeLibraryPaths (List&lt;String&gt;): List of search paths for libraries at runtime (eg. .so files)
+
+$ PublicRuntimeLibraryPaths (List&lt;String&gt;): List of search paths for libraries at runtime (eg. .so files)
+
+$ PublicAdditionalLibraries (List&lt;String&gt;): List of additional libraries (names of the .lib files including extension) - typically used for External (third party) modules
+
+$ PublicSystemLibraries (List&lt;String&gt;): List of system libraries to use - these are typically referenced via name and then found via the system paths. If you need to reference a .lib file use the PublicAdditionalLibraries instead
+
+$ PublicFrameworks (List&lt;String&gt;): List of XCode frameworks (iOS and MacOS)
+
+$ PublicWeakFrameworks (List&lt;String&gt;): List of weak frameworks (for OS version transitions)
+
+$ PublicAdditionalFrameworks (List&lt;Framework&gt;): List of addition frameworks - typically used for External (third party) modules on Mac and iOS
+
+$ AdditionalBundleResources (List&lt;BundleResource&gt;): List of addition resources that should be copied to the app bundle for Mac or iOS
+
+$ PublicAdditionalShadowFiles (List&lt;String&gt;): For builds that execute on a remote machine (e.g. iOS), this list contains additional files that need to be copied over in order for the app to link successfully.  Source/header files and PCHs are automatically copied.  Usually this is simply a list of precompiled third party library dependencies.
+
+$ PublicDelayLoadDLLs (List&lt;String&gt;): List of delay load DLLs - typically used for External (third party) modules
+
+$ PrivateDefinitions (List&lt;String&gt;): Private compiler definitions for this module
+
+$ PublicDefinitions (List&lt;String&gt;): Public compiler definitions for this module
+
+$ DynamicallyLoadedModuleNames (List&lt;String&gt;): Addition modules this module may require at run-time
+
+$ RuntimeDependencies (RuntimeDependencyList): List of files which this module depends on at runtime. These files will be staged along with the target.
+
+$ AdditionalPropertiesForReceipt (ReceiptPropertyList): List of additional properties to be added to the build receipt
+
+$ PrecompileForTargets (PrecompileTargetsType): Which targets this module should be precompiled for
+
+$ ExternalDependencies (List&lt;String&gt;): External files which invalidate the makefile if modified. Relative paths are resolved relative to the .build.cs file.
+
+$ SubclassRules (List&lt;String&gt;): Subclass rules files which invalidate the makefile if modified.
+
+$ bRequiresImplementModule (Nullable&lt;Boolean&gt;): Whether this module requires the IMPLEMENT_MODULE macro to be implemented. Most UE4 modules require this, since we use the IMPLEMENT_MODULE macro to do other global overloads (eg. operator new/delete forwarding to GMalloc).
+
+$ CppStandard (CppStandardVersion): Which stanard to use for compiling this module
+
+$ ModuleSymbolVisibility (SymbolVisibility): Control visibility of symbols
+
+
diff --git a/Engine/Documentation/Source/Programming/UnrealBuildSystem/TargetFiles/TargetFilesProperties/TargetFilesProperties.INT.udn b/Engine/Documentation/Source/Programming/UnrealBuildSystem/TargetFiles/TargetFilesProperties/TargetFilesProperties.INT.udn
new file mode 100644
--- /dev/null
+++ b/Engine/Documentation/Source/Programming/UnrealBuildSystem/TargetFiles/TargetFilesProperties/TargetFilesProperties.INT.udn
@@ -0,0 +1,291 @@
+Availability: NoPublish
+Title: Build Configuration Properties Page
+Crumbs:
+Description: This is a procedurally generated markdown page.
+Version: 4.24
+
+### Read-Only Properties
+
+$ Name (String): The name of this target.
+
+$ Platform (UnrealTargetPlatform): Platform that this target is being built for.
+
+$ Configuration (UnrealTargetConfiguration): The configuration being built.
+
+$ Architecture (String): Architecture that the target is being built for (or an empty string for the default).
+
+$ ProjectFile (FileReference): Path to the project file for the project containing this target.
+
+$ Version (ReadOnlyBuildVersion): The current build version
+
+
+### Read/Write Properties
+$ Type (TargetType): The type of target.
+
+$ bUsesSteam (Boolean): Whether the target uses Steam.
+
+$ bUsesCEF3 (Boolean): Whether the target uses CEF3.
+
+$ bUsesSlate (Boolean): Whether the project uses visual Slate UI (as opposed to the low level windowing/messaging, which is always available).
+
+$ bUseStaticCRT (Boolean): Forces linking against the static CRT. This is not fully supported across the engine due to the need for allocator implementations to be shared (for example), and TPS libraries to be consistent with each other, but can be used for utility programs.
+
+$ bDebugBuildsActuallyUseDebugCRT (Boolean): Enables the debug C++ runtime (CRT) for debug builds. By default we always use the release runtime, since the debug version isn't particularly useful when debugging Unreal Engine projects, and linking against the debug CRT libraries forces our third party library dependencies to also be compiled using the debug CRT (and often perform more slowly). Often it can be inconvenient to require a separate copy of the debug versions of third party static libraries simply so that you can debug your program's code.
+
+$ bLegalToDistributeBinary (Boolean): Whether the output from this target can be publicly distributed, even if it has dependencies on modules that are in folders with special restrictions (eg. CarefullyRedist, NotForLicensees, NoRedist).
+
+$ UndecoratedConfiguration (UnrealTargetConfiguration): Specifies the configuration whose binaries do not require a "-Platform-Configuration" suffix.
+
+$ bBuildAllPlugins (Boolean): Build all the plugins that we can find, even if they're not enabled. This is particularly useful for content-only projects, where you're building the UE4Editor target but running it with a game that enables a plugin.
+
+$ bBuildAllModules (Boolean): Build all the modules that are valid for this target type. Used for CIS and making installed engine builds.
+
+$ AdditionalPlugins (List&lt;String&gt;): A list of additional plugins which need to be included in this target. This allows referencing non-optional plugin modules which cannot be disabled, and allows building against specific modules in program targets which do not fit the categories in ModuleHostType.
+
+$ EnablePlugins (List&lt;String&gt;): Additional plugins that should be included for this target.
+
+$ DisablePlugins (List&lt;String&gt;): List of plugins to be disabled for this target. Note that the project file may still reference them, so they should be marked as optional to avoid failing to find them at runtime.
+
+$ PakSigningKeysFile (String): Path to the set of pak signing keys to embed in the executable.
+
+$ SolutionDirectory (String): Allows a Program Target to specify it's own solution folder path.
+
+$ bBuildInSolutionByDefault (Nullable&lt;Boolean&gt;): Whether the target should be included in the default solution build configuration
+
+$ bShouldCompileAsDLL (Boolean): Whether this target should be compiled as a DLL.  Requires LinkType to be set to TargetLinkType.Monolithic.
+
+$ ExeBinariesSubFolder (String): Subfolder to place executables in, relative to the default location.
+
+$ GeneratedCodeVersion (EGeneratedCodeVersion): Allow target module to override UHT code generation version.
+
+$ bEnableMeshEditor (Boolean): Whether to enable the mesh editor.
+
+$ bCompileChaos (Boolean): Whether to compile the Chaos physics plugin.
+
+$ bUseChaos (Boolean): Whether to use the Chaos physics interface. This overrides the physx flags to disable APEX and NvCloth
+
+$ bUseChaosChecked (Boolean): Whether to compile in checked chaos features for debugging
+
+$ bUseChaosMemoryTracking (Boolean): Whether to compile in chaos memory tracking features
+
+$ bCustomSceneQueryStructure (Boolean): Whether scene query acceleration is done by UE4. The physx scene query structure is still created, but we do not use it.
+
+$ bCompilePhysX (Boolean): Whether to include PhysX support.
+
+$ bCompileAPEX (Boolean): Whether to include PhysX APEX support.
+
+$ bCompileNvCloth (Boolean): Whether to include NvCloth.
+
+$ bCompileICU (Boolean): Whether to include ICU unicode/i18n support in Core.
+
+$ bCompileCEF3 (Boolean): Whether to compile CEF3 support.
+
+$ bForceBuildTargetPlatforms (Boolean): Whether to force compiling the target platform modules, even if they wouldn't normally be built.
+
+$ bForceBuildShaderFormats (Boolean): Whether to force compiling shader format modules, even if they wouldn't normally be built.
+
+$ bCompileCustomSQLitePlatform (Boolean): Whether we should compile SQLite using the custom "Unreal" platform (true), or using the native platform (false).
+
+$ bUseCacheFreedOSAllocs (Boolean): Whether to utilize cache freed OS allocs with MallocBinned
+
+$ bCompileAgainstEngine (Boolean): Enabled for all builds that include the engine project.  Disabled only when building standalone apps that only link with Core.
+
+$ bCompileAgainstCoreUObject (Boolean): Enabled for all builds that include the CoreUObject project.  Disabled only when building standalone apps that only link with Core.
+
+$ bCompileAgainstApplicationCore (Boolean): Enabled for builds that need to initialize the ApplicationCore module. Command line utilities do not normally need this.
+
+$ bCompileRecast (Boolean): Whether to compile Recast navmesh generation.
+
+$ bForceEnableExceptions (Boolean): Enable exceptions for all modules.
+
+$ bUseInlining (Boolean): Enable inlining for all modules.
+
+$ bForceEnableObjCExceptions (Boolean): Enable exceptions for all modules.
+
+$ bForceEnableRTTI (Boolean): Enable RTTI for all modules.
+
+$ bCompileWithStatsWithoutEngine (Boolean): Whether to include stats support even without the engine.
+
+$ bCompileWithPluginSupport (Boolean): Whether to include plugin support.
+
+$ bCompileWithAccessibilitySupport (Boolean): Whether to allow accessibility code in both Slate and the OS layer.
+
+$ bUseDebugLiveCodingConsole (Boolean): Whether to enable support for live coding
+
+$ bWithDirectXMath (Boolean): Whether to enable support for DirectX Math
+
+$ bUseLoggingInShipping (Boolean): Whether to turn on logging for test/shipping builds.
+
+$ bLoggingToMemoryEnabled (Boolean): Whether to turn on logging to memory for test/shipping builds.
+
+$ bUseLauncherChecks (Boolean): Whether to check that the process was launched through an external launcher.
+
+$ bUseChecksInShipping (Boolean): Whether to turn on checks (asserts) for test/shipping builds.
+
+$ bCompileFreeType (Boolean): True if we need FreeType support.
+
+$ bCompileForSize (Boolean): True if we want to favor optimizing size over speed.
+
+$ bForceCompileDevelopmentAutomationTests (Boolean): Whether to compile development automation tests.
+
+$ bForceCompilePerformanceAutomationTests (Boolean): Whether to compile performance automation tests.
+
+$ bEventDrivenLoader (Boolean): If true, event driven loader will be used in cooked builds. @todoio This needs to be replaced by a runtime solution after async loading refactor.
+
+$ bUseXGEController (Boolean): Whether the XGE controller worker and modules should be included in the engine build. These are required for distributed shader compilation using the XGE interception interface.
+
+$ bIWYU (Boolean): Enables "include what you use" by default for modules in this target. Changes the default PCH mode for any module in this project to PCHUsageMode.UseExplicitOrSharedPCHs.
+
+$ bEnforceIWYU (Boolean): Enforce "include what you use" rules; warns if monolithic headers (Engine.h, UnrealEd.h, etc...) are used, and checks that source files include their matching header first.
+
+$ bPrecompile (Boolean): Make static libraries for all engine modules as intermediates for this target.
+
+$ bEnableOSX109Support (Boolean): Whether we should compile with support for OS X 10.9 Mavericks. Used for some tools that we need to be compatible with this version of OS X.
+
+$ bIsBuildingConsoleApplication (Boolean): True if this is a console application that's being built.
+
+$ bDisableSymbolCache (Boolean): True if debug symbols that are cached for some platforms should not be created.
+
+$ bUseUnityBuild (Boolean): Whether to unify C++ code into larger files for faster compilation.
+
+$ bForceUnityBuild (Boolean): Whether to force C++ source files to be combined into larger files for faster compilation.
+
+$ bUseAdaptiveUnityBuild (Boolean): Use a heuristic to determine which files are currently being iterated on and exclude them from unity blobs, result in faster incremental compile times. The current implementation uses the read-only flag to distinguish the working set, assuming that files will be made writable by the source control system if they are being modified. This is true for Perforce, but not for Git.
+
+$ bAdaptiveUnityDisablesOptimizations (Boolean): Disable optimization for files that are in the adaptive non-unity working set.
+
+$ bAdaptiveUnityDisablesPCH (Boolean): Disables force-included PCHs for files that are in the adaptive non-unity working set.
+
+$ bAdaptiveUnityCreatesDedicatedPCH (Boolean): Creates a dedicated PCH for each source file in the working set, allowing faster iteration on cpp-only changes.
+
+$ bAdaptiveUnityEnablesEditAndContinue (Boolean): Creates a dedicated PCH for each source file in the working set, allowing faster iteration on cpp-only changes.
+
+$ MinGameModuleSourceFilesForUnityBuild (Int32): The number of source files in a game module before unity build will be activated for that module.  This allows small game modules to have faster iterative compile times for single files, at the expense of slower full rebuild times.  This setting can be overridden by the bFasterWithoutUnity option in a module's Build.cs file.
+
+$ bShadowVariableErrors (Boolean): Forces shadow variable warnings to be treated as errors on platforms that support it.
+
+$ bUndefinedIdentifierErrors (Boolean): Forces the use of undefined identifiers in conditional expressions to be treated as errors.
+
+$ bUseFastMonoCalls (Boolean): New Monolithic Graphics drivers have optional "fast calls" replacing various D3d functions
+
+$ bUseFastSemanticsRenderContexts (Boolean): New Xbox driver supports a "fast semantics" context type. This switches it on for the immediate and deferred contexts Try disabling this if you see rendering issues and/or crashes inthe Xbox RHI.
+
+$ NumIncludedBytesPerUnityCPP (Int32): An approximate number of bytes of C++ code to target for inclusion in a single unified C++ file.
+
+$ bStressTestUnity (Boolean): Whether to stress test the C++ unity build robustness by including all C++ files files in a project from a single unified file.
+
+$ bForceDebugInfo (Boolean): Whether to force debug info to be generated.
+
+$ bDisableDebugInfo (Boolean): Whether to globally disable debug info generation; see DebugInfoHeuristics.cs for per-config and per-platform options.
+
+$ bDisableDebugInfoForGeneratedCode (Boolean): Whether to disable debug info generation for generated files. This improves link times for modules that have a lot of generated glue code.
+
+$ bOmitPCDebugInfoInDevelopment (Boolean): Whether to disable debug info on PC in development builds (for faster developer iteration, as link times are extremely fast with debug info disabled).
+
+$ bUsePDBFiles (Boolean): Whether PDB files should be used for Visual C++ builds.
+
+$ bUsePCHFiles (Boolean): Whether PCH files should be used.
+
+$ bPreprocessOnly (Boolean): Whether to just preprocess source files for this target, and skip compilation
+
+$ MinFilesUsingPrecompiledHeader (Int32): The minimum number of files that must use a pre-compiled header before it will be created and used.
+
+$ bForcePrecompiledHeaderForGameModules (Boolean): When enabled, a precompiled header is always generated for game modules, even if there are only a few source files in the module.  This greatly improves compile times for iterative changes on a few files in the project, at the expense of slower full rebuild times for small game projects.  This can be overridden by setting MinFilesUsingPrecompiledHeaderOverride in a module's Build.cs file.
+
+$ bUseIncrementalLinking (Boolean): Whether to use incremental linking or not. Incremental linking can yield faster iteration times when making small changes. Currently disabled by default because it tends to behave a bit buggy on some computers (PDB-related compile errors).
+
+$ bAllowLTCG (Boolean): Whether to allow the use of link time code generation (LTCG).
+
+$ bPGOProfile (Boolean): Whether to enable Profile Guided Optimization (PGO) instrumentation in this build.
+
+$ bPGOOptimize (Boolean): Whether to optimize this build with Profile Guided Optimization (PGO).
+
+$ bAllowASLRInShipping (Boolean): Whether to allow the use of ASLR (address space layout randomization) if supported. Only applies to shipping builds.
+
+$ bSupportEditAndContinue (Boolean): Whether to support edit and continue.  Only works on Microsoft compilers.
+
+$ bOmitFramePointers (Boolean): Whether to omit frame pointers or not. Disabling is useful for e.g. memory profiling on the PC.
+
+$ bUseMallocProfiler (Boolean): If true, then enable memory profiling in the build (defines USE_MALLOC_PROFILER=1 and forces bOmitFramePointers=false).
+
+$ bUseSharedPCHs (Boolean): Enables "Shared PCHs", a feature which significantly speeds up compile times by attempting to share certain PCH files between modules that UBT detects is including those PCH's header files.
+
+$ bUseShippingPhysXLibraries (Boolean): True if Development and Release builds should use the release configuration of PhysX/APEX.
+
+$ bUseCheckedPhysXLibraries (Boolean): True if Development and Release builds should use the checked configuration of PhysX/APEX. if bUseShippingPhysXLibraries is true this is ignored.
+
+$ bCheckLicenseViolations (Boolean): Tells the UBT to check if module currently being built is violating EULA.
+
+$ bBreakBuildOnLicenseViolation (Boolean): Tells the UBT to break build if module currently being built is violating EULA.
+
+$ bUseFastPDBLinking (Nullable&lt;Boolean&gt;): Whether to use the :FASTLINK option when building with /DEBUG to create local PDBs on Windows. Fast, but currently seems to have problems finding symbols in the debugger.
+
+$ bCreateMapFile (Boolean): Outputs a map file as part of the build.
+
+$ bAllowRuntimeSymbolFiles (Boolean): True if runtime symbols files should be generated as a post build step for some platforms. These files are used by the engine to resolve symbol names of callstack backtraces in logs.
+
+$ BundleVersion (String): Bundle version for Mac apps.
+
+$ bDeployAfterCompile (Boolean): Whether to deploy the executable after compilation on platforms that require deployment.
+
+$ bAllowRemotelyCompiledPCHs (Boolean): When enabled, allows XGE to compile pre-compiled header files on remote machines.  Otherwise, PCHs are always generated locally.
+
+$ bCheckSystemHeadersForModification (Boolean): Whether headers in system paths should be checked for modification when determining outdated actions.
+
+$ bDisableLinking (Boolean): Whether to disable linking for this target.
+
+$ bFormalBuild (Boolean): Indicates that this is a formal build, intended for distribution. This flag is automatically set to true when Build.version has a changelist set. The only behavior currently bound to this flag is to compile the default resource file separately for each binary so that the OriginalFilename field is set correctly. By default, we only compile the resource once to reduce build times.
+
+$ bFlushBuildDirOnRemoteMac (Boolean): Whether to clean Builds directory on a remote Mac before building.
+
+$ bPrintToolChainTimingInfo (Boolean): Whether to write detailed timing info from the compiler and linker.
+
+$ bParseTimingInfoForTracing (Boolean): Whether to parse timing data into a tracing file compatible with chrome://tracing.
+
+$ bPublicSymbolsByDefault (Boolean): Whether to expose all symbols as public by default on POSIX platforms
+
+$ ToolChainName (String): Allows overriding the toolchain to be created for this target. This must match the name of a class declared in the UnrealBuildTool assembly.
+
+$ bDisableUnverifiedCertificates (Boolean): Whether to allow engine configuration to determine if we can load unverified certificates.
+
+$ bAllowGeneratedIniWhenCooked (Boolean): Whether to load generated ini files in cooked build, (GameUserSettings.ini loaded either way)
+
+$ bAllowNonUFSIniWhenCooked (Boolean): Whether to load non-ufs ini files in cooked build, (GameUserSettings.ini loaded either way)
+
+$ CppStandard (CppStandardVersion): Which C++ stanard to use for compiling this target
+
+$ BuildVersion (String): The build version string
+
+$ GlobalDefinitions (List&lt;String&gt;): Macros to define globally across the whole target.
+
+$ ProjectDefinitions (List&lt;String&gt;): Macros to define across all macros in the project.
+
+$ ExportPublicHeader (String): Specifies the path to write a header containing public definitions for this target. Useful when building a DLL to be consumed by external build processes.
+
+$ ExtraModuleNames (List&lt;String&gt;): List of additional modules to be compiled into the target.
+
+$ ManifestFileNames (List&lt;FileReference&gt;): Path to a manifest to output for this target
+
+$ DependencyListFileNames (List&lt;FileReference&gt;): Path to a list of dependencies for this target, when precompiling
+
+$ bOverrideBuildEnvironment (Boolean): Whether to ignore violations to the shared build environment (eg. editor targets modifying definitions)
+
+$ PreBuildSteps (List&lt;String&gt;): Specifies a list of steps which should be executed before this target is built, in the context of the host platform's shell. The following variables will be expanded before execution: $(EngineDir), $(ProjectDir), $(TargetName), $(TargetPlatform), $(TargetConfiguration), $(TargetType), $(ProjectFile).
+
+$ PostBuildSteps (List&lt;String&gt;): Specifies a list of steps which should be executed after this target is built, in the context of the host platform's shell. The following variables will be expanded before execution: $(EngineDir), $(ProjectDir), $(TargetName), $(TargetPlatform), $(TargetConfiguration), $(TargetType), $(ProjectFile).
+
+$ AdditionalBuildProducts (List&lt;String&gt;): Specifies additional build products produced as part of this target.
+
+$ AdditionalCompilerArguments (String): Additional arguments to pass to the compiler
+
+$ AdditionalLinkerArguments (String): Additional arguments to pass to the linker
+
+$ DisableUnityBuildForModules (String[]): List of modules to disable unity builds for
+
+$ EnableOptimizeCodeForModules (String[]): List of modules to enable optimizations for
+
+$ DisableOptimizeCodeForModules (String[]): List of modules to disable optimizations for
+
+$ GeneratedProjectName (String): When generating project files, specifies the name of the project file to use when there are multiple targets of the same type.
+
+
