Documentation for build tools

diff --git a/Engine/Documentation/Source/Programming/UnrealBuildSystem/Configuration/BuildConfigProperties/BuildConfigProperties.INT.udn b/Engine/Documentation/Source/Programming/UnrealBuildSystem/Configuration/BuildConfigProperties/BuildConfigProperties.INT.udn
--- a/Engine/Documentation/Source/Programming/UnrealBuildSystem/Configuration/BuildConfigProperties/BuildConfigProperties.INT.udn
+++ b/Engine/Documentation/Source/Programming/UnrealBuildSystem/Configuration/BuildConfigProperties/BuildConfigProperties.INT.udn
@@ -138,6 +138,8 @@ $ CppStandard : Which C++ stanard to use for compiling this target
 
 $ bStripSymbols : Whether to strip iOS symbols or not (implied by Shipping config).
 
+$ bEnableMemorySanitizer : Enables memory sanitizer (MSan)
+
 $ bIgnoreOutdatedImportLibraries : Whether to ignore import library files that are out of date when building targets. Set this to true to improve iteration time. By default, we do not bother re-linking targets if only a dependent .lib has changed, as chances are that the import library was not actually different unless a dependent header file of this target was also changed, in which case the target would automatically be rebuilt.
 
 $ bPrintDebugInfo : Whether debug info should be written to the console.
diff --git a/Engine/Documentation/Source/Programming/UnrealBuildSystem/ModuleFiles/ModuleFilesProperties/ModuleFilesProperties.INT.udn b/Engine/Documentation/Source/Programming/UnrealBuildSystem/ModuleFiles/ModuleFilesProperties/ModuleFilesProperties.INT.udn
--- a/Engine/Documentation/Source/Programming/UnrealBuildSystem/ModuleFiles/ModuleFilesProperties/ModuleFilesProperties.INT.udn
+++ b/Engine/Documentation/Source/Programming/UnrealBuildSystem/ModuleFiles/ModuleFilesProperties/ModuleFilesProperties.INT.udn
@@ -19,7 +19,7 @@ $ bTreatAsEngineModule (Boolean): Whether this module should be treated as an en
 
 $ bUseRTTI (Boolean): Use run time type information
 
-$ bUseAVX (Boolean): Use AVX instructions
+$ bUseAVX (Boolean): Direct the compiler to generate AVX instructions wherever SSE or AVX intrinsics are used, on the platforms that support it. Note that by enabling this you are changing the minspec for the PC platform, and the resultant executable will crash on machines without AVX support.
 
 $ bEnableBufferSecurityChecks (Boolean): Enable buffer security checks.  This should usually be enabled as it prevents severe security risks.
 
