Documentation for build tools

diff --git a/Engine/Documentation/Source/Programming/UnrealBuildSystem/ModuleFiles/ModuleFilesProperties/ModuleFilesProperties.INT.udn b/Engine/Documentation/Source/Programming/UnrealBuildSystem/ModuleFiles/ModuleFilesProperties/ModuleFilesProperties.INT.udn
--- a/Engine/Documentation/Source/Programming/UnrealBuildSystem/ModuleFiles/ModuleFilesProperties/ModuleFilesProperties.INT.udn
+++ b/Engine/Documentation/Source/Programming/UnrealBuildSystem/ModuleFiles/ModuleFilesProperties/ModuleFilesProperties.INT.udn
@@ -87,6 +87,8 @@ $ PublicAdditionalFrameworks (List&lt;Framework&gt;): List of addition framework
 
 $ AdditionalBundleResources (List&lt;BundleResource&gt;): List of addition resources that should be copied to the app bundle for Mac or iOS
 
+$ TypeLibraries (List&lt;TypeLibrary&gt;): List of type libraries that we need to generate headers for (Windows only)
+
 $ PublicAdditionalShadowFiles (List&lt;String&gt;): For builds that execute on a remote machine (e.g. iOS), this list contains additional files that need to be copied over in order for the app to link successfully.  Source/header files and PCHs are automatically copied.  Usually this is simply a list of precompiled third party library dependencies.
 
 $ PublicDelayLoadDLLs (List&lt;String&gt;): List of delay load DLLs - typically used for External (third party) modules
