Documentation for build tools

diff --git a/Engine/Documentation/Source/Programming/UnrealBuildSystem/Configuration/BuildConfigProperties/BuildConfigProperties.INT.udn b/Engine/Documentation/Source/Programming/UnrealBuildSystem/Configuration/BuildConfigProperties/BuildConfigProperties.INT.udn
--- a/Engine/Documentation/Source/Programming/UnrealBuildSystem/Configuration/BuildConfigProperties/BuildConfigProperties.INT.udn
+++ b/Engine/Documentation/Source/Programming/UnrealBuildSystem/Configuration/BuildConfigProperties/BuildConfigProperties.INT.udn
@@ -566,6 +566,10 @@ $ bForceDebugUnrealHeaderTool : If true, the Debug version of UnrealHeaderTool w
 
 ### WindowsPlatform
 
+$ MaxRootPathLength : Maximum recommended root path length.
+
+$ MaxNestedPathLength : Maximum length of a path relative to the root directory. Used on Windows to ensure paths are portable between machines. Defaults to off.
+
 $ Compiler : Version of the compiler toolchain to use on Windows platform. A value of "default" will be changed to a specific version at UBT start up.
 
 $ CompilerVersion : The specific toolchain version to use. This may be a specific version number (for example, "14.13.26128"), or the string "Latest", to select the newest available version. By default, and if it is available, we use the toolchain version indicated by WindowsPlatform.DefaultToolChainVersion (otherwise, we use the latest version).
