Documentation for build tools

diff --git a/Engine/Documentation/Source/Programming/UnrealBuildSystem/Configuration/BuildConfigProperties/BuildConfigProperties.INT.udn b/Engine/Documentation/Source/Programming/UnrealBuildSystem/Configuration/BuildConfigProperties/BuildConfigProperties.INT.udn
--- a/Engine/Documentation/Source/Programming/UnrealBuildSystem/Configuration/BuildConfigProperties/BuildConfigProperties.INT.udn
+++ b/Engine/Documentation/Source/Programming/UnrealBuildSystem/Configuration/BuildConfigProperties/BuildConfigProperties.INT.udn
@@ -166,6 +166,8 @@ $ bAllowDistcc : Whether DMUCS/Distcc may be used. Distcc requires some setup --
 
 $ bAllowParallelExecutor : Whether to allow using parallel executor on Windows.
 
+$ MaxParallelActions : Number of actions that can be executed in parallel. If 0 then code will pick a default based on the number of cores available. Only applies to the ParallelExecutor
+
 ### HybridExecutor
 
 $ MaxLocalActions : Maximum number of actions to execute locally.
@@ -228,10 +230,10 @@ $ MaxProcessorCount : Maximum processor count for local execution.
 
 ### ParallelExecutor
 
-$ ProcessorCountMultiplier : Processor count multiplier for local execution. Can be below 1 to reserve CPU for other tasks.
-
 $ MaxProcessorCount : Maximum processor count for local execution.
 
+$ ProcessorCountMultiplier : Processor count multiplier for local execution. Can be below 1 to reserve CPU for other tasks.
+
 $ bStopCompilationAfterErrors : When enabled, will stop compiling targets after a compile error occurs.
 
 ### SNDBS
