Documentation for build tools

diff --git a/Engine/Documentation/Source/Programming/UnrealBuildSystem/Configuration/BuildConfigProperties/BuildConfigProperties.INT.udn b/Engine/Documentation/Source/Programming/UnrealBuildSystem/Configuration/BuildConfigProperties/BuildConfigProperties.INT.udn
--- a/Engine/Documentation/Source/Programming/UnrealBuildSystem/Configuration/BuildConfigProperties/BuildConfigProperties.INT.udn
+++ b/Engine/Documentation/Source/Programming/UnrealBuildSystem/Configuration/BuildConfigProperties/BuildConfigProperties.INT.udn
@@ -446,6 +446,8 @@ $ bMasterProjectNameFromFolder : If true, sets the master project name according
 
 $ Version : The version of Visual Studio to generate project files for.
 
+$ bBuildUBTInDebug : Forces UBT to be built in debug configuration, regardless of the solution configuration
+
 $ bBuildLiveCodingConsole : Whether to include a dependency on LiveCodingConsole when building targets that support live coding.
 
 $ bIncludeDocumentation : True if we should include documentation in the generated projects.
@@ -510,6 +512,8 @@ $ GitPath : Sets the path to use for the Git executable. Defaults to "git" (assu
 
 $ Version : The version of Visual Studio to generate project files for.
 
+$ bBuildUBTInDebug : Forces UBT to be built in debug configuration, regardless of the solution configuration
+
 $ bIncludeDocumentation : True if we should include documentation in the generated projects.
 
 $ bUsePrecompiled : True if build targets should pass the -useprecompiled argument.
