Documentation for build tools

diff --git a/Engine/Documentation/Source/Programming/Development/BuildGraph/ScriptAnatomy/Tasks/PredefinedTasks/BuildGraphScriptPredefinedTasks.INT.udn b/Engine/Documentation/Source/Programming/Development/BuildGraph/ScriptAnatomy/Tasks/PredefinedTasks/BuildGraphScriptPredefinedTasks.INT.udn
--- a/Engine/Documentation/Source/Programming/Development/BuildGraph/ScriptAnatomy/Tasks/PredefinedTasks/BuildGraphScriptPredefinedTasks.INT.udn
+++ b/Engine/Documentation/Source/Programming/Development/BuildGraph/ScriptAnatomy/Tasks/PredefinedTasks/BuildGraphScriptPredefinedTasks.INT.udn
@@ -2,7 +2,7 @@ Availability: NoPublish
 Title: BuildGraph Predefined Tasks
 Crumbs: %ROOT%, Programming, Programming/Development, Programming/Development/BuildGraph, Programming/Development/BuildGraph/BuildGraphScriptTasks
 Description: This is a procedurally generated markdown page.
-version: 4.27
+version: 5.0
 parent:Programming/Development/BuildGraph/BuildGraphScriptTasks
 
 ### AgeStore
@@ -33,14 +33,14 @@ Invokes an AutomationTool child process to run the given command.
 
 Spawns the editor to run a commandlet.
 
-|            |               |          |                                                                                                          |
-| ---------- | ------------- | -------- | -------------------------------------------------------------------------------------------------------- |
-| Name       | String        | Required | The commandlet name to execute.                                                                          |
-| Project    | File Spec     | Optional | The project to run the editor with.                                                                      |
-| Arguments  | String        | Optional | Arguments to be passed to the commandlet.                                                                |
-| EditorExe  | FileReference | Optional | The editor executable to use. Defaults to the development UE4Editor executable for the current platform. |
-| ErrorLevel | Int32         | Optional | The minimum exit code, which is treated as an error.                                                     |
-| If         | Condition     | Optional | Whether to execute this task. It is ignored if this condition evaluates to false.                        |
+|            |               |          |                                                                                                             |
+| ---------- | ------------- | -------- | ----------------------------------------------------------------------------------------------------------- |
+| Name       | String        | Required | The commandlet name to execute.                                                                             |
+| Project    | File Spec     | Optional | The project to run the editor with.                                                                         |
+| Arguments  | String        | Optional | Arguments to be passed to the commandlet.                                                                   |
+| EditorExe  | FileReference | Optional | The editor executable to use. Defaults to the development UnrealEditor executable for the current platform. |
+| ErrorLevel | Int32         | Optional | The minimum exit code, which is treated as an error.                                                        |
+| If         | Condition     | Optional | Whether to execute this task. It is ignored if this condition evaluates to false.                           |
 
 ### Compile
 
@@ -66,7 +66,7 @@ Cook a selection of maps for a certain platform
 |           |           |          |                                                                                   |
 | --------- | --------- | -------- | --------------------------------------------------------------------------------- |
 | Project   | String    | Required | Project file to be cooked.                                                        |
-| Platform  | String    | Required | The cook platform to target (for example, WindowsNoEditor).                       |
+| Platform  | String    | Required | The cook platform to target (for example, Windows).                               |
 | Maps      | String    | Optional | List of maps to be cooked, separated by '+' characters.                           |
 | Versioned | Boolean   | Optional | Additional arguments to be passed to the cooker.                                  |
 | Arguments | String    | Optional | Additional arguments to be passed to the cooker.                                  |
@@ -96,6 +96,7 @@ Compiles C# project files, and their dependencies.
 | Configuration | String    | Optional | The configuration to compile.                                                                                                                       |
 | Platform      | String    | Optional | The platform to compile.                                                                                                                            |
 | Target        | String    | Optional | The target to build.                                                                                                                                |
+| Properties    | String    | Optional | Properties for the command                                                                                                                          |
 | Arguments     | String    | Optional | Additional options to pass to the compiler.                                                                                                         |
 | EnumerateOnly | Boolean   | Optional | Only enumerate build products -- do not actually compile the projects.                                                                              |
 | Tag           | Tag List  | Optional | Tag to be applied to build products of this task.                                                                                                   |
@@ -280,6 +281,7 @@ Creates a new changelist and submits a set of files to a Perforce stream.
 | RootDir         | DirectoryReference | Optional | Root directory for the stream. If not specified, defaults to the current root directory.                                                                                     |
 | RevertUnchanged | Boolean            | Optional | Whether to revert unchanged files before attempting to submit.                                                                                                               |
 | Force           | Boolean            | Optional | Force the submit to happen -- even if a resolve is needed (always accept current version).                                                                                   |
+| P4Verbose       | Boolean            | Optional | Allow verbose P4 output (spew).                                                                                                                                              |
 | If              | Condition          | Optional | Whether to execute this task. It is ignored if this condition evaluates to false.                                                                                            |
 
 ### SymStore
@@ -336,6 +338,25 @@ Extract files from a zip archive.
 | Tag     | Tag List           | Optional | Tag to be applied to the extracted files.                                         |
 | If      | Condition          | Optional | Whether to execute this task. It is ignored if this condition evaluates to false. |
 
+### Upload
+
+Uploads a set of files to Jupiter for future retrival
+
+|                    |                    |          |                                                                                                                                                                                       |
+| ------------------ | ------------------ | -------- | ------------------------------------------------------------------------------------------------------------------------------------------------------------------------------------- |
+| FromDir            | DirectoryReference | Required | The directory to read compressed files from.                                                                                                                                          |
+| Files              | File Spec          | Optional | List of file specifications separated by semicolons (for example, *.cpp;Engine/.../*.bat), or the name of a tag set. Relative paths are taken from FromDir.                           |
+| JupiterNamespace   | String             | Required | The jupiter namespace used to upload the build. Used to control who has access to the build.                                                                                          |
+| JupiterKey         | String             | Required | The key of the build as will be used to download the build again. This has to be globally unique for this particular upload.                                                          |
+| ArchiveType        | String             | Required | The type of archive these files are from, will be added to the metadata                                                                                                               |
+| ProjectName        | String             | Required | The name of the project this set of files are associated with, will be added to the metadata                                                                                          |
+| Branch             | String             | Required | The source control branch these files were generated from, will be added to the metadata                                                                                              |
+| Changelist         | String             | Required | The source control revision these files were generated from, will be added to the metadata                                                                                            |
+| JupiterUrl         | String             | Required | Specify the url to the Jupiter instance to upload to                                                                                                                                  |
+| AdditionalMetadata | String             | Optional | Semi-colon separated list of '=' separated key value mappings to add to the metadata. E.g. Foo=bar;spam=eggs                                                                          |
+| LimitMemoryUsage   | Boolean            | Optional | If enabled file content is not kept in memory, results in lower memory usage but increased io as file contents needs to be read multiple times (for hashing as well as during upload) |
+| If                 | Condition          | Optional | Whether to execute this task. It is ignored if this condition evaluates to false.                                                                                                     |
+
 ### WriteTextFile
 
 Writes text to a file.
diff --git a/Engine/Documentation/Source/Programming/UnrealBuildSystem/Configuration/BuildConfigProperties/BuildConfigProperties.INT.udn b/Engine/Documentation/Source/Programming/UnrealBuildSystem/Configuration/BuildConfigProperties/BuildConfigProperties.INT.udn
--- a/Engine/Documentation/Source/Programming/UnrealBuildSystem/Configuration/BuildConfigProperties/BuildConfigProperties.INT.udn
+++ b/Engine/Documentation/Source/Programming/UnrealBuildSystem/Configuration/BuildConfigProperties/BuildConfigProperties.INT.udn
@@ -2,7 +2,7 @@ Availability: NoPublish
 Title: Build Configuration Properties Page
 Crumbs:
 Description: This is a procedurally generated markdown page.
-Version: 4.27
+Version: 5.0
 
 ### BuildConfiguration
 
@@ -49,8 +49,6 @@ $ bUseDebugLiveCodingConsole : Whether to enable support for live coding
 
 $ bUseXGEController : Whether the XGE controller worker and modules should be included in the engine build. These are required for distributed shader compilation using the XGE interception interface.
 
-$ bUseUnityBuild : Whether to unify C++ code into larger files for faster compilation.
-
 $ bForceUnityBuild : Whether to force C++ source files to be combined into larger files for faster compilation.
 
 $ bUseAdaptiveUnityBuild : Use a heuristic to determine which files are currently being iterated on and exclude them from unity blobs, result in faster incremental compile times. The current implementation uses the read-only flag to distinguish the working set, assuming that files will be made writable by the source control system if they are being modified. This is true for Perforce, but not for Git.
@@ -67,6 +65,8 @@ $ bAdaptiveUnityEnablesEditAndContinue : Creates a dedicated PCH for each source
 
 $ MinGameModuleSourceFilesForUnityBuild : The number of source files in a game module before unity build will be activated for that module.  This allows small game modules to have faster iterative compile times for single files, at the expense of slower full rebuild times.  This setting can be overridden by the bFasterWithoutUnity option in a module's Build.cs file.
 
+$ bWarningsAsErrors : Whether to enable all warnings as errors. UE enables most warnings as errors already, but disables a few (such as deprecation warnings).
+
 $ UnsafeTypeCastWarningLevel : Indicates what warning/error level to treat unsafe type casts as on platforms that support it (e.g., double->float or int64->int32)
 
 $ bUndefinedIdentifierErrors : Forces the use of undefined identifiers in conditional expressions to be treated as errors.
@@ -193,8 +193,6 @@ $ bUsePrecompiled : True if build targets should pass the -useprecompiled argume
 
 $ bIncludeEngineSource : True if we should include engine source in the generated solution.
 
-$ bIncludeEnterpriseSource : Whether to include enterprise source in the generated solution.
-
 $ bIncludeShaderSource : True if shader source files should be included in generated projects.
 
 $ bIncludeBuildSystemFiles : True if build system files should be included.
@@ -271,8 +269,6 @@ $ bUsePrecompiled : True if build targets should pass the -useprecompiled argume
 
 $ bIncludeEngineSource : True if we should include engine source in the generated solution.
 
-$ bIncludeEnterpriseSource : Whether to include enterprise source in the generated solution.
-
 $ bIncludeShaderSource : True if shader source files should be included in generated projects.
 
 $ bIncludeConfigFiles : True if we should include config (.ini) files in the generated project.
@@ -293,8 +289,6 @@ $ bUsePrecompiled : True if build targets should pass the -useprecompiled argume
 
 $ bIncludeEngineSource : True if we should include engine source in the generated solution.
 
-$ bIncludeEnterpriseSource : Whether to include enterprise source in the generated solution.
-
 $ bIncludeShaderSource : True if shader source files should be included in generated projects.
 
 $ bIncludeConfigFiles : True if we should include config (.ini) files in the generated project.
@@ -315,8 +309,6 @@ $ bUsePrecompiled : True if build targets should pass the -useprecompiled argume
 
 $ bIncludeEngineSource : True if we should include engine source in the generated solution.
 
-$ bIncludeEnterpriseSource : Whether to include enterprise source in the generated solution.
-
 $ bIncludeShaderSource : True if shader source files should be included in generated projects.
 
 $ bIncludeConfigFiles : True if we should include config (.ini) files in the generated project.
@@ -337,8 +329,6 @@ $ bUsePrecompiled : True if build targets should pass the -useprecompiled argume
 
 $ bIncludeEngineSource : True if we should include engine source in the generated solution.
 
-$ bIncludeEnterpriseSource : Whether to include enterprise source in the generated solution.
-
 $ bIncludeShaderSource : True if shader source files should be included in generated projects.
 
 $ bIncludeConfigFiles : True if we should include config (.ini) files in the generated project.
@@ -359,8 +349,6 @@ $ bUsePrecompiled : True if build targets should pass the -useprecompiled argume
 
 $ bIncludeEngineSource : True if we should include engine source in the generated solution.
 
-$ bIncludeEnterpriseSource : Whether to include enterprise source in the generated solution.
-
 $ bIncludeShaderSource : True if shader source files should be included in generated projects.
 
 $ bIncludeConfigFiles : True if we should include config (.ini) files in the generated project.
@@ -381,8 +369,6 @@ $ bUsePrecompiled : True if build targets should pass the -useprecompiled argume
 
 $ bIncludeEngineSource : True if we should include engine source in the generated solution.
 
-$ bIncludeEnterpriseSource : Whether to include enterprise source in the generated solution.
-
 $ bIncludeShaderSource : True if shader source files should be included in generated projects.
 
 $ bIncludeConfigFiles : True if we should include config (.ini) files in the generated project.
@@ -403,8 +389,6 @@ $ bUsePrecompiled : True if build targets should pass the -useprecompiled argume
 
 $ bIncludeEngineSource : True if we should include engine source in the generated solution.
 
-$ bIncludeEnterpriseSource : Whether to include enterprise source in the generated solution.
-
 $ bIncludeShaderSource : True if shader source files should be included in generated projects.
 
 $ bIncludeConfigFiles : True if we should include config (.ini) files in the generated project.
@@ -425,8 +409,6 @@ $ bUsePrecompiled : True if build targets should pass the -useprecompiled argume
 
 $ bIncludeEngineSource : True if we should include engine source in the generated solution.
 
-$ bIncludeEnterpriseSource : Whether to include enterprise source in the generated solution.
-
 $ bIncludeShaderSource : True if shader source files should be included in generated projects.
 
 $ bIncludeConfigFiles : True if we should include config (.ini) files in the generated project.
@@ -449,8 +431,6 @@ $ bUsePrecompiled : True if build targets should pass the -useprecompiled argume
 
 $ bIncludeEngineSource : True if we should include engine source in the generated solution.
 
-$ bIncludeEnterpriseSource : Whether to include enterprise source in the generated solution.
-
 $ bIncludeShaderSource : True if shader source files should be included in generated projects.
 
 $ bIncludeConfigFiles : True if we should include config (.ini) files in the generated project.
@@ -481,8 +461,6 @@ $ bUsePrecompiled : True if build targets should pass the -useprecompiled argume
 
 $ bIncludeEngineSource : True if we should include engine source in the generated solution.
 
-$ bIncludeEnterpriseSource : Whether to include enterprise source in the generated solution.
-
 $ bIncludeShaderSource : True if shader source files should be included in generated projects.
 
 $ bIncludeConfigFiles : True if we should include config (.ini) files in the generated project.
@@ -503,8 +481,6 @@ $ bUsePrecompiled : True if build targets should pass the -useprecompiled argume
 
 $ bIncludeEngineSource : True if we should include engine source in the generated solution.
 
-$ bIncludeEnterpriseSource : Whether to include enterprise source in the generated solution.
-
 $ bIncludeShaderSource : True if shader source files should be included in generated projects.
 
 $ bIncludeConfigFiles : True if we should include config (.ini) files in the generated project.
@@ -541,8 +517,6 @@ $ bUsePrecompiled : True if build targets should pass the -useprecompiled argume
 
 $ bIncludeEngineSource : True if we should include engine source in the generated solution.
 
-$ bIncludeEnterpriseSource : Whether to include enterprise source in the generated solution.
-
 $ bIncludeShaderSource : True if shader source files should be included in generated projects.
 
 $ bIncludeConfigFiles : True if we should include config (.ini) files in the generated project.
diff --git a/Engine/Documentation/Source/Programming/UnrealBuildSystem/ModuleFiles/ModuleFilesProperties/ModuleFilesProperties.INT.udn b/Engine/Documentation/Source/Programming/UnrealBuildSystem/ModuleFiles/ModuleFilesProperties/ModuleFilesProperties.INT.udn
--- a/Engine/Documentation/Source/Programming/UnrealBuildSystem/ModuleFiles/ModuleFilesProperties/ModuleFilesProperties.INT.udn
+++ b/Engine/Documentation/Source/Programming/UnrealBuildSystem/ModuleFiles/ModuleFilesProperties/ModuleFilesProperties.INT.udn
@@ -2,7 +2,7 @@ Availability: NoPublish
 Title: Build Configuration Properties Page
 Crumbs:
 Description: This is a procedurally generated markdown page.
-Version: 4.27
+Version: 5.0
 
 ### Read/Write Properties
 $ Type (ModuleType): Type of module
diff --git a/Engine/Documentation/Source/Programming/UnrealBuildSystem/TargetFiles/TargetFilesProperties/TargetFilesProperties.INT.udn b/Engine/Documentation/Source/Programming/UnrealBuildSystem/TargetFiles/TargetFilesProperties/TargetFilesProperties.INT.udn
--- a/Engine/Documentation/Source/Programming/UnrealBuildSystem/TargetFiles/TargetFilesProperties/TargetFilesProperties.INT.udn
+++ b/Engine/Documentation/Source/Programming/UnrealBuildSystem/TargetFiles/TargetFilesProperties/TargetFilesProperties.INT.udn
@@ -2,7 +2,7 @@ Availability: NoPublish
 Title: Build Configuration Properties Page
 Crumbs:
 Description: This is a procedurally generated markdown page.
-Version: 4.27
+Version: 5.0
 
 ### Read-Only Properties
 
@@ -34,7 +34,7 @@ $ bLegalToDistributeBinary (Boolean): Whether the output from this target can be
 
 $ UndecoratedConfiguration (UnrealTargetConfiguration): Specifies the configuration whose binaries do not require a "-Platform-Configuration" suffix.
 
-$ bBuildAllPlugins (Boolean): Build all the plugins that we can find, even if they're not enabled. This is particularly useful for content-only projects, where you're building the UE4Editor target but running it with a game that enables a plugin.
+$ bBuildAllPlugins (Boolean): Build all the plugins that we can find, even if they're not enabled. This is particularly useful for content-only projects, where you're building the UnrealEditor target but running it with a game that enables a plugin.
 
 $ bBuildAllModules (Boolean): Build all the modules that are valid for this target type. Used for CIS and making installed engine builds.
 
@@ -154,8 +154,6 @@ $ bIsBuildingConsoleApplication (Boolean): True if this is a console application
 
 $ bDisableSymbolCache (Boolean): True if debug symbols that are cached for some platforms should not be created.
 
-$ bUseUnityBuild (Boolean): Whether to unify C++ code into larger files for faster compilation.
-
 $ bForceUnityBuild (Boolean): Whether to force C++ source files to be combined into larger files for faster compilation.
 
 $ bUseAdaptiveUnityBuild (Boolean): Use a heuristic to determine which files are currently being iterated on and exclude them from unity blobs, result in faster incremental compile times. The current implementation uses the read-only flag to distinguish the working set, assuming that files will be made writable by the source control system if they are being modified. This is true for Perforce, but not for Git.
@@ -172,6 +170,8 @@ $ MinGameModuleSourceFilesForUnityBuild (Int32): The number of source files in a
 
 $ ShadowVariableWarningLevel (WarningLevel): Forces shadow variable warnings to be treated as errors on platforms that support it.
 
+$ bWarningsAsErrors (Boolean): Whether to enable all warnings as errors. UE enables most warnings as errors already, but disables a few (such as deprecation warnings).
+
 $ UnsafeTypeCastWarningLevel (WarningLevel): Indicates what warning/error level to treat unsafe type casts as on platforms that support it (e.g., double->float or int64->int32)
 
 $ bUndefinedIdentifierErrors (Boolean): Forces the use of undefined identifiers in conditional expressions to be treated as errors.
@@ -214,6 +214,8 @@ $ bSupportEditAndContinue (Boolean): Whether to support edit and continue.  Only
 
 $ bOmitFramePointers (Boolean): Whether to omit frame pointers or not. Disabling is useful for e.g. memory profiling on the PC.
 
+$ bEnableCppModules (Boolean): Whether to enable support for C++20 modules
+
 $ bUseMallocProfiler (Boolean): If true, then enable memory profiling in the build (defines USE_MALLOC_PROFILER=1 and forces bOmitFramePointers=false).
 
 $ bUseSharedPCHs (Boolean): Enables "Shared PCHs", a feature which significantly speeds up compile times by attempting to share certain PCH files between modules that UBT detects is including those PCH's header files.
