Documentation for build tools

diff --git a/Engine/Documentation/Source/Programming/Development/BuildGraph/ScriptAnatomy/Tasks/PredefinedTasks/BuildGraphScriptPredefinedTasks.INT.udn b/Engine/Documentation/Source/Programming/Development/BuildGraph/ScriptAnatomy/Tasks/PredefinedTasks/BuildGraphScriptPredefinedTasks.INT.udn
--- a/Engine/Documentation/Source/Programming/Development/BuildGraph/ScriptAnatomy/Tasks/PredefinedTasks/BuildGraphScriptPredefinedTasks.INT.udn
+++ b/Engine/Documentation/Source/Programming/Development/BuildGraph/ScriptAnatomy/Tasks/PredefinedTasks/BuildGraphScriptPredefinedTasks.INT.udn
@@ -5,380 +5,3 @@ Description: This is a procedurally generated markdown page.
 version: 5.0
 parent:Programming/Development/BuildGraph/BuildGraphScriptTasks
 
-### AgeStore
-
-Task that strips symbols from a set of files. This task is named after the AGESTORE utility that comes with the Microsoft debugger tools SDK, but is actually a separate implementation. The main difference is that it uses the last modified time rather than last access time to determine which files to delete.
-
-|          |                      |          |                                                                                                                                                                                                                                                                                                                                     |
-| -------- | -------------------- | -------- | ----------------------------------------------------------------------------------------------------------------------------------------------------------------------------------------------------------------------------------------------------------------------------------------------------------------------------------- |
-| Platform | UnrealTargetPlatform | Required | The target platform to age symbols for.                                                                                                                                                                                                                                                                                             |
-| StoreDir | String               | Required | The symbol server directory.                                                                                                                                                                                                                                                                                                        |
-| Days     | Int32                | Required | Number of days worth of symbols to keep.                                                                                                                                                                                                                                                                                            |
-| BuildDir | String               | Optional | The root of the build directory to check for existing buildversion named directories.                                                                                                                                                                                                                                               |
-| Filter   | String               | Optional | A substring to match in directory file names before deleting symbols. This allows the "age store" task to avoid deleting symbols from other builds in the case where multiple builds share the same symbol server. Specific use of the filter value is determined by the symbol server structure defined by the platform toolchain. |
-| If       | Condition            | Optional | Whether to execute this task. It is ignored if this condition evaluates to false.                                                                                                                                                                                                                                                   |
-
-### Command
-
-Invokes an AutomationTool child process to run the given command.
-
-|                          |           |          |                                                                                                                                                                   |
-| ------------------------ | --------- | -------- | ----------------------------------------------------------------------------------------------------------------------------------------------------------------- |
-| Name                     | String    | Required | The command name to execute.                                                                                                                                      |
-| Arguments                | String    | Optional | Arguments to be passed to the command.                                                                                                                            |
-| MergeTelemetryWithPrefix | String    | Optional | If non-null, instructs telemetry from the command to be merged into the telemetry for this UAT instance with the given prefix. May be an empty (non-null) string. |
-| If                       | Condition | Optional | Whether to execute this task. It is ignored if this condition evaluates to false.                                                                                 |
-
-### Commandlet
-
-Spawns the editor to run a commandlet.
-
-|            |               |          |                                                                                                             |
-| ---------- | ------------- | -------- | ----------------------------------------------------------------------------------------------------------- |
-| Name       | String        | Required | The commandlet name to execute.                                                                             |
-| Project    | File Spec     | Optional | The project to run the editor with.                                                                         |
-| Arguments  | String        | Optional | Arguments to be passed to the commandlet.                                                                   |
-| EditorExe  | FileReference | Optional | The editor executable to use. Defaults to the development UnrealEditor executable for the current platform. |
-| ErrorLevel | Int32         | Optional | The minimum exit code, which is treated as an error.                                                        |
-| If         | Condition     | Optional | Whether to execute this task. It is ignored if this condition evaluates to false.                           |
-
-### Compile
-
-Compiles a target with UnrealBuildTool.
-
-|                       |                           |          |                                                                                                                                  |
-| --------------------- | ------------------------- | -------- | -------------------------------------------------------------------------------------------------------------------------------- |
-| Target                | String                    | Required | The target to compile.                                                                                                           |
-| Configuration         | UnrealTargetConfiguration | Required | The configuration to compile.                                                                                                    |
-| Platform              | UnrealTargetPlatform      | Required | The platform to compile for.                                                                                                     |
-| Project               | String                    | Optional | The project to compile with.                                                                                                     |
-| Arguments             | String                    | Optional | Additional arguments for UnrealBuildTool.                                                                                        |
-| AllowXGE              | Boolean                   | Optional | Whether to allow using XGE for compilation.                                                                                      |
-| AllowParallelExecutor | Boolean                   | Optional | Whether to allow using the parallel executor for this compile.                                                                   |
-| Clean                 | Nullable`1                | Optional | Whether to allow cleaning this target. If unspecified, targets are cleaned if the -Clean argument is passed on the command line. |
-| Tag                   | Tag List                  | Optional | Tag to be applied to build products of this task.                                                                                |
-| If                    | Condition                 | Optional | Whether to execute this task. It is ignored if this condition evaluates to false.                                                |
-
-### Cook
-
-Cook a selection of maps for a certain platform
-
-|           |           |          |                                                                                   |
-| --------- | --------- | -------- | --------------------------------------------------------------------------------- |
-| Project   | String    | Required | Project file to be cooked.                                                        |
-| Platform  | String    | Required | The cook platform to target (for example, Windows).                               |
-| Maps      | String    | Optional | List of maps to be cooked, separated by '+' characters.                           |
-| Versioned | Boolean   | Optional | Additional arguments to be passed to the cooker.                                  |
-| Arguments | String    | Optional | Additional arguments to be passed to the cooker.                                  |
-| Tag       | Tag List  | Optional | Tag to be applied to build products of this task.                                 |
-| If        | Condition | Optional | Whether to execute this task. It is ignored if this condition evaluates to false. |
-
-### Copy
-
-Copies files from one directory to another.
-
-|           |           |          |                                                                                   |
-| --------- | --------- | -------- | --------------------------------------------------------------------------------- |
-| Files     | File Spec | Optional | Optional filter to be applied to the list of input files.                         |
-| From      | File Spec | Required | The pattern(s) to copy from (for example, Engine/*.txt).                          |
-| To        | File Spec | Required | The directory to copy to.                                                         |
-| Overwrite | Boolean   | Optional | Whether or not to overwrite existing files.                                       |
-| Tag       | Tag List  | Optional | Tag to be applied to build products of this task.                                 |
-| If        | Condition | Optional | Whether to execute this task. It is ignored if this condition evaluates to false. |
-
-### CsCompile
-
-Compiles C# project files, and their dependencies.
-
-|               |           |          |                                                                                                                                                     |
-| ------------- | --------- | -------- | --------------------------------------------------------------------------------------------------------------------------------------------------- |
-| Project       | String    | Required | The C# project file to compile. Using semicolons, more than one project file can be specified.                                                      |
-| Configuration | String    | Optional | The configuration to compile.                                                                                                                       |
-| Platform      | String    | Optional | The platform to compile.                                                                                                                            |
-| Target        | String    | Optional | The target to build.                                                                                                                                |
-| Properties    | String    | Optional | Properties for the command                                                                                                                          |
-| Arguments     | String    | Optional | Additional options to pass to the compiler.                                                                                                         |
-| EnumerateOnly | Boolean   | Optional | Only enumerate build products -- do not actually compile the projects.                                                                              |
-| Tag           | Tag List  | Optional | Tag to be applied to build products of this task.                                                                                                   |
-| TagReferences | Tag List  | Optional | Tag to be applied to any non-private references the projects have. (for example, those that are external and not copied into the output directory). |
-| If            | Condition | Optional | Whether to execute this task. It is ignored if this condition evaluates to false.                                                                   |
-
-### Delete
-
-Delete a set of files.
-
-|                        |           |          |                                                                                                                     |
-| ---------------------- | --------- | -------- | ------------------------------------------------------------------------------------------------------------------- |
-| Files                  | File Spec | Required | List of file specifications separated by semicolons (for example, *.cpp;Engine/.../*.bat), or the name of a tag set |
-| DeleteEmptyDirectories | Boolean   | Optional | Whether to delete empty directories after deleting the files. Defaults to true.                                     |
-| If                     | Condition | Optional | Whether to execute this task. It is ignored if this condition evaluates to false.                                   |
-
-### Log
-
-Print a message (and other optional diagnostic information) to the output log.
-
-|                 |           |          |                                                                                     |
-| --------------- | --------- | -------- | ----------------------------------------------------------------------------------- |
-| Message         | String    | Optional | Message to print out.                                                               |
-| Files           | File Spec | Optional | If specified, causes the given list of files to be printed after the given message. |
-| IncludeContents | Boolean   | Optional | If specified, causes the contents of the given files to be printed out.             |
-| If              | Condition | Optional | Whether to execute this task. It is ignored if this condition evaluates to false.   |
-
-### ModifyConfig
-
-Modifies a config file
-
-|         |           |          |                                                                                   |
-| ------- | --------- | -------- | --------------------------------------------------------------------------------- |
-| File    | File Spec | Required | Path to the config file                                                           |
-| Section | String    | Required | The section name to modify                                                        |
-| Key     | String    | Required | The property name to set                                                          |
-| Value   | String    | Required | The property value to set                                                         |
-| Tag     | Tag List  | Optional | Tag to be applied to the extracted files                                          |
-| If      | Condition | Optional | Whether to execute this task. It is ignored if this condition evaluates to false. |
-
-### Move
-
-Moves files from one directory to another.
-
-|       |           |          |                                                                                   |
-| ----- | --------- | -------- | --------------------------------------------------------------------------------- |
-| Files | File Spec | Optional | Optional filter to be applied to the list of input files.                         |
-| From  | File Spec | Required | The pattern(s) to copy from (for example, Engine/*.txt).                          |
-| To    | File Spec | Required | The directory to copy to.                                                         |
-| Tag   | Tag List  | Optional | Tag to be applied to build products of this task.                                 |
-| If    | Condition | Optional | Whether to execute this task. It is ignored if this condition evaluates to false. |
-
-### MsBuild
-
-Executes MsBuild
-
-|               |           |          |                                                                                                |
-| ------------- | --------- | -------- | ---------------------------------------------------------------------------------------------- |
-| Project       | String    | Required | The C# project file to compile. Using semicolons, more than one project file can be specified. |
-| Configuration | String    | Optional | The configuration to compile.                                                                  |
-| Platform      | String    | Optional | The platform to compile.                                                                       |
-| Arguments     | String    | Optional | Additional options to pass to the compiler.                                                    |
-| Verbosity     | String    | Optional | The MSBuild output verbosity.                                                                  |
-| If            | Condition | Optional | Whether to execute this task. It is ignored if this condition evaluates to false.              |
-
-### PakFile
-
-Creates a PAK file from a given set of files.
-
-|              |               |          |                                                                                                                                   |
-| ------------ | ------------- | -------- | --------------------------------------------------------------------------------------------------------------------------------- |
-| Files        | File Spec     | Required | List of files, wildcards, and tag sets to add to the pak file, separated by ';' characters.                                       |
-| Output       | FileReference | Required | PAK file to output.                                                                                                               |
-| ResponseFile | FileReference | Optional | Path to a Response File that contains a list of files to add to the pak file -- instead of specifying them individually.          |
-| RebaseDir    | HashSet`1     | Optional | Directories to rebase the files relative to. If specified, the shortest path under a listed directory will be used for each file. |
-| Order        | FileReference | Optional | Script that gives the order of files.                                                                                             |
-| Sign         | FileReference | Optional | Encryption keys for this pak file.                                                                                                |
-| Compress     | Boolean       | Optional | Whether to compress files.                                                                                                        |
-| Arguments    | String        | Optional | Additional arguments to pass to UnrealPak.                                                                                        |
-| Tag          | Tag List      | Optional | Tag to be applied to build products of this task.                                                                                 |
-| If           | Condition     | Optional | Whether to execute this task. It is ignored if this condition evaluates to false.                                                 |
-
-### Rename
-
-Renames a file, or group of files.
-
-|       |           |          |                                                                                                          |
-| ----- | --------- | -------- | -------------------------------------------------------------------------------------------------------- |
-| Files | File Spec | Required | The file or files to rename.                                                                             |
-| From  | String    | Optional | The current file name, or pattern to match (for example, *.txt). Should not include any path separators. |
-| To    | String    | Required | The new name for the file(s). Should not include any path separators.                                    |
-| Tag   | Tag List  | Optional | Tag to be applied to the renamed files.                                                                  |
-| If    | Condition | Optional | Whether to execute this task. It is ignored if this condition evaluates to false.                        |
-
-### SanitizeReceipt
-
-Task that tags build products and/or runtime dependencies by reading from *.target files.
-
-|           |                    |          |                                                                                                                                                 |
-| --------- | ------------------ | -------- | ----------------------------------------------------------------------------------------------------------------------------------------------- |
-| Files     | File Spec          | Required | Set of receipt files (*.target) to read, including wildcards and tag names, separated by semicolons.                                            |
-| EngineDir | DirectoryReference | Optional | Path to the Engine folder, used to expand $(EngineDir) properties in receipt files. Defaults to the Engine directory for the current workspace. |
-| If        | Condition          | Optional | Whether to execute this task. It is ignored if this condition evaluates to false.                                                               |
-
-### SetVersion
-
-Updates the local version files (Engine/Source/Runtime/Launch/Resources/Version.h, Engine/Build/Build.version, and Engine/Source/Programs/DotNETCommon/Metadata.cs) with the given version information.
-
-|                  |           |          |                                                                                                                    |
-| ---------------- | --------- | -------- | ------------------------------------------------------------------------------------------------------------------ |
-| Change           | Int32     | Required | The changelist to set in the version files.                                                                        |
-| CompatibleChange | Int32     | Optional | The engine compatible changelist to set in the version files.                                                      |
-| Branch           | String    | Required | The branch string.                                                                                                 |
-| Build            | String    | Optional | The build version string.                                                                                          |
-| Licensee         | Boolean   | Optional | Whether to set the IS_LICENSEE_VERSION flag to true.                                                               |
-| Promoted         | Boolean   | Optional | Whether to set the ENGINE_IS_PROMOTED_BUILD flag to true.                                                          |
-| SkipWrite        | Boolean   | Optional | If set, do not write to the files -- just return the version files that would be updated. Useful for local builds. |
-| Tag              | Tag List  | Optional | Tag to be applied to build products of this task.                                                                  |
-| If               | Condition | Optional | Whether to execute this task. It is ignored if this condition evaluates to false.                                  |
-
-### Sign
-
-Signs a set of executable files with an installed certificate.
-
-|       |           |          |                                                                                                                      |
-| ----- | --------- | -------- | -------------------------------------------------------------------------------------------------------------------- |
-| Files | File Spec | Required | List of file specifications separated by semicolons (for example, *.cpp;Engine/.../*.bat), or the name of a tag set. |
-| Tag   | Tag List  | Optional | Tag to be applied to build products of this task.                                                                    |
-| If    | Condition | Optional | Whether to execute this task. It is ignored if this condition evaluates to false.                                    |
-
-### Spawn
-
-Spawns an external executable and waits for it to complete.
-
-|            |           |          |                                                                                   |
-| ---------- | --------- | -------- | --------------------------------------------------------------------------------- |
-| Exe        | String    | Required | Executable to spawn.                                                              |
-| Arguments  | String    | Optional | Arguments for the newly created process.                                          |
-| ErrorLevel | Int32     | Optional | The minimum exit code, which is treated as an error.                              |
-| If         | Condition | Optional | Whether to execute this task. It is ignored if this condition evaluates to false. |
-
-### Stage
-
-Stages files listed in a build receipt to an output directory.
-
-|               |                           |          |                                                                                   |
-| ------------- | ------------------------- | -------- | --------------------------------------------------------------------------------- |
-| Project       | FileReference             | Optional | The project that this target belongs to.                                          |
-| Target        | String                    | Required | Name of the target to stage.                                                      |
-| Platform      | UnrealTargetPlatform      | Required | Platform to stage.                                                                |
-| Configuration | UnrealTargetConfiguration | Required | Configuration to be staged.                                                       |
-| Architecture  | String                    | Optional | Architecture to be staged.                                                        |
-| ToDir         | DirectoryReference        | Required | Directory that the receipt files should be staged to.                             |
-| Overwrite     | Boolean                   | Optional | Whether to overwrite existing files.                                              |
-| Tag           | Tag List                  | Optional | Tag to be applied to build products of this task.                                 |
-| If            | Condition                 | Optional | Whether to execute this task. It is ignored if this condition evaluates to false. |
-
-### Strip
-
-Strips debugging information from a set of files.
-
-|           |                      |          |                                                                                                                |
-| --------- | -------------------- | -------- | -------------------------------------------------------------------------------------------------------------- |
-| Platform  | UnrealTargetPlatform | Required | The platform toolchain to strip binaries.                                                                      |
-| BaseDir   | DirectoryReference   | Optional | The directory to find files in.                                                                                |
-| Files     | File Spec            | Required | List of file specifications separated by semicolons (for example, Engine/.../*.pdb), or the name of a tag set. |
-| OutputDir | DirectoryReference   | Optional | Output directory for the stripped files. Defaults to the input path, overwriting the input files.              |
-| Tag       | Tag List             | Optional | Tag to be applied to build products of this task.                                                              |
-| If        | Condition            | Optional | Whether to execute this task. It is ignored if this condition evaluates to false.                              |
-
-### Submit
-
-Creates a new changelist and submits a set of files to a Perforce stream.
-
-|                 |                    |          |                                                                                                                                                                              |
-| --------------- | ------------------ | -------- | ---------------------------------------------------------------------------------------------------------------------------------------------------------------------------- |
-| Description     | String             | Required | The description for the submitted changelist.                                                                                                                                |
-| Files           | File Spec          | Required | The files to submit.                                                                                                                                                         |
-| FileType        | String             | Optional | The Perforce file type for the submitted files (for example, binary+FS32).                                                                                                   |
-| Workspace       | String             | Optional | The workspace name. If specified, a new workspace will be created using the given stream and root directory to submit the files. If not, the current workspace will be used. |
-| Stream          | String             | Optional | The stream for the workspace -- defaults to the current stream. Ignored unless the Workspace attribute is also specified.                                                    |
-| RootDir         | DirectoryReference | Optional | Root directory for the stream. If not specified, defaults to the current root directory.                                                                                     |
-| RevertUnchanged | Boolean            | Optional | Whether to revert unchanged files before attempting to submit.                                                                                                               |
-| Force           | Boolean            | Optional | Force the submit to happen -- even if a resolve is needed (always accept current version).                                                                                   |
-| P4Verbose       | Boolean            | Optional | Allow verbose P4 output (spew).                                                                                                                                              |
-| If              | Condition          | Optional | Whether to execute this task. It is ignored if this condition evaluates to false.                                                                                            |
-
-### SymStore
-
-Task that strips symbols from a set of files.
-
-|              |                      |          |                                                                                                                                               |
-| ------------ | -------------------- | -------- | --------------------------------------------------------------------------------------------------------------------------------------------- |
-| Platform     | UnrealTargetPlatform | Required | The platform toolchain required to handle symbol files.                                                                                       |
-| Files        | String               | Required | List of output files. PDBs will be extracted from this list.                                                                                  |
-| StoreDir     | String               | Required | Output directory for the compressed symbols.                                                                                                  |
-| Product      | String               | Required | Name of the product for the symbol store records.                                                                                             |
-| BuildVersion | String               | Optional | BuildVersion associated with these symbols. Used for clean-up in AgeStore by matching this version against a directory name in a build share. |
-| If           | Condition            | Optional | Whether to execute this task. It is ignored if this condition evaluates to false.                                                             |
-
-### Tag
-
-Applies a tag to a given set of files. The list of files is found by enumerating the tags and file specifications given by the 'Files' parameter. From this list, any files not matched by the 'Filter' parameter are removed, followed by any files matched by the 'Except' parameter.
-
-|           |                    |          |                                                                                                                                                                                                                           |
-| --------- | ------------------ | -------- | ------------------------------------------------------------------------------------------------------------------------------------------------------------------------------------------------------------------------- |
-| BaseDir   | DirectoryReference | Optional | Set the base directory to resolve relative paths and patterns against. If set, any absolute patterns (for example, /Engine/Build/...) are taken to be relative to this path. If not, they are taken to be truly absolute. |
-| Files     | File Spec          | Optional | Set of files to work from, including wildcards and tag names, separated by semicolons. If set, resolved relative to BaseDir, otherwise resolved to the branch root directory.                                             |
-| FileLists | File Spec          | Optional | Set of text files to add additional files from. Each file list should have one file per line.                                                                                                                             |
-| Filter    | File Spec          | Optional | Patterns to filter the list of files by, including tag names or wildcards. If set, may include patterns that apply to the base directory. If not specified, defaults to all files.                                        |
-| Except    | File Spec          | Optional | Set of patterns to exclude from the matched list. May include tag names of patterns that apply to the base directory.                                                                                                     |
-| With      | Tag List           | Required | Name of the tag to apply.                                                                                                                                                                                                 |
-| If        | Condition          | Optional | Whether to execute this task. It is ignored if this condition evaluates to false.                                                                                                                                         |
-
-### TagReceipt
-
-Task that tags build products and/or runtime dependencies by reading from *.target files.
-
-|                     |                    |          |                                                                                                                                                                 |
-| ------------------- | ------------------ | -------- | --------------------------------------------------------------------------------------------------------------------------------------------------------------- |
-| Files               | File Spec          | Required | Set of receipt files (*.target) to read, including wildcards and tag names, separated by semicolons.                                                            |
-| EngineDir           | DirectoryReference | Optional | Path to the Engine folder, used to expand $(EngineDir) properties in receipt files. Defaults to the Engine directory for the current workspace.                 |
-| ProjectDir          | DirectoryReference | Optional | Path to the project folder, used to expand $(ProjectDir) properties in receipt files. Defaults to the Engine directory for the current workspace -- DEPRECATED. |
-| BuildProducts       | Boolean            | Optional | Whether to tag the Build Products listed in receipts.                                                                                                           |
-| BuildProductType    | String             | Optional | Which type of Build Products to tag (see TargetReceipt.cs - UnrealBuildTool.BuildProductType for valid values).                                                 |
-| RuntimeDependencies | Boolean            | Optional | Whether to tag the Runtime Dependencies listed in receipts.                                                                                                     |
-| StagedFileType      | String             | Optional | Which type of Runtime Dependencies to tag (see TargetReceipt.cs - UnrealBuildTool.StagedFileType for valid values).                                             |
-| With                | Tag List           | Required | Name of the tag to apply.                                                                                                                                       |
-| If                  | Condition          | Optional | Whether to execute this task. It is ignored if this condition evaluates to false.                                                                               |
-
-### Unzip
-
-Extract files from a zip archive.
-
-|         |                    |          |                                                                                   |
-| ------- | ------------------ | -------- | --------------------------------------------------------------------------------- |
-| ZipFile | File Spec          | Required | Path to the zip file to extract.                                                  |
-| ToDir   | DirectoryReference | Required | Output directory for the extracted files.                                         |
-| Tag     | Tag List           | Optional | Tag to be applied to the extracted files.                                         |
-| If      | Condition          | Optional | Whether to execute this task. It is ignored if this condition evaluates to false. |
-
-### Upload
-
-Uploads a set of files to Jupiter for future retrival
-
-|                    |                    |          |                                                                                                                                                                                       |
-| ------------------ | ------------------ | -------- | ------------------------------------------------------------------------------------------------------------------------------------------------------------------------------------- |
-| FromDir            | DirectoryReference | Required | The directory to read compressed files from.                                                                                                                                          |
-| Files              | File Spec          | Optional | List of file specifications separated by semicolons (for example, *.cpp;Engine/.../*.bat), or the name of a tag set. Relative paths are taken from FromDir.                           |
-| JupiterNamespace   | String             | Required | The jupiter namespace used to upload the build. Used to control who has access to the build.                                                                                          |
-| JupiterKey         | String             | Required | The key of the build as will be used to download the build again. This has to be globally unique for this particular upload.                                                          |
-| ArchiveType        | String             | Required | The type of archive these files are from, will be added to the metadata                                                                                                               |
-| ProjectName        | String             | Required | The name of the project this set of files are associated with, will be added to the metadata                                                                                          |
-| Branch             | String             | Required | The source control branch these files were generated from, will be added to the metadata                                                                                              |
-| Changelist         | String             | Required | The source control revision these files were generated from, will be added to the metadata                                                                                            |
-| JupiterUrl         | String             | Required | Specify the url to the Jupiter instance to upload to                                                                                                                                  |
-| AdditionalMetadata | String             | Optional | Semi-colon separated list of '=' separated key value mappings to add to the metadata. E.g. Foo=bar;spam=eggs                                                                          |
-| LimitMemoryUsage   | Boolean            | Optional | If enabled file content is not kept in memory, results in lower memory usage but increased io as file contents needs to be read multiple times (for hashing as well as during upload) |
-| If                 | Condition          | Optional | Whether to execute this task. It is ignored if this condition evaluates to false.                                                                                                     |
-
-### WriteTextFile
-
-Writes text to a file.
-
-|        |               |          |                                                                                     |
-| ------ | ------------- | -------- | ----------------------------------------------------------------------------------- |
-| File   | FileReference | Required | Path to the file to write.                                                          |
-| Append | Boolean       | Optional | Optional, whether or not to append to the file rather than overwrite.               |
-| Text   | String        | Optional | The text to write to the file.                                                      |
-| Files  | File Spec     | Optional | If specified, causes the given list of files to be printed after the given message. |
-| Tag    | Tag List      | Optional | Tag to be applied to build products of this task.                                   |
-| If     | Condition     | Optional | Whether to execute this task. It is ignored if this condition evaluates to false.   |
-
-### Zip
-
-Compresses files into a zip archive.
-
-|         |                    |          |                                                                                                                                                             |
-| ------- | ------------------ | -------- | ----------------------------------------------------------------------------------------------------------------------------------------------------------- |
-| FromDir | DirectoryReference | Required | The directory to read compressed files from.                                                                                                                |
-| Files   | File Spec          | Optional | List of file specifications separated by semicolons (for example, *.cpp;Engine/.../*.bat), or the name of a tag set. Relative paths are taken from FromDir. |
-| ZipFile | FileReference      | Required | The zip file to create.                                                                                                                                     |
-| Tag     | Tag List           | Optional | Tag to be applied to the created zip file.                                                                                                                  |
-| If      | Condition          | Optional | Whether to execute this task. It is ignored if this condition evaluates to false.                                                                           |
-
diff --git a/Engine/Documentation/Source/Programming/UnrealBuildSystem/Configuration/BuildConfigProperties/BuildConfigProperties.INT.udn b/Engine/Documentation/Source/Programming/UnrealBuildSystem/Configuration/BuildConfigProperties/BuildConfigProperties.INT.udn
--- a/Engine/Documentation/Source/Programming/UnrealBuildSystem/Configuration/BuildConfigProperties/BuildConfigProperties.INT.udn
+++ b/Engine/Documentation/Source/Programming/UnrealBuildSystem/Configuration/BuildConfigProperties/BuildConfigProperties.INT.udn
@@ -6,40 +6,31 @@
 
 ### BuildConfiguration
 
-$ bGeneratedSYMFile : Whether to generate dSYM files. Defaults to true for Shipping builds and false for all other builds. -dsym will force generation of a dsym for other builds,
-        * NoDsym will force it off for shipping.
+$ bIgnoreOutdatedImportLibraries : Whether to ignore import library files that are out of date when building targets. Set this to true to improve iteration time. By default, we do not bother re-linking targets if only a dependent .lib has changed, as chances are that the import library was not actually different unless a dependent header file of this target was also changed, in which case the target would automatically be rebuilt.
 
-$ bEnableAddressSanitizer : Enables address sanitizer (ASan)
+$ bPrintDebugInfo : Whether debug info should be written to the console.
 
-$ bEnableThreadSanitizer : Enables thread sanitizer (TSan)
+$ bAllowHybridExecutor : Whether the hybrid executor will be used (a remote executor and local executor).
 
-$ bEnableUndefinedBehaviorSanitizer : Enables undefined behavior sanitizer (UBSan)
+$ bAllowXGE : Whether XGE may be used.
 
-$ BaseLogFileName : Specifies the file to use for logging.
+$ bAllowFASTBuild : Whether FASTBuild may be used.
 
-$ bAllowDistccLocalFallback : When enabled, allows DMUCS/Distcc to fallback to local compilation when remote compiling fails. Defaults to true as separation of pre-process and compile stages can introduce non-fatal errors.
+$ bAllowSNDBS : Whether SN-DBS may be used.
 
-$ bVerboseDistccOutput : When true enable verbose distcc output to aid debugging.
+$ bUseUBTMakefiles : Enables support for very fast iterative builds by caching target data. Turning this on causes Unreal Build Tool to emit 'UBT Makefiles' for targets when they are built the first time. Subsequent builds will load these Makefiles and begin outdatedness checking and build invocation very quickly. The caveat is that if source files are added or removed to the project, UBT will need to gather information about those in order for your build to complete successfully. Currently, you must run the project file generator after adding/removing source files to tell UBT to re-gather this information. 
+    * Events that can invalidate the 'UBT Makefile':
+        * Adding/removing .cpp files
+        * Adding/removing .h files with UObjects
+        * Adding new UObject types to a file that did not previously have any
+        * Changing global build settings (most settings in this file qualify)
+        * Changed code that affects how Unreal Header Tool works 
+    * You can force regeneration of the 'UBT Makefile' by passing the '-gather' argument, or simply regenerating project files. 
+    * This also enables the fast include file dependency scanning and caching system that allows Unreal Build Tool to detect out of date dependencies very quickly. When enabled, a deep C++ include graph does not have to be generated, and instead, we only scan and cache indirect includes for after a dependent build product was already found to be out of date. During the next build, we will load those cached indirect includes and check for outdatedness.
 
-$ DistccExecutablesPath : Path to the Distcc and DMUCS executables.
+$ bAllowParallelExecutor : Whether to allow using parallel executor on Windows.
 
-$ DMUCSCoordinator : The DMUCS coordinator hostname or IP address.
+$ MaxParallelActions : Number of actions that can be executed in parallel. If 0 then code will pick a default based on the number of cores available. Only applies to the ParallelExecutor
 
-$ DMUCSDistProp : The DMUCS distinguishing property value to request for compile hosts.
-
-$ bXGENoWatchdogThread : Whether to use the no_watchdog_thread option to prevent VS2015 toolchain stalls.
-
-$ bShowXGEMonitor : Whether to display the XGE build monitor.
-
-$ bStopXGECompilationAfterErrors : When enabled, XGE will stop compiling targets after a compile error occurs.  Recommended, as it saves computing resources for others.
-
-$ bWriteSolutionOptionFile : Whether to write a solution option (suo) file for the sln.
-
-$ bAddFastPDBToProjects : Whether to add the -FastPDB option to build command lines by default.
-
-$ bUsePerFileIntellisense : Whether to generate per-file intellisense data.
-
-$ bEditorDependsOnShaderCompileWorker : Whether to include a dependency on ShaderCompileWorker when generating project files for the editor.
-
 $ bDebugBuildsActuallyUseDebugCRT : Enables the debug C++ runtime (CRT) for debug builds. By default we always use the release runtime, since the debug version isn't particularly useful when debugging Unreal Engine projects, and linking against the debug CRT libraries forces our third party library dependencies to also be compiled using the debug CRT (and often perform more slowly). Often it can be inconvenient to require a separate copy of the debug versions of third party static libraries simply so that you can debug your program's code.
 
@@ -137,44 +128,123 @@
 
 $ CppStandard : Which C++ stanard to use for compiling this target
 
+$ bXGENoWatchdogThread : Whether to use the no_watchdog_thread option to prevent VS2015 toolchain stalls.
+
+$ bShowXGEMonitor : Whether to display the XGE build monitor.
+
+$ bStopXGECompilationAfterErrors : When enabled, XGE will stop compiling targets after a compile error occurs.  Recommended, as it saves computing resources for others.
+
+$ BaseLogFileName : Specifies the file to use for logging.
+
 $ bStripSymbols : Whether to strip iOS symbols or not (implied by Shipping config).
 
+$ bEnableAddressSanitizer : Enables address sanitizer (ASan).
+
+$ bEnableThreadSanitizer : Enables thread sanitizer (TSan).
+
+$ bEnableUndefinedBehaviorSanitizer : Enables undefined behavior sanitizer (UBSan).
+
 $ bEnableMemorySanitizer : Enables memory sanitizer (MSan)
 
-$ bIgnoreOutdatedImportLibraries : Whether to ignore import library files that are out of date when building targets. Set this to true to improve iteration time. By default, we do not bother re-linking targets if only a dependent .lib has changed, as chances are that the import library was not actually different unless a dependent header file of this target was also changed, in which case the target would automatically be rebuilt.
+$ bGeneratedSYMFile : Whether to generate dSYM files. Defaults to true for Shipping builds and false for all other builds. -dsym will force generation of a dsym for other builds,
+        * NoDsym will force it off for shipping.
 
-$ bPrintDebugInfo : Whether debug info should be written to the console.
+$ bWriteSolutionOptionFile : Whether to write a solution option (suo) file for the sln.
 
-$ bLogDetailedActionStats : Whether to log detailed action stats. This forces local execution.
+$ bAddFastPDBToProjects : Whether to add the -FastPDB option to build command lines by default.
 
-$ bAllowHybridExecutor : Whether the hybrid executor will be used (a remote executor and local executor).
+$ bUsePerFileIntellisense : Whether to generate per-file intellisense data.
 
-$ bAllowXGE : Whether XGE may be used.
+$ bEditorDependsOnShaderCompileWorker : Whether to include a dependency on ShaderCompileWorker when generating project files for the editor.
 
-$ bAllowFASTBuild : Whether FASTBuild may be used.
+### UEBuildConfiguration
 
-$ bAllowSNDBS : Whether SN-DBS may be used.
+$ bForceHeaderGeneration : If true, force header regeneration. Intended for the build machine.
 
-$ bUseUBTMakefiles : Enables support for very fast iterative builds by caching target data. Turning this on causes Unreal Build Tool to emit 'UBT Makefiles' for targets when they are built the first time. Subsequent builds will load these Makefiles and begin outdatedness checking and build invocation very quickly. The caveat is that if source files are added or removed to the project, UBT will need to gather information about those in order for your build to complete successfully. Currently, you must run the project file generator after adding/removing source files to tell UBT to re-gather this information. 
-    * Events that can invalidate the 'UBT Makefile':
-        * Adding/removing .cpp files
-        * Adding/removing .h files with UObjects
-        * Adding new UObject types to a file that did not previously have any
-        * Changing global build settings (most settings in this file qualify)
-        * Changed code that affects how Unreal Header Tool works 
-    * You can force regeneration of the 'UBT Makefile' by passing the '-gather' argument, or simply regenerating project files. 
-    * This also enables the fast include file dependency scanning and caching system that allows Unreal Build Tool to detect out of date dependencies very quickly. When enabled, a deep C++ include graph does not have to be generated, and instead, we only scan and cache indirect includes for after a dependent build product was already found to be out of date. During the next build, we will load those cached indirect includes and check for outdatedness.
+$ bDoNotBuildUHT : If true, do not build UHT, assume it is already built.
 
-$ bAllowDistcc : Whether DMUCS/Distcc may be used. Distcc requires some setup -- so by default, disable it so that we do not break local or remote building.
+$ bFailIfGeneratedCodeChanges : If true, fail if any of the generated header files is out of date.
 
-$ bAllowParallelExecutor : Whether to allow using parallel executor on Windows.
+$ bAllowHotReloadFromIDE : True if hot-reload from IDE is allowed.
 
-$ MaxParallelActions : Number of actions that can be executed in parallel. If 0 then code will pick a default based on the number of cores available. Only applies to the ParallelExecutor
+$ bForceDebugUnrealHeaderTool : If true, the Debug version of UnrealHeaderTool will be built and run instead of the Development version.
 
+### WindowsPlatform
+
+$ MaxRootPathLength : Maximum recommended root path length.
+
+$ MaxNestedPathLength : Maximum length of a path relative to the root directory. Used on Windows to ensure paths are portable between machines. Defaults to off.
+
+$ Compiler : Version of the compiler toolchain to use on Windows platform. A value of "default" will be changed to a specific version at UBT start up.
+
+$ CompilerVersion : The specific toolchain version to use. This may be a specific version number (for example, "14.13.26128"), the string "Latest" to select the newest available version, or the string "Preview" to select the newest available preview version. By default, and if it is available, we use the toolchain version indicated by WindowsPlatform.DefaultToolChainVersion (otherwise, we use the latest version).
+
+$ WindowsSdkVersion : The specific Windows SDK version to use. This may be a specific version number (for example, "8.1", "10.0" or "10.0.10150.0"), or the string "Latest", to select the newest available version. By default, and if it is available, we use the Windows SDK version indicated by WindowsPlatform.DefaultWindowsSdkVersion (otherwise, we use the latest version).
+
+$ StaticAnalyzer : The static analyzer to use.
+
+$ bStrictConformanceMode : Enables strict standard conformance mode (/permissive-) in VS2017+.
+
+$ PCHMemoryAllocationFactor : Determines the amount of memory that the compiler allocates to construct precompiled headers (/Zm).
+
+$ bClangTimeTrace : (Experimental) Appends the -ftime-trace argument to the command line for Clang to output a JSON file containing a timeline for the compile. See http://aras-p.info/blog/2019/01/16/time-trace-timeline-flame-chart-profiler-for-Clang/ for more info.
+
+$ bCompilerTrace : Outputs compile timing information so that it can be analyzed.
+
+$ bShowIncludes : Print out files that are included by each source file
+
+### ModuleConfiguration
+
+$ DisableUnityBuild : List of modules to disable unity builds for
+
+$ EnableOptimizeCode : List of modules to enable optimizations for
+
+$ DisableOptimizeCode : List of modules to disable optimizations for
+
+### FASTBuild
+
+$ FBuildExecutablePath : Used to specify the location of fbuild.exe if the distributed binary isn't being used
+
+$ bEnableDistribution : Controls network build distribution
+
+$ FBuildBrokeragePath : Used to specify the location of the brokerage. If null, FASTBuild will fall back to checking FASTBUILD_BROKERAGE_PATH
+
+$ bEnableCaching : Controls whether to use caching at all. CachePath and FASTCacheMode are only relevant if this is enabled.
+
+$ CacheMode : Cache access mode - only relevant if bEnableCaching is true;
+
+$ FBuildCachePath : Used to specify the location of the cache. If null, FASTBuild will fall back to checking FASTBUILD_CACHE_PATH
+
+$ bForceRemote : Whether to force remote
+
+$ bStopOnError : Whether to stop on error
+
+$ MsvcCRTRedistVersion : Which MSVC CRT Redist version to use
+
 ### HybridExecutor
 
 $ MaxLocalActions : Maximum number of actions to execute locally.
 
+### LocalExecutor
+
+$ ProcessorCountMultiplier : Processor count multiplier for local execution. Can be below 1 to reserve CPU for other tasks. When using the local executor (not XGE), run a single action on each CPU core. Note that you can set this to a larger value to get slightly faster build times in many cases, but your computer's responsiveness during compiling may be much worse.
+
+$ MaxProcessorCount : Maximum processor count for local execution.
+
+### ParallelExecutor
+
+$ MaxProcessorCount : Maximum processor count for local execution.
+
+$ ProcessorCountMultiplier : Processor count multiplier for local execution. Can be below 1 to reserve CPU for other tasks.
+
+$ bStopCompilationAfterErrors : When enabled, will stop compiling targets after a compile error occurs.
+
+### XGE
+
+$ bAllowOverVpn : When set to false, XGE will not be When enabled, XGE will stop compiling targets after a compile error occurs.  Recommended, as it saves computing resources for others.
+
+$ VpnSubnets : List of subnets containing IP addresses assigned by VPN
+
 ### BuildMode
 
 $ bIgnoreJunk : Whether to skip checking for files identified by the junk manifest.
@@ -209,8 +279,6 @@
 
 $ bIncludeTempTargets : Whether to include temporary targets generated by UAT to support content only projects with non-default settings.
 
-$ bIncludeDotNETCoreProjects : True if we should include .NET Core projects in the generated solution.
-
 $ bKeepSourceSubDirectories : True if we should reflect "Source" sub-directories on disk in the master project as master project directories. This (arguably) adds some visual clutter to the master project but it is truer to the on-disk file organization.
 
 $ Platforms : Names of platforms to include in the generated project files
@@ -223,45 +291,21 @@
 
 $ bIncludeTestAndShippingConfigs : Whether we should include configurations for "Test" and "Shipping" in generated projects. Pass "-NoShippingConfigs" to disable this.
 
-### LocalExecutor
+### HoloLensPlatform
 
-$ ProcessorCountMultiplier : Processor count multiplier for local execution. Can be below 1 to reserve CPU for other tasks. When using the local executor (not XGE), run a single action on each CPU core. Note that you can set this to a larger value to get slightly faster build times in many cases, but your computer's responsiveness during compiling may be much worse.
+$ Compiler : Version of the compiler toolchain to use on HoloLens. A value of "default" will be changed to a specific version at UBT startup.
 
-$ MaxProcessorCount : Maximum processor count for local execution.
+### IOSToolChain
 
-### ParallelExecutor
+$ IOSSDKVersion : The version of the iOS SDK to target at build time.
 
-$ MaxProcessorCount : Maximum processor count for local execution.
+$ BuildIOSVersion : The version of the iOS to allow at build time.
 
-$ ProcessorCountMultiplier : Processor count multiplier for local execution. Can be below 1 to reserve CPU for other tasks.
+$ bUseDangerouslyFastMode : If this is set, then we do not do any post-compile steps -- except moving the executable into the proper spot on Mac.
 
-$ bStopCompilationAfterErrors : When enabled, will stop compiling targets after a compile error occurs.
+### WindowsTargetRules
 
-### XGE
+$ ObjSrcMapFile : Whether we should export a file containing .obj to source file mappings.
 
-$ bAllowOverVpn : When set to false, XGE will not be When enabled, XGE will stop compiling targets after a compile error occurs.  Recommended, as it saves computing resources for others.
-
-$ VpnSubnets : List of subnets containing IP addresses assigned by VPN
-
-### FASTBuild
-
-$ FBuildExecutablePath : Used to specify the location of fbuild.exe if the distributed binary isn't being used
-
-$ bEnableDistribution : Controls network build distribution
-
-$ FBuildBrokeragePath : Used to specify the location of the brokerage. If null, FASTBuild will fall back to checking FASTBUILD_BROKERAGE_PATH
-
-$ bEnableCaching : Controls whether to use caching at all. CachePath and FASTCacheMode are only relevant if this is enabled.
-
-$ CacheMode : Cache access mode - only relevant if bEnableCaching is true;
-
-$ FBuildCachePath : Used to specify the location of the cache. If null, FASTBuild will fall back to checking FASTBUILD_CACHE_PATH
-
-$ bForceRemote : Whether to force remote
-
-$ bStopOnError : Whether to stop on error
-
-$ MsvcCRTRedistVersion : Which MSVC CRT Redist version to use
-
 ### CLionGenerator
 
@@ -445,17 +489,7 @@
 
 $ bMasterProjectNameFromFolder : If true, sets the master project name according to the name of the folder it is in.
 
-### VCProjectFileGenerator
+### VCMacProjectFileGenerator
 
-$ Version : The version of Visual Studio to generate project files for.
-
-$ MaxSharedIncludePaths : Puts the most common include paths in the IncludePath property in the MSBuild project. This significantly reduces Visual Studio memory usage (measured 1.1GB -> 500mb), but seems to be causing issues with Visual Assist. Value here specifies maximum length of the include path list in KB.
-
-$ ExcludedIncludePaths : Semi-colon separated list of paths that should not be added to the projects include paths. Useful for omitting third-party headers (e.g ThirdParty/WebRTC) from intellisense suggestions and reducing memory footprints.
-
-$ bBuildUBTInDebug : Forces UBT to be built in debug configuration, regardless of the solution configuration
-
-$ bBuildLiveCodingConsole : Whether to include a dependency on LiveCodingConsole when building targets that support live coding.
-
 $ bIncludeDocumentation : True if we should include documentation in the generated projects.
 
@@ -475,7 +509,17 @@
 
 $ bMasterProjectNameFromFolder : If true, sets the master project name according to the name of the folder it is in.
 
-### XcodeProjectFileGenerator
+### VCProjectFileGenerator
 
+$ Version : The version of Visual Studio to generate project files for.
+
+$ MaxSharedIncludePaths : Puts the most common include paths in the IncludePath property in the MSBuild project. This significantly reduces Visual Studio memory usage (measured 1.1GB -> 500mb), but seems to be causing issues with Visual Assist. Value here specifies maximum length of the include path list in KB.
+
+$ ExcludedIncludePaths : Semi-colon separated list of paths that should not be added to the projects include paths. Useful for omitting third-party headers (e.g ThirdParty/WebRTC) from intellisense suggestions and reducing memory footprints.
+
+$ bBuildUBTInDebug : Forces UBT to be built in debug configuration, regardless of the solution configuration
+
+$ bBuildLiveCodingConsole : Whether to include a dependency on LiveCodingConsole when building targets that support live coding.
+
 $ bIncludeDocumentation : True if we should include documentation in the generated projects.
 
@@ -495,23 +539,7 @@
 
 $ bMasterProjectNameFromFolder : If true, sets the master project name according to the name of the folder it is in.
 
-### ModuleConfiguration
+### XcodeProjectFileGenerator
 
-$ DisableUnityBuild : List of modules to disable unity builds for
-
-$ EnableOptimizeCode : List of modules to enable optimizations for
-
-$ DisableOptimizeCode : List of modules to disable optimizations for
-
-### SourceFileWorkingSet
-
-$ Provider : Sets the provider to use for determining the working set.
-
-$ RepositoryPath : Sets the path to use for the repository. Interpreted relative to the Unreal Engine root directory (the folder above the Engine folder) -- if relative.
-
-$ GitPath : Sets the path to use for the Git executable. Defaults to "git" (assuming it is in the PATH).
-
-### VCMacProjectFileGenerator
-
 $ bIncludeDocumentation : True if we should include documentation in the generated projects.
 
@@ -531,13 +559,13 @@
 
 $ bMasterProjectNameFromFolder : If true, sets the master project name according to the name of the folder it is in.
 
-### IOSToolChain
+### SourceFileWorkingSet
 
-$ IOSSDKVersion : The version of the iOS SDK to target at build time.
+$ Provider : Sets the provider to use for determining the working set.
 
-$ BuildIOSVersion : The version of the iOS to allow at build time.
+$ RepositoryPath : Sets the path to use for the repository. Interpreted relative to the Unreal Engine root directory (the folder above the Engine folder) -- if relative.
 
-$ bUseDangerouslyFastMode : If this is set, then we do not do any post-compile steps -- except moving the executable into the proper spot on Mac.
+$ GitPath : Sets the path to use for the Git executable. Defaults to "git" (assuming it is in the PATH).
 
 ### RemoteMac
 
@@ -551,47 +579,3 @@
 
 $ SshAuthentication : The authentication used for SSH (probably similar to RsyncAuthentication).
 
-### HoloLensPlatform
-
-$ Compiler : Version of the compiler toolchain to use on HoloLens. A value of "default" will be changed to a specific version at UBT startup.
-
-### UEBuildConfiguration
-
-$ bForceHeaderGeneration : If true, force header regeneration. Intended for the build machine.
-
-$ bDoNotBuildUHT : If true, do not build UHT, assume it is already built.
-
-$ bFailIfGeneratedCodeChanges : If true, fail if any of the generated header files is out of date.
-
-$ bAllowHotReloadFromIDE : True if hot-reload from IDE is allowed.
-
-$ bForceDebugUnrealHeaderTool : If true, the Debug version of UnrealHeaderTool will be built and run instead of the Development version.
-
-### WindowsPlatform
-
-$ MaxRootPathLength : Maximum recommended root path length.
-
-$ MaxNestedPathLength : Maximum length of a path relative to the root directory. Used on Windows to ensure paths are portable between machines. Defaults to off.
-
-$ Compiler : Version of the compiler toolchain to use on Windows platform. A value of "default" will be changed to a specific version at UBT start up.
-
-$ CompilerVersion : The specific toolchain version to use. This may be a specific version number (for example, "14.13.26128"), the string "Latest" to select the newest available version, or the string "Preview" to select the newest available preview version. By default, and if it is available, we use the toolchain version indicated by WindowsPlatform.DefaultToolChainVersion (otherwise, we use the latest version).
-
-$ WindowsSdkVersion : The specific Windows SDK version to use. This may be a specific version number (for example, "8.1", "10.0" or "10.0.10150.0"), or the string "Latest", to select the newest available version. By default, and if it is available, we use the Windows SDK version indicated by WindowsPlatform.DefaultWindowsSdkVersion (otherwise, we use the latest version).
-
-$ StaticAnalyzer : The static analyzer to use.
-
-$ bStrictConformanceMode : Enables strict standard conformance mode (/permissive-) in VS2017+.
-
-$ PCHMemoryAllocationFactor : Determines the amount of memory that the compiler allocates to construct precompiled headers (/Zm).
-
-$ bClangTimeTrace : (Experimental) Appends the -ftime-trace argument to the command line for Clang to output a JSON file containing a timeline for the compile. See http://aras-p.info/blog/2019/01/16/time-trace-timeline-flame-chart-profiler-for-Clang/ for more info.
-
-$ bCompilerTrace : Outputs compile timing information so that it can be analyzed.
-
-$ bShowIncludes : Print out files that are included by each source file
-
-### WindowsTargetRules
-
-$ ObjSrcMapFile : Whether we should export a file containing .obj to source file mappings.
-
