Documentation for build tools

diff --git a/Engine/Documentation/Source/Programming/UnrealBuildSystem/Configuration/BuildConfigProperties/BuildConfigProperties.INT.udn b/Engine/Documentation/Source/Programming/UnrealBuildSystem/Configuration/BuildConfigProperties/BuildConfigProperties.INT.udn
--- a/Engine/Documentation/Source/Programming/UnrealBuildSystem/Configuration/BuildConfigProperties/BuildConfigProperties.INT.udn
+++ b/Engine/Documentation/Source/Programming/UnrealBuildSystem/Configuration/BuildConfigProperties/BuildConfigProperties.INT.udn
@@ -40,6 +40,8 @@ $ bUseDebugLiveCodingConsole : Whether to enable support for live coding
 
 $ bUseXGEController : Whether the XGE controller worker and modules should be included in the engine build. These are required for distributed shader compilation using the XGE interception interface.
 
+$ bUseUnityBuild : Whether to unify C++ code into larger files for faster compilation.
+
 $ bForceUnityBuild : Whether to force C++ source files to be combined into larger files for faster compilation.
 
 $ bUseAdaptiveUnityBuild : Use a heuristic to determine which files are currently being iterated on and exclude them from unity blobs, result in faster incremental compile times. The current implementation uses the read-only flag to distinguish the working set, assuming that files will be made writable by the source control system if they are being modified. This is true for Perforce, but not for Git.
