Documentation for build tools

diff --git a/Engine/Documentation/Source/Programming/UnrealBuildSystem/Configuration/BuildConfigProperties/BuildConfigProperties.INT.udn b/Engine/Documentation/Source/Programming/UnrealBuildSystem/Configuration/BuildConfigProperties/BuildConfigProperties.INT.udn
--- a/Engine/Documentation/Source/Programming/UnrealBuildSystem/Configuration/BuildConfigProperties/BuildConfigProperties.INT.udn
+++ b/Engine/Documentation/Source/Programming/UnrealBuildSystem/Configuration/BuildConfigProperties/BuildConfigProperties.INT.udn
@@ -32,6 +32,14 @@ $ bAllowParallelExecutor : Whether to allow using parallel executor on Windows.
 
 $ MaxParallelActions : Number of actions that can be executed in parallel. If 0 then code will pick a default based on the number of cores available. Only applies to the ParallelExecutor
 
+$ bEnableAddressSanitizer : Enables address sanitizer (ASan).
+
+$ bEnableThreadSanitizer : Enables thread sanitizer (TSan).
+
+$ bEnableUndefinedBehaviorSanitizer : Enables undefined behavior sanitizer (UBSan).
+
+$ bEnableMemorySanitizer : Enables memory sanitizer (MSan)
+
 $ bDebugBuildsActuallyUseDebugCRT : Enables the debug C++ runtime (CRT) for debug builds. By default we always use the release runtime, since the debug version isn't particularly useful when debugging Unreal Engine projects, and linking against the debug CRT libraries forces our third party library dependencies to also be compiled using the debug CRT (and often perform more slowly). Often it can be inconvenient to require a separate copy of the debug versions of third party static libraries simply so that you can debug your program's code.
 
 $ bUseInlining : Enable inlining for all modules.
@@ -140,14 +148,6 @@ $ BaseLogFileName : Specifies the file to use for logging.
 
 $ bStripSymbols : Whether to strip iOS symbols or not (implied by Shipping config).
 
-$ bEnableAddressSanitizer : Enables address sanitizer (ASan).
-
-$ bEnableThreadSanitizer : Enables thread sanitizer (TSan).
-
-$ bEnableUndefinedBehaviorSanitizer : Enables undefined behavior sanitizer (UBSan).
-
-$ bEnableMemorySanitizer : Enables memory sanitizer (MSan)
-
 $ bGeneratedSYMFile : Whether to generate dSYM files. Defaults to true for Shipping builds and false for all other builds. -dsym will force generation of a dsym for other builds,
         * NoDsym will force it off for shipping.
 
@@ -243,7 +243,7 @@ $ bStopCompilationAfterErrors : When enabled, will stop compiling targets after
 
 ### XGE
 
-$ bAllowOverVpn : When set to false, XGE will not be When enabled, XGE will stop compiling targets after a compile error occurs.  Recommended, as it saves computing resources for others.
+$ bAllowOverVpn : When set to false, XGE will not be enabled when running connected to the coordinator over VPN. Configure VPN-assigned subnets via the VpnSubnets parameter.
 
 $ VpnSubnets : List of subnets containing IP addresses assigned by VPN
 
