Documentation for build tools

diff --git a/Engine/Documentation/Source/Programming/UnrealBuildSystem/Configuration/BuildConfigProperties/BuildConfigProperties.INT.udn b/Engine/Documentation/Source/Programming/UnrealBuildSystem/Configuration/BuildConfigProperties/BuildConfigProperties.INT.udn
--- a/Engine/Documentation/Source/Programming/UnrealBuildSystem/Configuration/BuildConfigProperties/BuildConfigProperties.INT.udn
+++ b/Engine/Documentation/Source/Programming/UnrealBuildSystem/Configuration/BuildConfigProperties/BuildConfigProperties.INT.udn
@@ -32,7 +32,7 @@ $ bAllowParallelExecutor : Whether to allow using parallel executor on Windows.
 
 $ MaxParallelActions : Number of actions that can be executed in parallel. If 0 then code will pick a default based on the number of cores available. Only applies to the ParallelExecutor
 
-$ bEnableAddressSanitizer : Enables address sanitizer (ASan).
+$ bEnableAddressSanitizer : Enables address sanitizer (ASan). Only supported for Visual Studio 2019 16.7.0 and up.
 
 $ bEnableThreadSanitizer : Enables thread sanitizer (TSan).
 
