Documentation for build tools

diff --git a/Engine/Documentation/Source/Programming/UnrealBuildSystem/Configuration/BuildConfigProperties/BuildConfigProperties.INT.udn b/Engine/Documentation/Source/Programming/UnrealBuildSystem/Configuration/BuildConfigProperties/BuildConfigProperties.INT.udn
--- a/Engine/Documentation/Source/Programming/UnrealBuildSystem/Configuration/BuildConfigProperties/BuildConfigProperties.INT.udn
+++ b/Engine/Documentation/Source/Programming/UnrealBuildSystem/Configuration/BuildConfigProperties/BuildConfigProperties.INT.udn
@@ -64,6 +64,8 @@ $ bAdaptiveUnityCreatesDedicatedPCH : Creates a dedicated PCH for each source fi
 
 $ bAdaptiveUnityEnablesEditAndContinue : Creates a dedicated PCH for each source file in the working set, allowing faster iteration on cpp-only changes.
 
+$ bAdaptiveUnityCompilesHeaderFiles : Creates a dedicated source file for each header file in the working set to detect missing includes in headers.
+
 $ MinGameModuleSourceFilesForUnityBuild : The number of source files in a game module before unity build will be activated for that module.  This allows small game modules to have faster iterative compile times for single files, at the expense of slower full rebuild times.  This setting can be overridden by the bFasterWithoutUnity option in a module's Build.cs file.
 
 $ bWarningsAsErrors : Whether to enable all warnings as errors. UE enables most warnings as errors already, but disables a few (such as deprecation warnings).
diff --git a/Engine/Documentation/Source/Programming/UnrealBuildSystem/TargetFiles/TargetFilesProperties/TargetFilesProperties.INT.udn b/Engine/Documentation/Source/Programming/UnrealBuildSystem/TargetFiles/TargetFilesProperties/TargetFilesProperties.INT.udn
--- a/Engine/Documentation/Source/Programming/UnrealBuildSystem/TargetFiles/TargetFilesProperties/TargetFilesProperties.INT.udn
+++ b/Engine/Documentation/Source/Programming/UnrealBuildSystem/TargetFiles/TargetFilesProperties/TargetFilesProperties.INT.udn
@@ -166,6 +166,8 @@ $ bAdaptiveUnityCreatesDedicatedPCH (Boolean): Creates a dedicated PCH for each
 
 $ bAdaptiveUnityEnablesEditAndContinue (Boolean): Creates a dedicated PCH for each source file in the working set, allowing faster iteration on cpp-only changes.
 
+$ bAdaptiveUnityCompilesHeaderFiles (Boolean): Creates a dedicated source file for each header file in the working set to detect missing includes in headers.
+
 $ MinGameModuleSourceFilesForUnityBuild (Int32): The number of source files in a game module before unity build will be activated for that module.  This allows small game modules to have faster iterative compile times for single files, at the expense of slower full rebuild times.  This setting can be overridden by the bFasterWithoutUnity option in a module's Build.cs file.
 
 $ ShadowVariableWarningLevel (WarningLevel): Forces shadow variable warnings to be treated as errors on platforms that support it.
