Documentation for build tools

diff --git a/Engine/Documentation/Source/Programming/UnrealBuildSystem/Configuration/BuildConfigProperties/BuildConfigProperties.INT.udn b/Engine/Documentation/Source/Programming/UnrealBuildSystem/Configuration/BuildConfigProperties/BuildConfigProperties.INT.udn
--- a/Engine/Documentation/Source/Programming/UnrealBuildSystem/Configuration/BuildConfigProperties/BuildConfigProperties.INT.udn
+++ b/Engine/Documentation/Source/Programming/UnrealBuildSystem/Configuration/BuildConfigProperties/BuildConfigProperties.INT.udn
@@ -213,6 +213,8 @@ $ bEnableDistribution : Controls network build distribution
 
 $ FBuildBrokeragePath : Used to specify the location of the brokerage. If null, FASTBuild will fall back to checking FASTBUILD_BROKERAGE_PATH
 
+$ FBuildCoordinator : Used to specify the FASTBuild coordinator IP or network name. If null, FASTBuild will fall back to checking FASTBUILD_COORDINATOR
+
 $ bEnableCaching : Controls whether to use caching at all. CachePath and FASTCacheMode are only relevant if this is enabled.
 
 $ CacheMode : Cache access mode - only relevant if bEnableCaching is true;
