Documentation for build tools

diff --git a/Engine/Documentation/Source/Programming/UnrealBuildSystem/Configuration/BuildConfigProperties/BuildConfigProperties.INT.udn b/Engine/Documentation/Source/Programming/UnrealBuildSystem/Configuration/BuildConfigProperties/BuildConfigProperties.INT.udn
--- a/Engine/Documentation/Source/Programming/UnrealBuildSystem/Configuration/BuildConfigProperties/BuildConfigProperties.INT.udn
+++ b/Engine/Documentation/Source/Programming/UnrealBuildSystem/Configuration/BuildConfigProperties/BuildConfigProperties.INT.udn
@@ -30,7 +30,7 @@ $ bUseUBTMakefiles : Enables support for very fast iterative builds by caching t
 
 $ bAllowParallelExecutor : Whether to allow using parallel executor on Windows.
 
-$ MaxParallelActions : Number of actions that can be executed in parallel. If 0 then code will pick a default based on the number of cores available. Only applies to the ParallelExecutor
+$ MaxParallelActions : Number of actions that can be executed in parallel. If 0 then code will pick a default based on the number of cores and memory available. Applies to the ParallelExecutor, HybridExecutor, and LocalExecutor
 
 $ bEnableAddressSanitizer : Enables address sanitizer (ASan). Only supported for Visual Studio 2019 16.7.0 and up.
 
@@ -235,15 +235,19 @@ $ MaxLocalActions : Maximum number of actions to execute locally.
 
 ### LocalExecutor
 
-$ ProcessorCountMultiplier : Processor count multiplier for local execution. Can be below 1 to reserve CPU for other tasks. When using the local executor (not XGE), run a single action on each CPU core. Note that you can set this to a larger value to get slightly faster build times in many cases, but your computer's responsiveness during compiling may be much worse.
-
 $ MaxProcessorCount : Maximum processor count for local execution.
 
+$ ProcessorCountMultiplier : Processor count multiplier for local execution. Can be below 1 to reserve CPU for other tasks. When using the local executor (not XGE), run a single action on each CPU core. Note that you can set this to a larger value to get slightly faster build times in many cases, but your computer's responsiveness during compiling may be much worse. This value is ignored if the CPU does not support hyper-threading.
+
+$ MemoryPerActionBytes : Free memory per action in bytes, used to limit the number of parallel actions if the machine is memory starved. Set to 0 to disable free memory checking.
+
 ### ParallelExecutor
 
 $ MaxProcessorCount : Maximum processor count for local execution.
 
-$ ProcessorCountMultiplier : Processor count multiplier for local execution. Can be below 1 to reserve CPU for other tasks.
+$ ProcessorCountMultiplier : Processor count multiplier for local execution. Can be below 1 to reserve CPU for other tasks. When using the local executor (not XGE), run a single action on each CPU core. Note that you can set this to a larger value to get slightly faster build times in many cases, but your computer's responsiveness during compiling may be much worse. This value is ignored if the CPU does not support hyper-threading.
+
+$ MemoryPerActionBytes : Free memory per action in bytes, used to limit the number of parallel actions if the machine is memory starved. Set to 0 to disable free memory checking.
 
 $ bStopCompilationAfterErrors : When enabled, will stop compiling targets after a compile error occurs.
 
