Documentation for build tools

diff --git a/Engine/Documentation/Source/Programming/UnrealBuildSystem/Configuration/BuildConfigProperties/BuildConfigProperties.INT.udn b/Engine/Documentation/Source/Programming/UnrealBuildSystem/Configuration/BuildConfigProperties/BuildConfigProperties.INT.udn
--- a/Engine/Documentation/Source/Programming/UnrealBuildSystem/Configuration/BuildConfigProperties/BuildConfigProperties.INT.udn
+++ b/Engine/Documentation/Source/Programming/UnrealBuildSystem/Configuration/BuildConfigProperties/BuildConfigProperties.INT.udn
@@ -66,6 +66,8 @@ $ bAdaptiveUnityCreatesDedicatedPCH : Creates a dedicated PCH for each source fi
 
 $ bAdaptiveUnityEnablesEditAndContinue : Creates a dedicated PCH for each source file in the working set, allowing faster iteration on cpp-only changes.
 
+$ bAdaptiveUnityCompilesHeaderFiles : Creates a dedicated source file for each header file in the working set to detect missing includes in headers.
+
 $ MinGameModuleSourceFilesForUnityBuild : The number of source files in a game module before unity build will be activated for that module.  This allows small game modules to have faster iterative compile times for single files, at the expense of slower full rebuild times.  This setting can be overridden by the bFasterWithoutUnity option in a module's Build.cs file.
 
 $ DefaultWarningLevel : Default treatment of uncategorized warnings
@@ -245,14 +247,14 @@ $ MemoryPerActionBytes : Free memory per action in bytes, used to limit the numb
 
 ### ParallelExecutor
 
+$ bStopCompilationAfterErrors : When enabled, will stop compiling targets after a compile error occurs.
+
 $ MaxProcessorCount : Maximum processor count for local execution.
 
 $ ProcessorCountMultiplier : Processor count multiplier for local execution. Can be below 1 to reserve CPU for other tasks. When using the local executor (not XGE), run a single action on each CPU core. Note that you can set this to a larger value to get slightly faster build times in many cases, but your computer's responsiveness during compiling may be much worse. This value is ignored if the CPU does not support hyper-threading.
 
 $ MemoryPerActionBytes : Free memory per action in bytes, used to limit the number of parallel actions if the machine is memory starved. Set to 0 to disable free memory checking.
 
-$ bStopCompilationAfterErrors : When enabled, will stop compiling targets after a compile error occurs.
-
 ### XGE
 
 $ bAllowOverVpn : When set to false, XGE will not be enabled when running connected to the coordinator over VPN. Configure VPN-assigned subnets via the VpnSubnets parameter.
diff --git a/Engine/Documentation/Source/Programming/UnrealBuildSystem/TargetFiles/TargetFilesProperties/TargetFilesProperties.INT.udn b/Engine/Documentation/Source/Programming/UnrealBuildSystem/TargetFiles/TargetFilesProperties/TargetFilesProperties.INT.udn
--- a/Engine/Documentation/Source/Programming/UnrealBuildSystem/TargetFiles/TargetFilesProperties/TargetFilesProperties.INT.udn
+++ b/Engine/Documentation/Source/Programming/UnrealBuildSystem/TargetFiles/TargetFilesProperties/TargetFilesProperties.INT.udn
@@ -308,4 +308,6 @@ $ DisableOptimizeCodeForModules (String[]): List of modules to disable optimizat
 
 $ GeneratedProjectName (String): When generating project files, specifies the name of the project file to use when there are multiple targets of the same type.
 
+$ OptedInModulePlatforms (UnrealTargetPlatform[]): If this is non-null, then any platforms NOT listed will not be allowed to have modules in their directories be created
+
 
