Documentation for build tools

diff --git a/Engine/Documentation/Source/Programming/UnrealBuildSystem/TargetFiles/TargetFilesProperties/TargetFilesProperties.INT.udn b/Engine/Documentation/Source/Programming/UnrealBuildSystem/TargetFiles/TargetFilesProperties/TargetFilesProperties.INT.udn
--- a/Engine/Documentation/Source/Programming/UnrealBuildSystem/TargetFiles/TargetFilesProperties/TargetFilesProperties.INT.udn
+++ b/Engine/Documentation/Source/Programming/UnrealBuildSystem/TargetFiles/TargetFilesProperties/TargetFilesProperties.INT.udn
@@ -118,7 +118,7 @@ $ bCompileWithAccessibilitySupport (Boolean): Whether to allow accessibility cod
 
 $ bUseDebugLiveCodingConsole (Boolean): Whether to enable support for live coding
 
-$ bWithDirectXMath (Boolean): Whether to enable support for DirectX Math
+$ bWithDirectXMath (Boolean): Whether to enable support for DirectX Math LWC_TODO: No longer supported. Needs removing.
 
 $ bUseLoggingInShipping (Boolean): Whether to turn on logging for test/shipping builds.
 
