Documentation for build tools

diff --git a/Engine/Documentation/Source/Programming/UnrealBuildSystem/Configuration/BuildConfigProperties/BuildConfigProperties.INT.udn b/Engine/Documentation/Source/Programming/UnrealBuildSystem/Configuration/BuildConfigProperties/BuildConfigProperties.INT.udn
--- a/Engine/Documentation/Source/Programming/UnrealBuildSystem/Configuration/BuildConfigProperties/BuildConfigProperties.INT.udn
+++ b/Engine/Documentation/Source/Programming/UnrealBuildSystem/Configuration/BuildConfigProperties/BuildConfigProperties.INT.udn
@@ -293,6 +293,8 @@ $ bIncludeTemplateFiles : True if we should include template files in the genera
 
 $ IncludeEnginePrograms : True if we should include program projects in the generated solution.
 
+$ bIncludeDotNetPrograms : True if we should include csharp program projects in the generated solution. Pass "-DotNet" to enable this.
+
 $ bIncludeTempTargets : Whether to include temporary targets generated by UAT to support content only projects with non-default settings.
 
 $ bKeepSourceSubDirectories : True if we should reflect "Source" sub-directories on disk in the master project as master project directories. This (arguably) adds some visual clutter to the master project but it is truer to the on-disk file organization.
@@ -307,6 +309,10 @@ $ bMasterProjectNameFromFolder : If true, sets the master project name according
 
 $ bIncludeTestAndShippingConfigs : Whether we should include configurations for "Test" and "Shipping" in generated projects. Pass "-NoShippingConfigs" to disable this.
 
+$ bIncludeDebugConfigs : Whether we should include configurations for "Debug" and "DebugGame" in generated projects. Pass "-NoDebugConfigs" to disable this.
+
+$ bIncludeDevelopmentConfigs : Whether we should include configurations for "Development" in generated projects. Pass "-NoDevelopmentConfigs" to disable this.
+
 ### HoloLensPlatform
 
 $ Compiler : Version of the compiler toolchain to use on HoloLens. A value of "default" will be changed to a specific version at UBT startup.
@@ -339,6 +345,8 @@ $ bIncludeTemplateFiles : True if we should include template files in the genera
 
 $ IncludeEnginePrograms : True if we should include program projects in the generated solution.
 
+$ bIncludeDotNetPrograms : True if we should include csharp program projects in the generated solution. Pass "-DotNet" to enable this.
+
 $ MasterProjectName : Name of the master project file -- for example, the base file name for the Visual Studio solution file, or the Xcode project file on Mac.
 
 $ bMasterProjectNameFromFolder : If true, sets the master project name according to the name of the folder it is in.
@@ -359,6 +367,8 @@ $ bIncludeTemplateFiles : True if we should include template files in the genera
 
 $ IncludeEnginePrograms : True if we should include program projects in the generated solution.
 
+$ bIncludeDotNetPrograms : True if we should include csharp program projects in the generated solution. Pass "-DotNet" to enable this.
+
 $ MasterProjectName : Name of the master project file -- for example, the base file name for the Visual Studio solution file, or the Xcode project file on Mac.
 
 $ bMasterProjectNameFromFolder : If true, sets the master project name according to the name of the folder it is in.
@@ -379,6 +389,8 @@ $ bIncludeTemplateFiles : True if we should include template files in the genera
 
 $ IncludeEnginePrograms : True if we should include program projects in the generated solution.
 
+$ bIncludeDotNetPrograms : True if we should include csharp program projects in the generated solution. Pass "-DotNet" to enable this.
+
 $ MasterProjectName : Name of the master project file -- for example, the base file name for the Visual Studio solution file, or the Xcode project file on Mac.
 
 $ bMasterProjectNameFromFolder : If true, sets the master project name according to the name of the folder it is in.
@@ -399,6 +411,8 @@ $ bIncludeTemplateFiles : True if we should include template files in the genera
 
 $ IncludeEnginePrograms : True if we should include program projects in the generated solution.
 
+$ bIncludeDotNetPrograms : True if we should include csharp program projects in the generated solution. Pass "-DotNet" to enable this.
+
 $ MasterProjectName : Name of the master project file -- for example, the base file name for the Visual Studio solution file, or the Xcode project file on Mac.
 
 $ bMasterProjectNameFromFolder : If true, sets the master project name according to the name of the folder it is in.
@@ -419,6 +433,8 @@ $ bIncludeTemplateFiles : True if we should include template files in the genera
 
 $ IncludeEnginePrograms : True if we should include program projects in the generated solution.
 
+$ bIncludeDotNetPrograms : True if we should include csharp program projects in the generated solution. Pass "-DotNet" to enable this.
+
 $ MasterProjectName : Name of the master project file -- for example, the base file name for the Visual Studio solution file, or the Xcode project file on Mac.
 
 $ bMasterProjectNameFromFolder : If true, sets the master project name according to the name of the folder it is in.
@@ -439,6 +455,8 @@ $ bIncludeTemplateFiles : True if we should include template files in the genera
 
 $ IncludeEnginePrograms : True if we should include program projects in the generated solution.
 
+$ bIncludeDotNetPrograms : True if we should include csharp program projects in the generated solution. Pass "-DotNet" to enable this.
+
 $ MasterProjectName : Name of the master project file -- for example, the base file name for the Visual Studio solution file, or the Xcode project file on Mac.
 
 $ bMasterProjectNameFromFolder : If true, sets the master project name according to the name of the folder it is in.
@@ -459,6 +477,8 @@ $ bIncludeTemplateFiles : True if we should include template files in the genera
 
 $ IncludeEnginePrograms : True if we should include program projects in the generated solution.
 
+$ bIncludeDotNetPrograms : True if we should include csharp program projects in the generated solution. Pass "-DotNet" to enable this.
+
 $ MasterProjectName : Name of the master project file -- for example, the base file name for the Visual Studio solution file, or the Xcode project file on Mac.
 
 $ bMasterProjectNameFromFolder : If true, sets the master project name according to the name of the folder it is in.
@@ -479,6 +499,8 @@ $ bIncludeTemplateFiles : True if we should include template files in the genera
 
 $ IncludeEnginePrograms : True if we should include program projects in the generated solution.
 
+$ bIncludeDotNetPrograms : True if we should include csharp program projects in the generated solution. Pass "-DotNet" to enable this.
+
 $ MasterProjectName : Name of the master project file -- for example, the base file name for the Visual Studio solution file, or the Xcode project file on Mac.
 
 $ bMasterProjectNameFromFolder : If true, sets the master project name according to the name of the folder it is in.
@@ -501,6 +523,8 @@ $ bIncludeTemplateFiles : True if we should include template files in the genera
 
 $ IncludeEnginePrograms : True if we should include program projects in the generated solution.
 
+$ bIncludeDotNetPrograms : True if we should include csharp program projects in the generated solution. Pass "-DotNet" to enable this.
+
 $ MasterProjectName : Name of the master project file -- for example, the base file name for the Visual Studio solution file, or the Xcode project file on Mac.
 
 $ bMasterProjectNameFromFolder : If true, sets the master project name according to the name of the folder it is in.
@@ -521,6 +545,8 @@ $ bIncludeTemplateFiles : True if we should include template files in the genera
 
 $ IncludeEnginePrograms : True if we should include program projects in the generated solution.
 
+$ bIncludeDotNetPrograms : True if we should include csharp program projects in the generated solution. Pass "-DotNet" to enable this.
+
 $ MasterProjectName : Name of the master project file -- for example, the base file name for the Visual Studio solution file, or the Xcode project file on Mac.
 
 $ bMasterProjectNameFromFolder : If true, sets the master project name according to the name of the folder it is in.
@@ -551,6 +577,8 @@ $ bIncludeTemplateFiles : True if we should include template files in the genera
 
 $ IncludeEnginePrograms : True if we should include program projects in the generated solution.
 
+$ bIncludeDotNetPrograms : True if we should include csharp program projects in the generated solution. Pass "-DotNet" to enable this.
+
 $ MasterProjectName : Name of the master project file -- for example, the base file name for the Visual Studio solution file, or the Xcode project file on Mac.
 
 $ bMasterProjectNameFromFolder : If true, sets the master project name according to the name of the folder it is in.
@@ -571,6 +599,8 @@ $ bIncludeTemplateFiles : True if we should include template files in the genera
 
 $ IncludeEnginePrograms : True if we should include program projects in the generated solution.
 
+$ bIncludeDotNetPrograms : True if we should include csharp program projects in the generated solution. Pass "-DotNet" to enable this.
+
 $ MasterProjectName : Name of the master project file -- for example, the base file name for the Visual Studio solution file, or the Xcode project file on Mac.
 
 $ bMasterProjectNameFromFolder : If true, sets the master project name according to the name of the folder it is in.
