Documentation for build tools

diff --git a/Engine/Documentation/Source/Programming/UnrealBuildSystem/Configuration/BuildConfigProperties/BuildConfigProperties.INT.udn b/Engine/Documentation/Source/Programming/UnrealBuildSystem/Configuration/BuildConfigProperties/BuildConfigProperties.INT.udn
--- a/Engine/Documentation/Source/Programming/UnrealBuildSystem/Configuration/BuildConfigProperties/BuildConfigProperties.INT.udn
+++ b/Engine/Documentation/Source/Programming/UnrealBuildSystem/Configuration/BuildConfigProperties/BuildConfigProperties.INT.udn
@@ -509,6 +509,10 @@ $ bMasterProjectNameFromFolder : If true, sets the master project name according
 
 $ IncludeAllFiles : Includes all files in the generated workspace.
 
+$ AddDebugAttachConfig : Whether VS Code project generation should include debug configurations to allow attaching to already running processes
+
+$ AddDebugCoreConfig : Whether VS Code project generation should include debug configurations to allow core dump debugging
+
 $ bIncludeDocumentation : True if we should include documentation in the generated projects.
 
 $ bUsePrecompiled : True if build targets should pass the -useprecompiled argument.
