Documentation for build tools

diff --git a/Engine/Documentation/Source/Programming/UnrealBuildSystem/Configuration/BuildConfigProperties/BuildConfigProperties.INT.udn b/Engine/Documentation/Source/Programming/UnrealBuildSystem/Configuration/BuildConfigProperties/BuildConfigProperties.INT.udn
--- a/Engine/Documentation/Source/Programming/UnrealBuildSystem/Configuration/BuildConfigProperties/BuildConfigProperties.INT.udn
+++ b/Engine/Documentation/Source/Programming/UnrealBuildSystem/Configuration/BuildConfigProperties/BuildConfigProperties.INT.udn
@@ -203,6 +203,10 @@ $ bCompilerTrace : Outputs compile timing information so that it can be analyzed
 
 $ bShowIncludes : Print out files that are included by each source file
 
+### PS4Platform
+
+$ bUseLiveCodeEdit : Enable support for edit and continue code modification on ps4
+
 ### ModuleConfiguration
 
 $ DisableUnityBuild : List of modules to disable unity builds for
@@ -629,3 +633,27 @@ $ RsyncAuthentication : The authentication used for Rsync (for the -e rsync flag
 
 $ SshAuthentication : The authentication used for SSH (probably similar to RsyncAuthentication).
 
+### PS4ProjectGenerator
+
+$ bUseApplicationDebugger : No documentation available.
+
+### PS5ProjectGenerator
+
+$ bUseApplicationDebugger : When enabled, generates project files that default to using the PS5 application debugger.
+
+$ bUseGP5File : When enabled, generates project files where the PS5 local debugger is set to use a .gp5 file from the StagedBuild directory. Otherwise, the StagedBuild directory itself is set as the working directory, and the title is launched without the help of a .gp5 file.
+
+### XboxCommonPlatform
+
+$ bPackageDebugAssistance : Assistance for debugging an installed xvc package: Disables the F5 debugger and copies the executable to an overlay folder instead. Can only be done on the default devkit
+
+$ PackageDebugOverrideName : Custom app package name for package debugging assistance. e.g. GameName_1.0.0.0_neutral__z844xmwvfmpty
+
+### XboxOneGDKProjectGenerator
+
+$ bRemoveExtraDeployFiles : 
+
+### XSXProjectGenerator
+
+$ bRemoveExtraDeployFiles : 
+
