Documentation for build tools

diff --git a/Engine/Documentation/Source/Programming/UnrealBuildSystem/Configuration/BuildConfigProperties/BuildConfigProperties.INT.udn b/Engine/Documentation/Source/Programming/UnrealBuildSystem/Configuration/BuildConfigProperties/BuildConfigProperties.INT.udn
--- a/Engine/Documentation/Source/Programming/UnrealBuildSystem/Configuration/BuildConfigProperties/BuildConfigProperties.INT.udn
+++ b/Engine/Documentation/Source/Programming/UnrealBuildSystem/Configuration/BuildConfigProperties/BuildConfigProperties.INT.udn
@@ -633,6 +633,14 @@ $ RsyncAuthentication : The authentication used for Rsync (for the -e rsync flag
 
 $ SshAuthentication : The authentication used for SSH (probably similar to RsyncAuthentication).
 
+### GDKPlatform
+
+$ bVerifyLibGDKEditions : Whether to verify that any libs the application links to are using the same GDK edition
+
+$ bVerifyDLLGDKEditions : Whether to verify that any DLLs the application links to are using the same GDK edition
+
+$ bBuildForIteration : Whether to automatically enable incremental linking etc. when building non-shipping, non-formal builds. Can speed up building but can cause instability, and on very large projects PDB generation may fail.
+
 ### PS4ProjectGenerator
 
 $ bUseApplicationDebugger : No documentation available.
@@ -643,8 +651,16 @@ $ bUseApplicationDebugger : When enabled, generates project files that default t
 
 $ bUseGP5File : When enabled, generates project files where the PS5 local debugger is set to use a .gp5 file from the StagedBuild directory. Otherwise, the StagedBuild directory itself is set as the working directory, and the title is launched without the help of a .gp5 file.
 
+### PS5Platform
+
+$ bUseLiveCodeEdit : Enable support for edit and continue code modification
+
 ### XboxCommonPlatform
 
+$ Compiler : Override the default compiler
+
+$ CompilerVersion : Override the default compiler version
+
 $ bPackageDebugAssistance : Assistance for debugging an installed xvc package: Disables the F5 debugger and copies the executable to an overlay folder instead. Can only be done on the default devkit
 
 $ PackageDebugOverrideName : Custom app package name for package debugging assistance. e.g. GameName_1.0.0.0_neutral__z844xmwvfmpty
