Documentation for build tools

diff --git a/Engine/Documentation/Source/Programming/UnrealBuildSystem/TargetFiles/TargetFilesProperties/TargetFilesProperties.INT.udn b/Engine/Documentation/Source/Programming/UnrealBuildSystem/TargetFiles/TargetFilesProperties/TargetFilesProperties.INT.udn
--- a/Engine/Documentation/Source/Programming/UnrealBuildSystem/TargetFiles/TargetFilesProperties/TargetFilesProperties.INT.udn
+++ b/Engine/Documentation/Source/Programming/UnrealBuildSystem/TargetFiles/TargetFilesProperties/TargetFilesProperties.INT.udn
@@ -60,6 +60,8 @@ $ GeneratedCodeVersion (EGeneratedCodeVersion): Allow target module to override
 
 $ bEnableMeshEditor (Boolean): Whether to enable the mesh editor.
 
+$ bUseVerse (Boolean): Whether to use the verse script interface.
+
 $ bCompileChaos (Boolean): Whether to compile the Chaos physics plugin.
 
 $ bUseChaos (Boolean): Whether to use the Chaos physics interface. This overrides the physx flags to disable APEX and NvCloth
@@ -82,6 +84,8 @@ $ bCompileCEF3 (Boolean): Whether to compile CEF3 support.
 
 $ bCompileISPC (Boolean): Whether to compile using ISPC.
 
+$ bCompilePython (Boolean): Whether to compile in python support
+
 $ bForceBuildTargetPlatforms (Boolean): Whether to force compiling the target platform modules, even if they wouldn't normally be built.
 
 $ bForceBuildShaderFormats (Boolean): Whether to force compiling shader format modules, even if they wouldn't normally be built.
