Documentation for build tools

diff --git a/Engine/Documentation/Source/Programming/UnrealBuildSystem/Configuration/BuildConfigProperties/BuildConfigProperties.INT.udn b/Engine/Documentation/Source/Programming/UnrealBuildSystem/Configuration/BuildConfigProperties/BuildConfigProperties.INT.udn
--- a/Engine/Documentation/Source/Programming/UnrealBuildSystem/Configuration/BuildConfigProperties/BuildConfigProperties.INT.udn
+++ b/Engine/Documentation/Source/Programming/UnrealBuildSystem/Configuration/BuildConfigProperties/BuildConfigProperties.INT.udn
@@ -28,8 +28,6 @@ $ bUseUBTMakefiles : Enables support for very fast iterative builds by caching t
     * You can force regeneration of the 'UBT Makefile' by passing the '-gather' argument, or simply regenerating project files. 
     * This also enables the fast include file dependency scanning and caching system that allows Unreal Build Tool to detect out of date dependencies very quickly. When enabled, a deep C++ include graph does not have to be generated, and instead, we only scan and cache indirect includes for after a dependent build product was already found to be out of date. During the next build, we will load those cached indirect includes and check for outdatedness.
 
-$ bAllowParallelExecutor : Whether to allow using parallel executor on Windows.
-
 $ MaxParallelActions : Number of actions that can be executed in parallel. If 0 then code will pick a default based on the number of cores and memory available. Applies to the ParallelExecutor, HybridExecutor, and LocalExecutor
 
 $ bEnableAddressSanitizer : Enables address sanitizer (ASan). Only supported for Visual Studio 2019 16.7.0 and up.
@@ -241,14 +239,6 @@ $ MsvcCRTRedistVersion : Which MSVC CRT Redist version to use
 
 $ MaxLocalActions : Maximum number of actions to execute locally.
 
-### LocalExecutor
-
-$ MaxProcessorCount : Maximum processor count for local execution.
-
-$ ProcessorCountMultiplier : Processor count multiplier for local execution. Can be below 1 to reserve CPU for other tasks. When using the local executor (not XGE), run a single action on each CPU core. Note that you can set this to a larger value to get slightly faster build times in many cases, but your computer's responsiveness during compiling may be much worse. This value is ignored if the CPU does not support hyper-threading.
-
-$ MemoryPerActionBytes : Free memory per action in bytes, used to limit the number of parallel actions if the machine is memory starved. Set to 0 to disable free memory checking.
-
 ### ParallelExecutor
 
 $ bStopCompilationAfterErrors : When enabled, will stop compiling targets after a compile error occurs.
diff --git a/Engine/Documentation/Source/Programming/UnrealBuildSystem/TargetFiles/TargetFilesProperties/TargetFilesProperties.INT.udn b/Engine/Documentation/Source/Programming/UnrealBuildSystem/TargetFiles/TargetFilesProperties/TargetFilesProperties.INT.udn
--- a/Engine/Documentation/Source/Programming/UnrealBuildSystem/TargetFiles/TargetFilesProperties/TargetFilesProperties.INT.udn
+++ b/Engine/Documentation/Source/Programming/UnrealBuildSystem/TargetFiles/TargetFilesProperties/TargetFilesProperties.INT.udn
@@ -144,6 +144,8 @@ $ bForceCompilePerformanceAutomationTests (Boolean): Whether to compile performa
 
 $ bEventDrivenLoader (Boolean): If true, event driven loader will be used in cooked builds. @todoio This needs to be replaced by a runtime solution after async loading refactor.
 
+$ NativePointerMemberBehaviorOverride (Nullable&lt;PointerMemberBehavior&gt;): Used to override the behavior controlling whether UCLASSes and USTRUCTs are allowed to have native pointer members, if disallowed they will be a UHT error and should be substituted with TObjectPtr members instead.
+
 $ bUseXGEController (Boolean): Whether the XGE controller worker and modules should be included in the engine build. These are required for distributed shader compilation using the XGE interception interface.
 
 $ bIWYU (Boolean): Enables "include what you use" by default for modules in this target. Changes the default PCH mode for any module in this project to PCHUsageMode.UseExplicitOrSharedPCHs.
