Documentation for build tools

diff --git a/Engine/Documentation/Source/Programming/UnrealBuildSystem/Configuration/BuildConfigProperties/BuildConfigProperties.INT.udn b/Engine/Documentation/Source/Programming/UnrealBuildSystem/Configuration/BuildConfigProperties/BuildConfigProperties.INT.udn
--- a/Engine/Documentation/Source/Programming/UnrealBuildSystem/Configuration/BuildConfigProperties/BuildConfigProperties.INT.udn
+++ b/Engine/Documentation/Source/Programming/UnrealBuildSystem/Configuration/BuildConfigProperties/BuildConfigProperties.INT.udn
@@ -651,6 +651,8 @@ $ Compiler : Override the default compiler
 
 $ CompilerVersion : Override the default compiler version
 
+$ bNoXboxCommonGDKWarning : Whether to log out a warning when linking against modules in an XboxCommonGDK path
+
 $ bPackageDebugAssistance : Assistance for debugging an installed xvc package: Disables the F5 debugger and copies the executable to an overlay folder instead. Can only be done on the default devkit
 
 $ PackageDebugOverrideName : Custom app package name for package debugging assistance. e.g. GameName_1.0.0.0_neutral__z844xmwvfmpty
