Documentation for build tools

diff --git a/Engine/Documentation/Source/Programming/UnrealBuildSystem/Configuration/BuildConfigProperties/BuildConfigProperties.INT.udn b/Engine/Documentation/Source/Programming/UnrealBuildSystem/Configuration/BuildConfigProperties/BuildConfigProperties.INT.udn
--- a/Engine/Documentation/Source/Programming/UnrealBuildSystem/Configuration/BuildConfigProperties/BuildConfigProperties.INT.udn
+++ b/Engine/Documentation/Source/Programming/UnrealBuildSystem/Configuration/BuildConfigProperties/BuildConfigProperties.INT.udn
@@ -18,6 +18,8 @@ $ bAllowFASTBuild : Whether FASTBuild may be used.
 
 $ bAllowSNDBS : Whether SN-DBS may be used.
 
+$ bAllowTaskExecutor : Whether the experimental async TaskExecutor may be used.
+
 $ bUseUBTMakefiles : Enables support for very fast iterative builds by caching target data. Turning this on causes Unreal Build Tool to emit 'UBT Makefiles' for targets when they are built the first time. Subsequent builds will load these Makefiles and begin outdatedness checking and build invocation very quickly. The caveat is that if source files are added or removed to the project, UBT will need to gather information about those in order for your build to complete successfully. Currently, you must run the project file generator after adding/removing source files to tell UBT to re-gather this information. 
     * Events that can invalidate the 'UBT Makefile':
         * Adding/removing .cpp files
@@ -235,6 +237,16 @@ $ bStopOnError : Whether to stop on error
 
 $ MsvcCRTRedistVersion : Which MSVC CRT Redist version to use
 
+### TaskExecutor
+
+$ bStopCompilationAfterErrors : When enabled, will stop compiling targets after a compile error occurs.
+
+$ MaxProcessorCount : Maximum processor count for local execution.
+
+$ ProcessorCountMultiplier : Processor count multiplier for local execution. Can be below 1 to reserve CPU for other tasks. When using the local executor (not XGE), run a single action on each CPU core. Note that you can set this to a larger value to get slightly faster build times in many cases, but your computer's responsiveness during compiling may be much worse. This value is ignored if the CPU does not support hyper-threading.
+
+$ MemoryPerActionBytes : Free memory per action in bytes, used to limit the number of parallel actions if the machine is memory starved. Set to 0 to disable free memory checking.
+
 ### HybridExecutor
 
 $ MaxLocalActions : Maximum number of actions to execute locally.
