Documentation for build tools

diff --git a/Engine/Documentation/Source/Programming/UnrealBuildSystem/Configuration/BuildConfigProperties/BuildConfigProperties.INT.udn b/Engine/Documentation/Source/Programming/UnrealBuildSystem/Configuration/BuildConfigProperties/BuildConfigProperties.INT.udn
--- a/Engine/Documentation/Source/Programming/UnrealBuildSystem/Configuration/BuildConfigProperties/BuildConfigProperties.INT.udn
+++ b/Engine/Documentation/Source/Programming/UnrealBuildSystem/Configuration/BuildConfigProperties/BuildConfigProperties.INT.udn
@@ -197,6 +197,8 @@ $ bStrictConformanceMode : Enables strict standard conformance mode (/permissive
 
 $ PCHMemoryAllocationFactor : Determines the amount of memory that the compiler allocates to construct precompiled headers (/Zm).
 
+$ PDBPageSize : Override the PDB page size. Used when the PDB file grows over 4GB. Expected values are 4096 (default), 8192, 16384, 32768
+
 $ bClangTimeTrace : (Experimental) Appends the -ftime-trace argument to the command line for Clang to output a JSON file containing a timeline for the compile. See http://aras-p.info/blog/2019/01/16/time-trace-timeline-flame-chart-profiler-for-Clang/ for more info.
 
 $ bCompilerTrace : Outputs compile timing information so that it can be analyzed.
