Documentation for build tools

diff --git a/Engine/Documentation/Source/Programming/UnrealBuildSystem/TargetFiles/TargetFilesProperties/TargetFilesProperties.INT.udn b/Engine/Documentation/Source/Programming/UnrealBuildSystem/TargetFiles/TargetFilesProperties/TargetFilesProperties.INT.udn
--- a/Engine/Documentation/Source/Programming/UnrealBuildSystem/TargetFiles/TargetFilesProperties/TargetFilesProperties.INT.udn
+++ b/Engine/Documentation/Source/Programming/UnrealBuildSystem/TargetFiles/TargetFilesProperties/TargetFilesProperties.INT.udn
@@ -120,6 +120,8 @@ $ bCompileWithPluginSupport (Boolean): Whether to include plugin support.
 
 $ bCompileWithAccessibilitySupport (Boolean): Whether to allow accessibility code in both Slate and the OS layer.
 
+$ bCompileWithSlateWidgetTracking (Boolean): Whether to allow the Slate Widget Tracker to track registered Widgets.
+
 $ bUseDebugLiveCodingConsole (Boolean): Whether to enable support for live coding
 
 $ bWithDirectXMath (Boolean): Whether to enable support for DirectX Math LWC_TODO: No longer supported. Needs removing.
