Documentation for build tools

diff --git a/Engine/Documentation/Source/Programming/UnrealBuildSystem/Configuration/BuildConfigProperties/BuildConfigProperties.INT.udn b/Engine/Documentation/Source/Programming/UnrealBuildSystem/Configuration/BuildConfigProperties/BuildConfigProperties.INT.udn
--- a/Engine/Documentation/Source/Programming/UnrealBuildSystem/Configuration/BuildConfigProperties/BuildConfigProperties.INT.udn
+++ b/Engine/Documentation/Source/Programming/UnrealBuildSystem/Configuration/BuildConfigProperties/BuildConfigProperties.INT.udn
@@ -239,6 +239,8 @@ $ MsvcCRTRedistVersion : Which MSVC CRT Redist version to use
 
 $ bStopCompilationAfterErrors : When enabled, will stop compiling targets after a compile error occurs.
 
+$ bShowCompilationTimes : Whether to show compilation times along with worst offenders or not.
+
 $ MaxProcessorCount : Maximum processor count for local execution.
 
 $ ProcessorCountMultiplier : Processor count multiplier for local execution. Can be below 1 to reserve CPU for other tasks. When using the local executor (not XGE), run a single action on each CPU core. Note that you can set this to a larger value to get slightly faster build times in many cases, but your computer's responsiveness during compiling may be much worse. This value is ignored if the CPU does not support hyper-threading.
@@ -253,6 +255,8 @@ $ MaxLocalActions : Maximum number of actions to execute locally.
 
 $ bStopCompilationAfterErrors : When enabled, will stop compiling targets after a compile error occurs.
 
+$ bShowCompilationTimes : Whether to show compilation times along with worst offenders or not.
+
 $ MaxProcessorCount : Maximum processor count for local execution.
 
 $ ProcessorCountMultiplier : Processor count multiplier for local execution. Can be below 1 to reserve CPU for other tasks. When using the local executor (not XGE), run a single action on each CPU core. Note that you can set this to a larger value to get slightly faster build times in many cases, but your computer's responsiveness during compiling may be much worse. This value is ignored if the CPU does not support hyper-threading.
