Documentation for build tools

diff --git a/Engine/Documentation/Source/Programming/UnrealBuildSystem/ModuleFiles/ModuleFilesProperties/ModuleFilesProperties.INT.udn b/Engine/Documentation/Source/Programming/UnrealBuildSystem/ModuleFiles/ModuleFilesProperties/ModuleFilesProperties.INT.udn
--- a/Engine/Documentation/Source/Programming/UnrealBuildSystem/ModuleFiles/ModuleFilesProperties/ModuleFilesProperties.INT.udn
+++ b/Engine/Documentation/Source/Programming/UnrealBuildSystem/ModuleFiles/ModuleFilesProperties/ModuleFilesProperties.INT.udn
@@ -69,6 +69,8 @@ $ PublicSystemIncludePaths (List&lt;String&gt;): List of system/library include
 
 $ PublicIncludePaths (List&lt;String&gt;): (This setting is currently not need as we discover all files from the 'Public' folder) List of all paths to include files that are exposed to other modules
 
+$ InternalncludePaths (List&lt;String&gt;): (This setting is currently not need as we discover all files from the 'Internal' folder) List of all paths to include files that are exposed to other internal modules
+
 $ PrivateIncludePaths (List&lt;String&gt;): List of all paths to this module's internal include files, not exposed to other modules (at least one include to the 'Private' path, more if we want to avoid relative paths)
 
 $ PublicSystemLibraryPaths (List&lt;String&gt;): List of system library paths (directory of .lib files) - for External (third party) modules please use the PublicAdditionalLibaries instead
