Documentation for build tools

diff --git a/Engine/Documentation/Source/Programming/UnrealBuildSystem/TargetFiles/TargetFilesProperties/TargetFilesProperties.INT.udn b/Engine/Documentation/Source/Programming/UnrealBuildSystem/TargetFiles/TargetFilesProperties/TargetFilesProperties.INT.udn
--- a/Engine/Documentation/Source/Programming/UnrealBuildSystem/TargetFiles/TargetFilesProperties/TargetFilesProperties.INT.udn
+++ b/Engine/Documentation/Source/Programming/UnrealBuildSystem/TargetFiles/TargetFilesProperties/TargetFilesProperties.INT.udn
@@ -54,6 +54,8 @@ $ bBuildInSolutionByDefault (Nullable&lt;Boolean&gt;): Whether the target should
 
 $ bShouldCompileAsDLL (Boolean): Whether this target should be compiled as a DLL.  Requires LinkType to be set to TargetLinkType.Monolithic.
 
+$ CustomConfig (String): Extra subdirectory to load config files out of, for making multiple types of builds with the same platform This will get baked into the game executable as CUSTOM_CONFIG and used when staging to filter files and settings
+
 $ ExeBinariesSubFolder (String): Subfolder to place executables in, relative to the default location.
 
 $ GeneratedCodeVersion (EGeneratedCodeVersion): Allow target module to override UHT code generation version.
