Documentation for build tools

diff --git a/Engine/Documentation/Source/Programming/UnrealBuildSystem/TargetFiles/TargetFilesProperties/TargetFilesProperties.INT.udn b/Engine/Documentation/Source/Programming/UnrealBuildSystem/TargetFiles/TargetFilesProperties/TargetFilesProperties.INT.udn
--- a/Engine/Documentation/Source/Programming/UnrealBuildSystem/TargetFiles/TargetFilesProperties/TargetFilesProperties.INT.udn
+++ b/Engine/Documentation/Source/Programming/UnrealBuildSystem/TargetFiles/TargetFilesProperties/TargetFilesProperties.INT.udn
@@ -122,8 +122,6 @@ $ bCompileWithPluginSupport (Boolean): Whether to include plugin support.
 
 $ bCompileWithAccessibilitySupport (Boolean): Whether to allow accessibility code in both Slate and the OS layer.
 
-$ bCompileWithSlateWidgetTracking (Boolean): Whether to allow the Slate Widget Tracker to track registered Widgets.
-
 $ bUseDebugLiveCodingConsole (Boolean): Whether to enable support for live coding
 
 $ bWithDirectXMath (Boolean): Whether to enable support for DirectX Math LWC_TODO: No longer supported. Needs removing.
@@ -146,6 +144,8 @@ $ bForceCompileDevelopmentAutomationTests (Boolean): Whether to compile developm
 
 $ bForceCompilePerformanceAutomationTests (Boolean): Whether to compile performance automation tests.
 
+$ bForceDisableAutomationTests (Boolean): Whether to override the defaults for automation tests (Debug/Development configs)
+
 $ bEventDrivenLoader (Boolean): If true, event driven loader will be used in cooked builds. @todoio This needs to be replaced by a runtime solution after async loading refactor.
 
 $ NativePointerMemberBehaviorOverride (Nullable&lt;PointerMemberBehavior&gt;): Used to override the behavior controlling whether UCLASSes and USTRUCTs are allowed to have native pointer members, if disallowed they will be a UHT error and should be substituted with TObjectPtr members instead.
