Documentation for build tools

diff --git a/Engine/Documentation/Source/Programming/UnrealBuildSystem/Configuration/BuildConfigProperties/BuildConfigProperties.INT.udn b/Engine/Documentation/Source/Programming/UnrealBuildSystem/Configuration/BuildConfigProperties/BuildConfigProperties.INT.udn
--- a/Engine/Documentation/Source/Programming/UnrealBuildSystem/Configuration/BuildConfigProperties/BuildConfigProperties.INT.udn
+++ b/Engine/Documentation/Source/Programming/UnrealBuildSystem/Configuration/BuildConfigProperties/BuildConfigProperties.INT.udn
@@ -70,10 +70,6 @@ $ bAdaptiveUnityCompilesHeaderFiles : Creates a dedicated source file for each h
 
 $ MinGameModuleSourceFilesForUnityBuild : The number of source files in a game module before unity build will be activated for that module.  This allows small game modules to have faster iterative compile times for single files, at the expense of slower full rebuild times.  This setting can be overridden by the bFasterWithoutUnity option in a module's Build.cs file.
 
-$ DefaultWarningLevel : Default treatment of uncategorized warnings
-
-$ DeprecationWarningLevel : Whether to treat all warnings as errors. UE generally treats most warnings as errors, with the exception of deprecation warnings,
-
 $ bWarningsAsErrors : Whether to enable all warnings as errors. UE enables most warnings as errors already, but disables a few (such as deprecation warnings).
 
 $ UnsafeTypeCastWarningLevel : Indicates what warning/error level to treat unsafe type casts as on platforms that support it (e.g., double->float or int64->int32)
diff --git a/Engine/Documentation/Source/Programming/UnrealBuildSystem/TargetFiles/TargetFilesProperties/TargetFilesProperties.INT.udn b/Engine/Documentation/Source/Programming/UnrealBuildSystem/TargetFiles/TargetFilesProperties/TargetFilesProperties.INT.udn
--- a/Engine/Documentation/Source/Programming/UnrealBuildSystem/TargetFiles/TargetFilesProperties/TargetFilesProperties.INT.udn
+++ b/Engine/Documentation/Source/Programming/UnrealBuildSystem/TargetFiles/TargetFilesProperties/TargetFilesProperties.INT.udn
@@ -180,10 +180,6 @@ $ bAdaptiveUnityCompilesHeaderFiles (Boolean): Creates a dedicated source file f
 
 $ MinGameModuleSourceFilesForUnityBuild (Int32): The number of source files in a game module before unity build will be activated for that module.  This allows small game modules to have faster iterative compile times for single files, at the expense of slower full rebuild times.  This setting can be overridden by the bFasterWithoutUnity option in a module's Build.cs file.
 
-$ DefaultWarningLevel (WarningLevel): Default treatment of uncategorized warnings
-
-$ DeprecationWarningLevel (WarningLevel): Whether to treat all warnings as errors. UE generally treats most warnings as errors, with the exception of deprecation warnings,
-
 $ ShadowVariableWarningLevel (WarningLevel): Forces shadow variable warnings to be treated as errors on platforms that support it.
 
 $ bWarningsAsErrors (Boolean): Whether to enable all warnings as errors. UE enables most warnings as errors already, but disables a few (such as deprecation warnings).
