Documentation for build tools

diff --git a/Engine/Documentation/Source/Programming/UnrealBuildSystem/Configuration/BuildConfigProperties/BuildConfigProperties.INT.udn b/Engine/Documentation/Source/Programming/UnrealBuildSystem/Configuration/BuildConfigProperties/BuildConfigProperties.INT.udn
--- a/Engine/Documentation/Source/Programming/UnrealBuildSystem/Configuration/BuildConfigProperties/BuildConfigProperties.INT.udn
+++ b/Engine/Documentation/Source/Programming/UnrealBuildSystem/Configuration/BuildConfigProperties/BuildConfigProperties.INT.udn
@@ -189,6 +189,8 @@ $ WindowsSdkVersion : The specific Windows SDK version to use. This may be a spe
 
 $ StaticAnalyzer : The static analyzer to use.
 
+$ StaticAnalyzerOutputType : The output type to use for the static analyzer.
+
 $ bStrictConformanceMode : Enables strict standard conformance mode (/permissive-) in VS2017+.
 
 $ PCHMemoryAllocationFactor : Determines the amount of memory that the compiler allocates to construct precompiled headers (/Zm).
