Documentation for build tools

diff --git a/Engine/Documentation/Source/Programming/UnrealBuildSystem/TargetFiles/TargetFilesProperties/TargetFilesProperties.INT.udn b/Engine/Documentation/Source/Programming/UnrealBuildSystem/TargetFiles/TargetFilesProperties/TargetFilesProperties.INT.udn
--- a/Engine/Documentation/Source/Programming/UnrealBuildSystem/TargetFiles/TargetFilesProperties/TargetFilesProperties.INT.udn
+++ b/Engine/Documentation/Source/Programming/UnrealBuildSystem/TargetFiles/TargetFilesProperties/TargetFilesProperties.INT.udn
@@ -40,6 +40,8 @@ $ bBuildAllModules (Boolean): Build all the modules that are valid for this targ
 
 $ BuildPlugins (List&lt;String&gt;): Additional plugins that are built for this target type but not enabled.
 
+$ bRuntimeDependenciesComeFromBuildPlugins (Boolean): If this is true, then the BuildPlugins list will be used to populate RuntimeDependencies, rather than EnablePlugins
+
 $ AdditionalPlugins (List&lt;String&gt;): A list of additional plugins which need to be included in this target. This allows referencing non-optional plugin modules which cannot be disabled, and allows building against specific modules in program targets which do not fit the categories in ModuleHostType.
 
 $ EnablePlugins (List&lt;String&gt;): Additional plugins that should be included for this target.
