Documentation for build tools

diff --git a/Engine/Documentation/Source/Programming/UnrealBuildSystem/Configuration/BuildConfigProperties/BuildConfigProperties.INT.udn b/Engine/Documentation/Source/Programming/UnrealBuildSystem/Configuration/BuildConfigProperties/BuildConfigProperties.INT.udn
--- a/Engine/Documentation/Source/Programming/UnrealBuildSystem/Configuration/BuildConfigProperties/BuildConfigProperties.INT.udn
+++ b/Engine/Documentation/Source/Programming/UnrealBuildSystem/Configuration/BuildConfigProperties/BuildConfigProperties.INT.udn
@@ -203,6 +203,8 @@ $ bCompilerTrace : Outputs compile timing information so that it can be analyzed
 
 $ bShowIncludes : Print out files that are included by each source file
 
+$ bDeterministic : Set flags require for determinstic compiles (experimental)
+
 ### ModuleConfiguration
 
 $ DisableUnityBuild : List of modules to disable unity builds for
