Documentation for build tools

diff --git a/Engine/Documentation/Source/Programming/UnrealBuildSystem/TargetFiles/TargetFilesProperties/TargetFilesProperties.INT.udn b/Engine/Documentation/Source/Programming/UnrealBuildSystem/TargetFiles/TargetFilesProperties/TargetFilesProperties.INT.udn
--- a/Engine/Documentation/Source/Programming/UnrealBuildSystem/TargetFiles/TargetFilesProperties/TargetFilesProperties.INT.udn
+++ b/Engine/Documentation/Source/Programming/UnrealBuildSystem/TargetFiles/TargetFilesProperties/TargetFilesProperties.INT.udn
@@ -308,6 +308,8 @@ $ AdditionalCompilerArguments (String): Additional arguments to pass to the comp
 
 $ AdditionalLinkerArguments (String): Additional arguments to pass to the linker
 
+$ MemoryPerActionGB (Double): Max amount of memory that each compile action may require. Used by ParallelExecutor and TaskExecutor to decide the maximum number of parallel actions to start at one time.
+
 $ DisableUnityBuildForModules (String[]): List of modules to disable unity builds for
 
 $ EnableOptimizeCodeForModules (String[]): List of modules to enable optimizations for
