Documentation for build tools

diff --git a/Engine/Documentation/Source/Programming/UnrealBuildSystem/Configuration/BuildConfigProperties/BuildConfigProperties.INT.udn b/Engine/Documentation/Source/Programming/UnrealBuildSystem/Configuration/BuildConfigProperties/BuildConfigProperties.INT.udn
--- a/Engine/Documentation/Source/Programming/UnrealBuildSystem/Configuration/BuildConfigProperties/BuildConfigProperties.INT.udn
+++ b/Engine/Documentation/Source/Programming/UnrealBuildSystem/Configuration/BuildConfigProperties/BuildConfigProperties.INT.udn
@@ -92,7 +92,7 @@ $ bDisableDebugInfo : Whether to globally disable debug info generation; see Deb
 
 $ bDisableDebugInfoForGeneratedCode : Whether to disable debug info generation for generated files. This improves link times for modules that have a lot of generated glue code.
 
-$ bOmitPCDebugInfoInDevelopment : Whether to disable debug info on PC in development builds (for faster developer iteration, as link times are extremely fast with debug info disabled).
+$ bOmitPCDebugInfoInDevelopment : Whether to disable debug info on PC/Mac in development builds (for faster developer iteration, as link times are extremely fast with debug info disabled).
 
 $ bUsePDBFiles : Whether PDB files should be used for Visual C++ builds.
 
@@ -154,9 +154,6 @@ $ BaseLogFileName : Specifies the file to use for logging.
 
 $ bStripSymbols : Whether to strip iOS symbols or not (implied by Shipping config).
 
-$ bGeneratedSYMFile : Whether to generate dSYM files. Defaults to true for Shipping builds and false for all other builds. -dsym will force generation of a dsym for other builds,
-        * NoDsym will force it off for shipping.
-
 $ bWriteSolutionOptionFile : Whether to write a solution option (suo) file for the sln.
 
 $ bAddFastPDBToProjects : Whether to add the -FastPDB option to build command lines by default.
@@ -275,6 +272,10 @@ $ bAllowOverVpn : When set to false, XGE will not be enabled when running connec
 
 $ VpnSubnets : List of subnets containing IP addresses assigned by VPN
 
+$ bAllowRemoteLinking : Whether to allow remote linking
+
+$ bUseVCCompilerMode : Whether to enable the VCCompiler=true setting. This requires an additional license for VC tools.
+
 ### BuildMode
 
 $ bIgnoreJunk : Whether to skip checking for files identified by the junk manifest.
@@ -643,3 +644,5 @@ $ RsyncAuthentication : The authentication used for Rsync (for the -e rsync flag
 
 $ SshAuthentication : The authentication used for SSH (probably similar to RsyncAuthentication).
 
+### Log
+
diff --git a/Engine/Documentation/Source/Programming/UnrealBuildSystem/TargetFiles/TargetFilesProperties/TargetFilesProperties.INT.udn b/Engine/Documentation/Source/Programming/UnrealBuildSystem/TargetFiles/TargetFilesProperties/TargetFilesProperties.INT.udn
--- a/Engine/Documentation/Source/Programming/UnrealBuildSystem/TargetFiles/TargetFilesProperties/TargetFilesProperties.INT.udn
+++ b/Engine/Documentation/Source/Programming/UnrealBuildSystem/TargetFiles/TargetFilesProperties/TargetFilesProperties.INT.udn
@@ -214,7 +214,7 @@ $ bDisableDebugInfo (Boolean): Whether to globally disable debug info generation
 
 $ bDisableDebugInfoForGeneratedCode (Boolean): Whether to disable debug info generation for generated files. This improves link times for modules that have a lot of generated glue code.
 
-$ bOmitPCDebugInfoInDevelopment (Boolean): Whether to disable debug info on PC in development builds (for faster developer iteration, as link times are extremely fast with debug info disabled).
+$ bOmitPCDebugInfoInDevelopment (Boolean): Whether to disable debug info on PC/Mac in development builds (for faster developer iteration, as link times are extremely fast with debug info disabled).
 
 $ bUsePDBFiles (Boolean): Whether PDB files should be used for Visual C++ builds.
 
