Documentation for build tools

diff --git a/Engine/Documentation/Source/Programming/UnrealBuildSystem/Configuration/BuildConfigProperties/BuildConfigProperties.INT.udn b/Engine/Documentation/Source/Programming/UnrealBuildSystem/Configuration/BuildConfigProperties/BuildConfigProperties.INT.udn
--- a/Engine/Documentation/Source/Programming/UnrealBuildSystem/Configuration/BuildConfigProperties/BuildConfigProperties.INT.udn
+++ b/Engine/Documentation/Source/Programming/UnrealBuildSystem/Configuration/BuildConfigProperties/BuildConfigProperties.INT.udn
@@ -248,6 +248,10 @@ $ bShowCompilationTimes : Whether to show compilation times along with worst off
 
 $ bShowPerActionCompilationTimes : Whether to show compilation times for each executed action
 
+$ bLogActionCommandLines : Whether to log command lines for actions being executed
+
+$ bPrintActionTargetNames : Add target names for each action executed
+
 ### HybridExecutor
 
 $ MaxLocalActions : Maximum number of actions to execute locally.
@@ -260,6 +264,10 @@ $ bShowCompilationTimes : Whether to show compilation times along with worst off
 
 $ bShowPerActionCompilationTimes : Whether to show compilation times for each executed action
 
+$ bLogActionCommandLines : Whether to log command lines for actions being executed
+
+$ bPrintActionTargetNames : Add target names for each action executed
+
 $ MaxProcessorCount : Maximum processor count for local execution.
 
 $ ProcessorCountMultiplier : Processor count multiplier for local execution. Can be below 1 to reserve CPU for other tasks. When using the local executor (not XGE), run a single action on each CPU core. Note that you can set this to a larger value to get slightly faster build times in many cases, but your computer's responsiveness during compiling may be much worse. This value is ignored if the CPU does not support hyper-threading.
diff --git a/Engine/Documentation/Source/Programming/UnrealBuildSystem/ModuleFiles/ModuleFilesProperties/ModuleFilesProperties.INT.udn b/Engine/Documentation/Source/Programming/UnrealBuildSystem/ModuleFiles/ModuleFilesProperties/ModuleFilesProperties.INT.udn
--- a/Engine/Documentation/Source/Programming/UnrealBuildSystem/ModuleFiles/ModuleFilesProperties/ModuleFilesProperties.INT.udn
+++ b/Engine/Documentation/Source/Programming/UnrealBuildSystem/ModuleFiles/ModuleFilesProperties/ModuleFilesProperties.INT.udn
@@ -10,12 +10,6 @@ $ IsPlugin (Boolean): True if a Plugin contains this module
 
 $ HasPackageOverride (Boolean): Returns true if there has been an override type specified on this module
 
-$ PublicLibraryPaths (List&lt;String&gt;): List of system/library paths (directory of .lib files) - typically used for External (third party) modules
-
-$ Definitions (List&lt;String&gt;): Accessor for the PublicDefinitions list
-
-$ PlatformSpecificDynamicallyLoadedModuleNames (List&lt;String&gt;): Extra modules this module may require at run time, that are on behalf of another platform (i.e. shader formats and the like)
-
 $ AutoSdkDirectory (String): The AutoSDK directory for the active host platform
 
 $ EngineDirectory (String): The current engine directory
@@ -118,8 +112,6 @@ $ AdditionalBundleResources (List&lt;BundleResource&gt;): List of addition resou
 
 $ TypeLibraries (List&lt;TypeLibrary&gt;): List of type libraries that we need to generate headers for (Windows only)
 
-$ PublicAdditionalShadowFiles (List&lt;String&gt;): For builds that execute on a remote machine (e.g. iOS), this list contains additional files that need to be copied over in order for the app to link successfully.  Source/header files and PCHs are automatically copied.  Usually this is simply a list of precompiled third party library dependencies.
-
 $ PublicDelayLoadDLLs (List&lt;String&gt;): List of delay load DLLs - typically used for External (third party) modules
 
 $ PrivateDefinitions (List&lt;String&gt;): Private compiler definitions for this module
@@ -156,20 +148,12 @@ $ PCHUsage (PCHUsageMode): Precompiled header usage for this module
 
 $ DefaultBuildSettings (BuildSettingsVersion): Which engine version's build settings to use by default.
 
-$ bUseBackwardsCompatibleDefaults (Boolean): Whether to use backwards compatible defaults for this module. By default, engine modules always use the latest default settings, while project modules do not (to support an easier migration path).
-
 $ ShadowVariableWarningLevel (WarningLevel): How to treat shadow variable warnings
 
-$ bEnableShadowVariableWarnings (Boolean): Enable warnings for shadowed variables
-
 $ UnsafeTypeCastWarningLevel (WarningLevel): How to treat unsafe implicit type cast warnings (e.g., double->float or int64->int32)
 
-$ bFasterWithoutUnity (Boolean): If true and unity builds are enabled, this module will build without unity.
-
 $ bUseUnity (Boolean): If unity builds are enabled this can be used to override if this specific module will build using Unity. This is set using the per module configurations in BuildConfiguration.
 
-$ bOutputPubliclyDistributable (Boolean): Obsolete. Use bLegalToDistributeObjectCode instead.
-
 $ bLegacyPublicIncludePaths (Boolean): Whether this module qualifies included headers from other modules relative to the root of their 'Public' folder. This reduces the number of search paths that have to be passed to the compiler, improving performance and reducing the length of the compiler command line.
 
 
diff --git a/Engine/Documentation/Source/Programming/UnrealBuildSystem/TargetFiles/TargetFilesProperties/TargetFilesProperties.INT.udn b/Engine/Documentation/Source/Programming/UnrealBuildSystem/TargetFiles/TargetFilesProperties/TargetFilesProperties.INT.udn
--- a/Engine/Documentation/Source/Programming/UnrealBuildSystem/TargetFiles/TargetFilesProperties/TargetFilesProperties.INT.udn
+++ b/Engine/Documentation/Source/Programming/UnrealBuildSystem/TargetFiles/TargetFilesProperties/TargetFilesProperties.INT.udn
@@ -16,8 +16,6 @@ $ ProjectFile (FileReference): Path to the project file for the project containi
 
 $ Version (ReadOnlyBuildVersion): The current build version
 
-$ ExcludePlugins (List&lt;String&gt;): Accessor for
-
 $ HostPlatform (UnrealTargetPlatform): Gets the host platform being built on
 
 $ bGenerateProjectFiles (Boolean): Expose the bGenerateProjectFiles flag to targets, so we can modify behavior as appropriate for better intellisense
@@ -44,8 +42,6 @@ $ bLegalToDistributeBinary (Boolean): Whether the output from this target can be
 
 $ UndecoratedConfiguration (UnrealTargetConfiguration): Specifies the configuration whose binaries do not require a "-Platform-Configuration" suffix.
 
-$ bBuildAllPlugins (Boolean): Build all the plugins that we can find, even if they're not enabled. This is particularly useful for content-only projects, where you're building the UnrealEditor target but running it with a game that enables a plugin.
-
 $ bBuildAllModules (Boolean): Build all the modules that are valid for this target type. Used for CIS and making installed engine builds.
 
 $ BuildPlugins (List&lt;String&gt;): Additional plugins that are built for this target type but not enabled.
@@ -332,8 +328,6 @@ $ Name (String): The name of this target
 
 $ DefaultBuildSettings (BuildSettingsVersion): Specifies the engine version to maintain backwards-compatible default build settings with (eg. DefaultSettingsVersion.Release_4_23, DefaultSettingsVersion.Release_4_24). Specify DefaultSettingsVersion.Latest to always use defaults for the current engine version, at the risk of introducing build errors while upgrading.
 
-$ bOutputPubliclyDistributable (Boolean): Obsolete. Use bLegalToDistributeBinary instead.
-
 $ bAllowHotReload (Boolean): Whether this target supports hot reload
 
 $ bBuildEditor (Boolean): Whether to compile the editor or not. Only desktop platforms (Windows or Mac) will use this, other platforms force this to false.
@@ -346,8 +340,6 @@ $ bBuildDeveloperTools (Boolean): Whether to compile the developer tools.
 
 $ bBuildTargetDeveloperTools (Boolean): Whether to compile the developer tools that are for target platforms or connected devices (defaults to bBuildDeveloperTools)
 
-$ bCompileLeanAndMeanUE (Boolean): Whether to compile lean and mean version of UE.
-
 $ bCompileSpeedTree (Boolean): Whether we should compile in support for Simplygon or not.
 
 $ bWithServerCode (Boolean): Compile server-only code.
@@ -360,8 +352,6 @@ $ bWithPerfCounters (Boolean): Whether to include PerfCounters support.
 
 $ bWithLiveCoding (Boolean): Whether to enable support for live coding
 
-$ bUseBackwardsCompatibleDefaults (Boolean): Whether to use backwards compatible defaults for this module. By default, engine modules always use the latest default settings, while project modules do not (to support an easier migration path).
-
 $ bHasExports (Boolean): Whether the final executable should export symbols.
 
 $ bBuildAdditionalConsoleApp (Boolean): If true, creates an additional console application. Hack for Windows, where it's not possible to conditionally inherit a parent's console Window depending on how the application is invoked; you have to link the same executable with a different subsystem setting.
@@ -374,12 +364,6 @@ $ DefaultWarningLevel (WarningLevel): Default treatment of uncategorized warning
 
 $ DeprecationWarningLevel (WarningLevel): Level to report deprecation warnings as errors
 
-$ bShadowVariableErrors (Boolean): Forces shadow variable warnings to be treated as errors on platforms that support it.
-
-$ bStripSymbolsOnIOS (Boolean): Whether to strip iOS symbols or not (implied by bGeneratedSYMFile).
-
-$ bCreateStubIPA (Boolean): If true, then a stub IPA will be generated when compiling is done (minimal files needed for a valid IPA).
-
 $ bLegacyPublicIncludePaths (Boolean): Add all the public folders as include paths for the compile environment.
 
 $ LinkType (TargetLinkType): Specifies how to link modules in this target (monolithic or modular). This is currently protected for backwards compatibility. Call the GetLinkType() accessor until support for the deprecated ShouldCompileMonolithic() override has been removed.
