Documentation for build tools

diff --git a/Engine/Documentation/Source/Programming/UnrealBuildSystem/Configuration/BuildConfigProperties/BuildConfigProperties.INT.udn b/Engine/Documentation/Source/Programming/UnrealBuildSystem/Configuration/BuildConfigProperties/BuildConfigProperties.INT.udn
--- a/Engine/Documentation/Source/Programming/UnrealBuildSystem/Configuration/BuildConfigProperties/BuildConfigProperties.INT.udn
+++ b/Engine/Documentation/Source/Programming/UnrealBuildSystem/Configuration/BuildConfigProperties/BuildConfigProperties.INT.udn
@@ -654,3 +654,7 @@ $ SshAuthentication : The authentication used for SSH (probably similar to Rsync
 
 ### Log
 
+$ bBackupLogFiles : Whether to backup an existing log file, rather than overwriting it.
+
+$ LogFileBackupCount : The number of log file backups to preserve. Older backups will be deleted.
+
