Documentation for build tools

diff --git a/Engine/Documentation/Source/Programming/UnrealBuildSystem/Configuration/BuildConfigProperties/BuildConfigProperties.INT.udn b/Engine/Documentation/Source/Programming/UnrealBuildSystem/Configuration/BuildConfigProperties/BuildConfigProperties.INT.udn
--- a/Engine/Documentation/Source/Programming/UnrealBuildSystem/Configuration/BuildConfigProperties/BuildConfigProperties.INT.udn
+++ b/Engine/Documentation/Source/Programming/UnrealBuildSystem/Configuration/BuildConfigProperties/BuildConfigProperties.INT.udn
@@ -98,6 +98,8 @@ $ bUseIncrementalLinking : Whether to use incremental linking or not. Incrementa
 
 $ bAllowLTCG : Whether to allow the use of link time code generation (LTCG).
 
+$ bPreferThinLTO : When Link Time Code Generation (LTCG) is enabled, whether to prefer using the lighter weight version on supported platforms.
+
 $ bPGOProfile : Whether to enable Profile Guided Optimization (PGO) instrumentation in this build.
 
 $ bPGOOptimize : Whether to optimize this build with Profile Guided Optimization (PGO).
diff --git a/Engine/Documentation/Source/Programming/UnrealBuildSystem/TargetFiles/TargetFilesProperties/TargetFilesProperties.INT.udn b/Engine/Documentation/Source/Programming/UnrealBuildSystem/TargetFiles/TargetFilesProperties/TargetFilesProperties.INT.udn
--- a/Engine/Documentation/Source/Programming/UnrealBuildSystem/TargetFiles/TargetFilesProperties/TargetFilesProperties.INT.udn
+++ b/Engine/Documentation/Source/Programming/UnrealBuildSystem/TargetFiles/TargetFilesProperties/TargetFilesProperties.INT.udn
@@ -228,6 +228,8 @@ $ bUseIncrementalLinking (Boolean): Whether to use incremental linking or not. I
 
 $ bAllowLTCG (Boolean): Whether to allow the use of link time code generation (LTCG).
 
+$ bPreferThinLTO (Boolean): When Link Time Code Generation (LTCG) is enabled, whether to prefer using the lighter weight version on supported platforms.
+
 $ bPGOProfile (Boolean): Whether to enable Profile Guided Optimization (PGO) instrumentation in this build.
 
 $ bPGOOptimize (Boolean): Whether to optimize this build with Profile Guided Optimization (PGO).
