Documentation for build tools

diff --git a/Engine/Documentation/Source/Programming/UnrealBuildSystem/TargetFiles/TargetFilesProperties/TargetFilesProperties.INT.udn b/Engine/Documentation/Source/Programming/UnrealBuildSystem/TargetFiles/TargetFilesProperties/TargetFilesProperties.INT.udn
--- a/Engine/Documentation/Source/Programming/UnrealBuildSystem/TargetFiles/TargetFilesProperties/TargetFilesProperties.INT.udn
+++ b/Engine/Documentation/Source/Programming/UnrealBuildSystem/TargetFiles/TargetFilesProperties/TargetFilesProperties.INT.udn
@@ -334,7 +334,7 @@ $ DefaultBuildSettings (BuildSettingsVersion): Specifies the engine version to m
 
 $ bAllowHotReload (Boolean): Whether this target supports hot reload
 
-$ bBuildEditor (Boolean): Whether to compile the editor or not. Only desktop platforms (Windows or Mac) will use this, other platforms force this to false.
+$ bBuildEditor (Boolean): Whether we are compiling editor code or not. Prefer the more explicit bCompileAgainstEditor instead.
 
 $ bBuildRequiresCookedData (Boolean): Whether to compile code related to building assets. Consoles generally cannot build assets. Desktop platforms generally can.
 
@@ -344,7 +344,9 @@ $ bBuildDeveloperTools (Boolean): Whether to compile the developer tools.
 
 $ bBuildTargetDeveloperTools (Boolean): Whether to compile the developer tools that are for target platforms or connected devices (defaults to bBuildDeveloperTools)
 
-$ bCompileSpeedTree (Boolean): Whether we should compile in support for Simplygon or not.
+$ bCompileAgainstEditor (Boolean): Enabled for editor builds (TargetType.Editor). Can be overridden for programs (TargetType.Program) that would need to compile against editor code. Not available for other target types. Mainly drives the value of WITH_EDITOR.
+
+$ bCompileSpeedTree (Boolean): Whether we should compile in support for SpeedTree or not.
 
 $ bWithServerCode (Boolean): Compile server-only code.
 
