Documentation for build tools

diff --git a/Engine/Documentation/Source/Programming/UnrealBuildSystem/Configuration/BuildConfigProperties/BuildConfigProperties.INT.udn b/Engine/Documentation/Source/Programming/UnrealBuildSystem/Configuration/BuildConfigProperties/BuildConfigProperties.INT.udn
--- a/Engine/Documentation/Source/Programming/UnrealBuildSystem/Configuration/BuildConfigProperties/BuildConfigProperties.INT.udn
+++ b/Engine/Documentation/Source/Programming/UnrealBuildSystem/Configuration/BuildConfigProperties/BuildConfigProperties.INT.udn
@@ -36,6 +36,8 @@ $ bDebugBuildsActuallyUseDebugCRT : Enables the debug C++ runtime (CRT) for debu
 
 $ bLegalToDistributeBinary : Whether the output from this target can be publicly distributed, even if it has dependencies on modules that are in folders with special restrictions (eg. CarefullyRedist, NotForLicensees, NoRedist).
 
+$ bUseVerse : Whether to use the verse script interface.
+
 $ bUseInlining : Enable inlining for all modules.
 
 $ bUseDebugLiveCodingConsole : Whether to enable support for live coding
@@ -588,6 +590,8 @@ $ MaxSharedIncludePaths : Puts the most common include paths in the IncludePath
 
 $ ExcludedIncludePaths : Semi-colon separated list of paths that should not be added to the projects include paths. Useful for omitting third-party headers (e.g ThirdParty/WebRTC) from intellisense suggestions and reducing memory footprints.
 
+$ ExcludedFilePaths : Semi-colon separated list of paths that should not be added to the projects. Useful for omitting third-party files (e.g ThirdParty/WebRTC) from intellisense suggestions and reducing memory footprints.
+
 $ bBuildUBTInDebug : Forces UBT to be built in debug configuration, regardless of the solution configuration
 
 $ bBuildLiveCodingConsole : Whether to include a dependency on LiveCodingConsole when building targets that support live coding.
