Documentation for build tools

diff --git a/Engine/Documentation/Source/Programming/UnrealBuildSystem/Configuration/BuildConfigProperties/BuildConfigProperties.INT.udn b/Engine/Documentation/Source/Programming/UnrealBuildSystem/Configuration/BuildConfigProperties/BuildConfigProperties.INT.udn
--- a/Engine/Documentation/Source/Programming/UnrealBuildSystem/Configuration/BuildConfigProperties/BuildConfigProperties.INT.udn
+++ b/Engine/Documentation/Source/Programming/UnrealBuildSystem/Configuration/BuildConfigProperties/BuildConfigProperties.INT.udn
@@ -320,6 +320,8 @@ $ bIncludeTemplateFiles : True if we should include template files in the genera
 
 $ IncludeEnginePrograms : True if we should include program projects in the generated solution.
 
+$ IncludeCppSource : Whether to include C++ targets
+
 $ bIncludeDotNetPrograms : True if we should include csharp program projects in the generated solution. Pass "-DotNet" to enable this.
 
 $ bIncludeTempTargets : Whether to include temporary targets generated by UAT to support content only projects with non-default settings.
@@ -372,6 +374,8 @@ $ bIncludeTemplateFiles : True if we should include template files in the genera
 
 $ IncludeEnginePrograms : True if we should include program projects in the generated solution.
 
+$ IncludeCppSource : Whether to include C++ targets
+
 $ bIncludeDotNetPrograms : True if we should include csharp program projects in the generated solution. Pass "-DotNet" to enable this.
 
 $ MasterProjectName : Name of the master project file -- for example, the base file name for the Visual Studio solution file, or the Xcode project file on Mac.
@@ -394,6 +398,8 @@ $ bIncludeTemplateFiles : True if we should include template files in the genera
 
 $ IncludeEnginePrograms : True if we should include program projects in the generated solution.
 
+$ IncludeCppSource : Whether to include C++ targets
+
 $ bIncludeDotNetPrograms : True if we should include csharp program projects in the generated solution. Pass "-DotNet" to enable this.
 
 $ MasterProjectName : Name of the master project file -- for example, the base file name for the Visual Studio solution file, or the Xcode project file on Mac.
@@ -416,6 +422,8 @@ $ bIncludeTemplateFiles : True if we should include template files in the genera
 
 $ IncludeEnginePrograms : True if we should include program projects in the generated solution.
 
+$ IncludeCppSource : Whether to include C++ targets
+
 $ bIncludeDotNetPrograms : True if we should include csharp program projects in the generated solution. Pass "-DotNet" to enable this.
 
 $ MasterProjectName : Name of the master project file -- for example, the base file name for the Visual Studio solution file, or the Xcode project file on Mac.
@@ -438,6 +446,8 @@ $ bIncludeTemplateFiles : True if we should include template files in the genera
 
 $ IncludeEnginePrograms : True if we should include program projects in the generated solution.
 
+$ IncludeCppSource : Whether to include C++ targets
+
 $ bIncludeDotNetPrograms : True if we should include csharp program projects in the generated solution. Pass "-DotNet" to enable this.
 
 $ MasterProjectName : Name of the master project file -- for example, the base file name for the Visual Studio solution file, or the Xcode project file on Mac.
@@ -460,6 +470,8 @@ $ bIncludeTemplateFiles : True if we should include template files in the genera
 
 $ IncludeEnginePrograms : True if we should include program projects in the generated solution.
 
+$ IncludeCppSource : Whether to include C++ targets
+
 $ bIncludeDotNetPrograms : True if we should include csharp program projects in the generated solution. Pass "-DotNet" to enable this.
 
 $ MasterProjectName : Name of the master project file -- for example, the base file name for the Visual Studio solution file, or the Xcode project file on Mac.
@@ -482,6 +494,8 @@ $ bIncludeTemplateFiles : True if we should include template files in the genera
 
 $ IncludeEnginePrograms : True if we should include program projects in the generated solution.
 
+$ IncludeCppSource : Whether to include C++ targets
+
 $ bIncludeDotNetPrograms : True if we should include csharp program projects in the generated solution. Pass "-DotNet" to enable this.
 
 $ MasterProjectName : Name of the master project file -- for example, the base file name for the Visual Studio solution file, or the Xcode project file on Mac.
@@ -504,6 +518,8 @@ $ bIncludeTemplateFiles : True if we should include template files in the genera
 
 $ IncludeEnginePrograms : True if we should include program projects in the generated solution.
 
+$ IncludeCppSource : Whether to include C++ targets
+
 $ bIncludeDotNetPrograms : True if we should include csharp program projects in the generated solution. Pass "-DotNet" to enable this.
 
 $ MasterProjectName : Name of the master project file -- for example, the base file name for the Visual Studio solution file, or the Xcode project file on Mac.
@@ -526,6 +542,8 @@ $ bIncludeTemplateFiles : True if we should include template files in the genera
 
 $ IncludeEnginePrograms : True if we should include program projects in the generated solution.
 
+$ IncludeCppSource : Whether to include C++ targets
+
 $ bIncludeDotNetPrograms : True if we should include csharp program projects in the generated solution. Pass "-DotNet" to enable this.
 
 $ MasterProjectName : Name of the master project file -- for example, the base file name for the Visual Studio solution file, or the Xcode project file on Mac.
@@ -554,6 +572,8 @@ $ bIncludeTemplateFiles : True if we should include template files in the genera
 
 $ IncludeEnginePrograms : True if we should include program projects in the generated solution.
 
+$ IncludeCppSource : Whether to include C++ targets
+
 $ bIncludeDotNetPrograms : True if we should include csharp program projects in the generated solution. Pass "-DotNet" to enable this.
 
 $ MasterProjectName : Name of the master project file -- for example, the base file name for the Visual Studio solution file, or the Xcode project file on Mac.
@@ -576,6 +596,8 @@ $ bIncludeTemplateFiles : True if we should include template files in the genera
 
 $ IncludeEnginePrograms : True if we should include program projects in the generated solution.
 
+$ IncludeCppSource : Whether to include C++ targets
+
 $ bIncludeDotNetPrograms : True if we should include csharp program projects in the generated solution. Pass "-DotNet" to enable this.
 
 $ MasterProjectName : Name of the master project file -- for example, the base file name for the Visual Studio solution file, or the Xcode project file on Mac.
@@ -610,6 +632,8 @@ $ bIncludeTemplateFiles : True if we should include template files in the genera
 
 $ IncludeEnginePrograms : True if we should include program projects in the generated solution.
 
+$ IncludeCppSource : Whether to include C++ targets
+
 $ bIncludeDotNetPrograms : True if we should include csharp program projects in the generated solution. Pass "-DotNet" to enable this.
 
 $ MasterProjectName : Name of the master project file -- for example, the base file name for the Visual Studio solution file, or the Xcode project file on Mac.
@@ -632,6 +656,8 @@ $ bIncludeTemplateFiles : True if we should include template files in the genera
 
 $ IncludeEnginePrograms : True if we should include program projects in the generated solution.
 
+$ IncludeCppSource : Whether to include C++ targets
+
 $ bIncludeDotNetPrograms : True if we should include csharp program projects in the generated solution. Pass "-DotNet" to enable this.
 
 $ MasterProjectName : Name of the master project file -- for example, the base file name for the Visual Studio solution file, or the Xcode project file on Mac.
