Documentation for build tools

diff --git a/Engine/Documentation/Source/Programming/UnrealBuildSystem/Configuration/BuildConfigProperties/BuildConfigProperties.INT.udn b/Engine/Documentation/Source/Programming/UnrealBuildSystem/Configuration/BuildConfigProperties/BuildConfigProperties.INT.udn
--- a/Engine/Documentation/Source/Programming/UnrealBuildSystem/Configuration/BuildConfigProperties/BuildConfigProperties.INT.udn
+++ b/Engine/Documentation/Source/Programming/UnrealBuildSystem/Configuration/BuildConfigProperties/BuildConfigProperties.INT.udn
@@ -326,15 +326,15 @@ $ bIncludeDotNetPrograms : True if we should include csharp program projects in
 
 $ bIncludeTempTargets : Whether to include temporary targets generated by UAT to support content only projects with non-default settings.
 
-$ bKeepSourceSubDirectories : True if we should reflect "Source" sub-directories on disk in the master project as master project directories. This (arguably) adds some visual clutter to the master project but it is truer to the on-disk file organization.
+$ bKeepSourceSubDirectories : True if we should reflect "Source" sub-directories on disk in the primary project as project directories. This (arguably) adds some visual clutter to the primary project but it is truer to the on-disk file organization.
 
 $ Platforms : Names of platforms to include in the generated project files
 
 $ Configurations : Names of configurations to include in the generated project files. See UnrealTargetConfiguration for valid entries
 
-$ MasterProjectName : Name of the master project file -- for example, the base file name for the Visual Studio solution file, or the Xcode project file on Mac.
+$ PrimaryProjectName : Name of the primary project file -- for example, the base file name for the Visual Studio solution file, or the Xcode project file on Mac.
 
-$ bMasterProjectNameFromFolder : If true, sets the master project name according to the name of the folder it is in.
+$ bPrimaryProjectNameFromFolder : If true, sets the primary project name according to the name of the folder it is in.
 
 $ bIncludeTestAndShippingConfigs : Whether we should include configurations for "Test" and "Shipping" in generated projects. Pass "-NoShippingConfigs" to disable this.
 
@@ -378,9 +378,9 @@ $ IncludeCppSource : Whether to include C++ targets
 
 $ bIncludeDotNetPrograms : True if we should include csharp program projects in the generated solution. Pass "-DotNet" to enable this.
 
-$ MasterProjectName : Name of the master project file -- for example, the base file name for the Visual Studio solution file, or the Xcode project file on Mac.
+$ PrimaryProjectName : Name of the primary project file -- for example, the base file name for the Visual Studio solution file, or the Xcode project file on Mac.
 
-$ bMasterProjectNameFromFolder : If true, sets the master project name according to the name of the folder it is in.
+$ bPrimaryProjectNameFromFolder : If true, sets the primary project name according to the name of the folder it is in.
 
 ### CMakefileGenerator
 
@@ -402,9 +402,9 @@ $ IncludeCppSource : Whether to include C++ targets
 
 $ bIncludeDotNetPrograms : True if we should include csharp program projects in the generated solution. Pass "-DotNet" to enable this.
 
-$ MasterProjectName : Name of the master project file -- for example, the base file name for the Visual Studio solution file, or the Xcode project file on Mac.
+$ PrimaryProjectName : Name of the primary project file -- for example, the base file name for the Visual Studio solution file, or the Xcode project file on Mac.
 
-$ bMasterProjectNameFromFolder : If true, sets the master project name according to the name of the folder it is in.
+$ bPrimaryProjectNameFromFolder : If true, sets the primary project name according to the name of the folder it is in.
 
 ### CodeLiteGenerator
 
@@ -426,9 +426,9 @@ $ IncludeCppSource : Whether to include C++ targets
 
 $ bIncludeDotNetPrograms : True if we should include csharp program projects in the generated solution. Pass "-DotNet" to enable this.
 
-$ MasterProjectName : Name of the master project file -- for example, the base file name for the Visual Studio solution file, or the Xcode project file on Mac.
+$ PrimaryProjectName : Name of the primary project file -- for example, the base file name for the Visual Studio solution file, or the Xcode project file on Mac.
 
-$ bMasterProjectNameFromFolder : If true, sets the master project name according to the name of the folder it is in.
+$ bPrimaryProjectNameFromFolder : If true, sets the primary project name according to the name of the folder it is in.
 
 ### EddieProjectFileGenerator
 
@@ -450,9 +450,9 @@ $ IncludeCppSource : Whether to include C++ targets
 
 $ bIncludeDotNetPrograms : True if we should include csharp program projects in the generated solution. Pass "-DotNet" to enable this.
 
-$ MasterProjectName : Name of the master project file -- for example, the base file name for the Visual Studio solution file, or the Xcode project file on Mac.
+$ PrimaryProjectName : Name of the primary project file -- for example, the base file name for the Visual Studio solution file, or the Xcode project file on Mac.
 
-$ bMasterProjectNameFromFolder : If true, sets the master project name according to the name of the folder it is in.
+$ bPrimaryProjectNameFromFolder : If true, sets the primary project name according to the name of the folder it is in.
 
 ### KDevelopGenerator
 
@@ -474,9 +474,9 @@ $ IncludeCppSource : Whether to include C++ targets
 
 $ bIncludeDotNetPrograms : True if we should include csharp program projects in the generated solution. Pass "-DotNet" to enable this.
 
-$ MasterProjectName : Name of the master project file -- for example, the base file name for the Visual Studio solution file, or the Xcode project file on Mac.
+$ PrimaryProjectName : Name of the primary project file -- for example, the base file name for the Visual Studio solution file, or the Xcode project file on Mac.
 
-$ bMasterProjectNameFromFolder : If true, sets the master project name according to the name of the folder it is in.
+$ bPrimaryProjectNameFromFolder : If true, sets the primary project name according to the name of the folder it is in.
 
 ### MakefileGenerator
 
@@ -498,9 +498,9 @@ $ IncludeCppSource : Whether to include C++ targets
 
 $ bIncludeDotNetPrograms : True if we should include csharp program projects in the generated solution. Pass "-DotNet" to enable this.
 
-$ MasterProjectName : Name of the master project file -- for example, the base file name for the Visual Studio solution file, or the Xcode project file on Mac.
+$ PrimaryProjectName : Name of the primary project file -- for example, the base file name for the Visual Studio solution file, or the Xcode project file on Mac.
 
-$ bMasterProjectNameFromFolder : If true, sets the master project name according to the name of the folder it is in.
+$ bPrimaryProjectNameFromFolder : If true, sets the primary project name according to the name of the folder it is in.
 
 ### QMakefileGenerator
 
@@ -522,9 +522,9 @@ $ IncludeCppSource : Whether to include C++ targets
 
 $ bIncludeDotNetPrograms : True if we should include csharp program projects in the generated solution. Pass "-DotNet" to enable this.
 
-$ MasterProjectName : Name of the master project file -- for example, the base file name for the Visual Studio solution file, or the Xcode project file on Mac.
+$ PrimaryProjectName : Name of the primary project file -- for example, the base file name for the Visual Studio solution file, or the Xcode project file on Mac.
 
-$ bMasterProjectNameFromFolder : If true, sets the master project name according to the name of the folder it is in.
+$ bPrimaryProjectNameFromFolder : If true, sets the primary project name according to the name of the folder it is in.
 
 ### RiderProjectFileGenerator
 
@@ -546,9 +546,9 @@ $ IncludeCppSource : Whether to include C++ targets
 
 $ bIncludeDotNetPrograms : True if we should include csharp program projects in the generated solution. Pass "-DotNet" to enable this.
 
-$ MasterProjectName : Name of the master project file -- for example, the base file name for the Visual Studio solution file, or the Xcode project file on Mac.
+$ PrimaryProjectName : Name of the primary project file -- for example, the base file name for the Visual Studio solution file, or the Xcode project file on Mac.
 
-$ bMasterProjectNameFromFolder : If true, sets the master project name according to the name of the folder it is in.
+$ bPrimaryProjectNameFromFolder : If true, sets the primary project name according to the name of the folder it is in.
 
 ### VSCodeProjectFileGenerator
 
@@ -576,9 +576,9 @@ $ IncludeCppSource : Whether to include C++ targets
 
 $ bIncludeDotNetPrograms : True if we should include csharp program projects in the generated solution. Pass "-DotNet" to enable this.
 
-$ MasterProjectName : Name of the master project file -- for example, the base file name for the Visual Studio solution file, or the Xcode project file on Mac.
+$ PrimaryProjectName : Name of the primary project file -- for example, the base file name for the Visual Studio solution file, or the Xcode project file on Mac.
 
-$ bMasterProjectNameFromFolder : If true, sets the master project name according to the name of the folder it is in.
+$ bPrimaryProjectNameFromFolder : If true, sets the primary project name according to the name of the folder it is in.
 
 ### VCMacProjectFileGenerator
 
@@ -600,9 +600,9 @@ $ IncludeCppSource : Whether to include C++ targets
 
 $ bIncludeDotNetPrograms : True if we should include csharp program projects in the generated solution. Pass "-DotNet" to enable this.
 
-$ MasterProjectName : Name of the master project file -- for example, the base file name for the Visual Studio solution file, or the Xcode project file on Mac.
+$ PrimaryProjectName : Name of the primary project file -- for example, the base file name for the Visual Studio solution file, or the Xcode project file on Mac.
 
-$ bMasterProjectNameFromFolder : If true, sets the master project name according to the name of the folder it is in.
+$ bPrimaryProjectNameFromFolder : If true, sets the primary project name according to the name of the folder it is in.
 
 ### VCProjectFileGenerator
 
@@ -636,9 +636,9 @@ $ IncludeCppSource : Whether to include C++ targets
 
 $ bIncludeDotNetPrograms : True if we should include csharp program projects in the generated solution. Pass "-DotNet" to enable this.
 
-$ MasterProjectName : Name of the master project file -- for example, the base file name for the Visual Studio solution file, or the Xcode project file on Mac.
+$ PrimaryProjectName : Name of the primary project file -- for example, the base file name for the Visual Studio solution file, or the Xcode project file on Mac.
 
-$ bMasterProjectNameFromFolder : If true, sets the master project name according to the name of the folder it is in.
+$ bPrimaryProjectNameFromFolder : If true, sets the primary project name according to the name of the folder it is in.
 
 ### XcodeProjectFileGenerator
 
@@ -660,9 +660,9 @@ $ IncludeCppSource : Whether to include C++ targets
 
 $ bIncludeDotNetPrograms : True if we should include csharp program projects in the generated solution. Pass "-DotNet" to enable this.
 
-$ MasterProjectName : Name of the master project file -- for example, the base file name for the Visual Studio solution file, or the Xcode project file on Mac.
+$ PrimaryProjectName : Name of the primary project file -- for example, the base file name for the Visual Studio solution file, or the Xcode project file on Mac.
 
-$ bMasterProjectNameFromFolder : If true, sets the master project name according to the name of the folder it is in.
+$ bPrimaryProjectNameFromFolder : If true, sets the primary project name according to the name of the folder it is in.
 
 ### SourceFileWorkingSet
 
