Documentation for build tools

diff --git a/Engine/Documentation/Source/Programming/UnrealBuildSystem/Configuration/BuildConfigProperties/BuildConfigProperties.INT.udn b/Engine/Documentation/Source/Programming/UnrealBuildSystem/Configuration/BuildConfigProperties/BuildConfigProperties.INT.udn
--- a/Engine/Documentation/Source/Programming/UnrealBuildSystem/Configuration/BuildConfigProperties/BuildConfigProperties.INT.udn
+++ b/Engine/Documentation/Source/Programming/UnrealBuildSystem/Configuration/BuildConfigProperties/BuildConfigProperties.INT.udn
@@ -158,6 +158,8 @@ $ bEnableUndefinedBehaviorSanitizer : Enables undefined behavior sanitizer (UBSa
 
 $ bEnableMemorySanitizer : Enables memory sanitizer (MSan)
 
+$ bTuneDebugInfoForLLDB : Turns on tuning of debug info for LLDB
+
 $ bWriteSolutionOptionFile : Whether to write a solution option (suo) file for the sln.
 
 $ bAddFastPDBToProjects : Whether to add the -FastPDB option to build command lines by default.
