Documentation for build tools

diff --git a/Engine/Documentation/Source/Programming/UnrealBuildSystem/ModuleFiles/ModuleFilesProperties/ModuleFilesProperties.INT.udn b/Engine/Documentation/Source/Programming/UnrealBuildSystem/ModuleFiles/ModuleFilesProperties/ModuleFilesProperties.INT.udn
--- a/Engine/Documentation/Source/Programming/UnrealBuildSystem/ModuleFiles/ModuleFilesProperties/ModuleFilesProperties.INT.udn
+++ b/Engine/Documentation/Source/Programming/UnrealBuildSystem/ModuleFiles/ModuleFilesProperties/ModuleFilesProperties.INT.udn
@@ -132,7 +132,7 @@ $ AdditionalCodeGenDirectories (List&lt;String&gt;): External directories contai
 
 $ SubclassRules (List&lt;String&gt;): Subclass rules files which invalidate the makefile if modified.
 
-$ bRequiresImplementModule (Nullable&lt;Boolean&gt;): Whether this module requires the IMPLEMENT_MODULE macro to be implemented. Most UE4 modules require this, since we use the IMPLEMENT_MODULE macro to do other global overloads (eg. operator new/delete forwarding to GMalloc).
+$ bRequiresImplementModule (Nullable&lt;Boolean&gt;): Whether this module requires the IMPLEMENT_MODULE macro to be implemented. Most UE modules require this, since we use the IMPLEMENT_MODULE macro to do other global overloads (eg. operator new/delete forwarding to GMalloc).
 
 $ CppStandard (CppStandardVersion): Which stanard to use for compiling this module
 
diff --git a/Engine/Documentation/Source/Programming/UnrealBuildSystem/TargetFiles/TargetFilesProperties/TargetFilesProperties.INT.udn b/Engine/Documentation/Source/Programming/UnrealBuildSystem/TargetFiles/TargetFilesProperties/TargetFilesProperties.INT.udn
--- a/Engine/Documentation/Source/Programming/UnrealBuildSystem/TargetFiles/TargetFilesProperties/TargetFilesProperties.INT.udn
+++ b/Engine/Documentation/Source/Programming/UnrealBuildSystem/TargetFiles/TargetFilesProperties/TargetFilesProperties.INT.udn
@@ -82,7 +82,7 @@ $ bUseChaosChecked (Boolean): Whether to compile in checked chaos features for d
 
 $ bUseChaosMemoryTracking (Boolean): Whether to compile in chaos memory tracking features
 
-$ bCustomSceneQueryStructure (Boolean): Whether scene query acceleration is done by UE4. The physx scene query structure is still created, but we do not use it.
+$ bCustomSceneQueryStructure (Boolean): Whether scene query acceleration is done by UE. The physx scene query structure is still created, but we do not use it.
 
 $ bCompilePhysX (Boolean): Whether to include PhysX support.
 
