Documentation for build tools

diff --git a/Engine/Documentation/Source/Programming/UnrealBuildSystem/Configuration/BuildConfigProperties/BuildConfigProperties.INT.udn b/Engine/Documentation/Source/Programming/UnrealBuildSystem/Configuration/BuildConfigProperties/BuildConfigProperties.INT.udn
--- a/Engine/Documentation/Source/Programming/UnrealBuildSystem/Configuration/BuildConfigProperties/BuildConfigProperties.INT.udn
+++ b/Engine/Documentation/Source/Programming/UnrealBuildSystem/Configuration/BuildConfigProperties/BuildConfigProperties.INT.udn
@@ -18,8 +18,6 @@ $ bAllowFASTBuild : Whether FASTBuild may be used.
 
 $ bAllowSNDBS : Whether SN-DBS may be used.
 
-$ bAllowTaskExecutor : Whether the experimental async TaskExecutor may be used.
-
 $ bUseUBTMakefiles : Enables support for very fast iterative builds by caching target data. Turning this on causes Unreal Build Tool to emit 'UBT Makefiles' for targets when they are built the first time. Subsequent builds will load these Makefiles and begin outdatedness checking and build invocation very quickly. The caveat is that if source files are added or removed to the project, UBT will need to gather information about those in order for your build to complete successfully. Currently, you must run the project file generator after adding/removing source files to tell UBT to re-gather this information. 
     * Events that can invalidate the 'UBT Makefile':
         * Adding/removing .cpp files
@@ -240,7 +238,11 @@ $ bStopOnError : Whether to stop on error
 
 $ MsvcCRTRedistVersion : Which MSVC CRT Redist version to use
 
-### TaskExecutor
+### HybridExecutor
+
+$ MaxLocalActions : Maximum number of actions to execute locally.
+
+### ParallelExecutor
 
 $ MaxProcessorCount : Maximum processor count for local execution.
 
@@ -258,28 +260,6 @@ $ bLogActionCommandLines : Whether to log command lines for actions being execut
 
 $ bPrintActionTargetNames : Add target names for each action executed
 
-### HybridExecutor
-
-$ MaxLocalActions : Maximum number of actions to execute locally.
-
-### ParallelExecutor
-
-$ bStopCompilationAfterErrors : When enabled, will stop compiling targets after a compile error occurs.
-
-$ bShowCompilationTimes : Whether to show compilation times along with worst offenders or not.
-
-$ bShowPerActionCompilationTimes : Whether to show compilation times for each executed action
-
-$ bLogActionCommandLines : Whether to log command lines for actions being executed
-
-$ bPrintActionTargetNames : Add target names for each action executed
-
-$ MaxProcessorCount : Maximum processor count for local execution.
-
-$ ProcessorCountMultiplier : Processor count multiplier for local execution. Can be below 1 to reserve CPU for other tasks. When using the local executor (not XGE), run a single action on each CPU core. Note that you can set this to a larger value to get slightly faster build times in many cases, but your computer's responsiveness during compiling may be much worse. This value is ignored if the CPU does not support hyper-threading.
-
-$ MemoryPerActionBytes : Free memory per action in bytes, used to limit the number of parallel actions if the machine is memory starved. Set to 0 to disable free memory checking.
-
 ### XGE
 
 $ bAllowOverVpn : When set to false, XGE will not be enabled when running connected to the coordinator over VPN. Configure VPN-assigned subnets via the VpnSubnets parameter.
diff --git a/Engine/Documentation/Source/Programming/UnrealBuildSystem/TargetFiles/TargetFilesProperties/TargetFilesProperties.INT.udn b/Engine/Documentation/Source/Programming/UnrealBuildSystem/TargetFiles/TargetFilesProperties/TargetFilesProperties.INT.udn
--- a/Engine/Documentation/Source/Programming/UnrealBuildSystem/TargetFiles/TargetFilesProperties/TargetFilesProperties.INT.udn
+++ b/Engine/Documentation/Source/Programming/UnrealBuildSystem/TargetFiles/TargetFilesProperties/TargetFilesProperties.INT.udn
@@ -316,7 +316,7 @@ $ AdditionalCompilerArguments (String): Additional arguments to pass to the comp
 
 $ AdditionalLinkerArguments (String): Additional arguments to pass to the linker
 
-$ MemoryPerActionGB (Double): Max amount of memory that each compile action may require. Used by ParallelExecutor and TaskExecutor to decide the maximum number of parallel actions to start at one time.
+$ MemoryPerActionGB (Double): Max amount of memory that each compile action may require. Used by ParallelExecutor to decide the maximum number of parallel actions to start at one time.
 
 $ DisableUnityBuildForModules (String[]): List of modules to disable unity builds for
 
