Documentation for build tools

diff --git a/Engine/Documentation/Source/Programming/UnrealBuildSystem/ModuleFiles/ModuleFilesProperties/ModuleFilesProperties.INT.udn b/Engine/Documentation/Source/Programming/UnrealBuildSystem/ModuleFiles/ModuleFilesProperties/ModuleFilesProperties.INT.udn
--- a/Engine/Documentation/Source/Programming/UnrealBuildSystem/ModuleFiles/ModuleFilesProperties/ModuleFilesProperties.INT.udn
+++ b/Engine/Documentation/Source/Programming/UnrealBuildSystem/ModuleFiles/ModuleFilesProperties/ModuleFilesProperties.INT.udn
@@ -18,6 +18,8 @@ $ PluginDirectory (String): Property for the directory containing this plugin. U
 
 $ ModuleDirectory (String): Property for the directory containing this module. Useful for adding paths to third party dependencies.
 
+$ TestsDirectory (String): Returns module's low level tests directory "Tests".
+
 $ IsVcPackageSupported (Boolean): Returns if VcPkg is supported for the build configuration.
 
 
diff --git a/Engine/Documentation/Source/Programming/UnrealBuildSystem/TargetFiles/TargetFilesProperties/TargetFilesProperties.INT.udn b/Engine/Documentation/Source/Programming/UnrealBuildSystem/TargetFiles/TargetFilesProperties/TargetFilesProperties.INT.udn
--- a/Engine/Documentation/Source/Programming/UnrealBuildSystem/TargetFiles/TargetFilesProperties/TargetFilesProperties.INT.udn
+++ b/Engine/Documentation/Source/Programming/UnrealBuildSystem/TargetFiles/TargetFilesProperties/TargetFilesProperties.INT.udn
@@ -240,6 +240,8 @@ $ bOmitFramePointers (Boolean): Whether to omit frame pointers or not. Disabling
 
 $ bEnableCppModules (Boolean): Whether to enable support for C++20 modules
 
+$ bEnableCppCoroutinesForEvaluation (Boolean): Whether to enable support for C++20 coroutines This option is provided to facilitate evaluation of the feature. Expect the name of the option to change. Use of coroutines in with UE is untested and unsupported.
+
 $ bUseMallocProfiler (Boolean): If true, then enable memory profiling in the build (defines USE_MALLOC_PROFILER=1 and forces bOmitFramePointers=false).
 
 $ bUseSharedPCHs (Boolean): Enables "Shared PCHs", a feature which significantly speeds up compile times by attempting to share certain PCH files between modules that UBT detects is including those PCH's header files.
