Documentation for build tools

diff --git a/Engine/Documentation/Source/Programming/UnrealBuildSystem/ModuleFiles/ModuleFilesProperties/ModuleFilesProperties.INT.udn b/Engine/Documentation/Source/Programming/UnrealBuildSystem/ModuleFiles/ModuleFilesProperties/ModuleFilesProperties.INT.udn
--- a/Engine/Documentation/Source/Programming/UnrealBuildSystem/ModuleFiles/ModuleFilesProperties/ModuleFilesProperties.INT.udn
+++ b/Engine/Documentation/Source/Programming/UnrealBuildSystem/ModuleFiles/ModuleFilesProperties/ModuleFilesProperties.INT.udn
@@ -10,6 +10,8 @@ $ IsPlugin (Boolean): True if a Plugin contains this module
 
 $ HasPackageOverride (Boolean): Returns true if there has been an override type specified on this module
 
+$ bBuildIncludeTestsFolder (Boolean): Decides whether to include the "Tests" folder at the root of the module as part of the compilation. By default this folder is excluded from the module's compilation. The "Tests" folder should be reserved for low level tests but there are exceptions to this rule.
+
 $ AutoSdkDirectory (String): The AutoSDK directory for the active host platform
 
 $ EngineDirectory (String): The current engine directory
