Documentation for build tools

diff --git a/Engine/Documentation/Source/Programming/UnrealBuildSystem/TargetFiles/TargetFilesProperties/TargetFilesProperties.INT.udn b/Engine/Documentation/Source/Programming/UnrealBuildSystem/TargetFiles/TargetFilesProperties/TargetFilesProperties.INT.udn
--- a/Engine/Documentation/Source/Programming/UnrealBuildSystem/TargetFiles/TargetFilesProperties/TargetFilesProperties.INT.udn
+++ b/Engine/Documentation/Source/Programming/UnrealBuildSystem/TargetFiles/TargetFilesProperties/TargetFilesProperties.INT.udn
@@ -106,6 +106,8 @@ $ bForceBuildTargetPlatforms (Boolean): Whether to force compiling the target pl
 
 $ bForceBuildShaderFormats (Boolean): Whether to force compiling shader format modules, even if they wouldn't normally be built.
 
+$ bNeedsExtraShaderFormatsOverride (Nullable&lt;Boolean&gt;): Override for including extra shader formats
+
 $ bCompileCustomSQLitePlatform (Boolean): Whether we should compile SQLite using the custom "Unreal" platform (true), or using the native platform (false).
 
 $ bUseCacheFreedOSAllocs (Boolean): Whether to utilize cache freed OS allocs with MallocBinned
@@ -350,6 +352,8 @@ $ bBuildDeveloperTools (Boolean): Whether to compile the developer tools.
 
 $ bBuildTargetDeveloperTools (Boolean): Whether to compile the developer tools that are for target platforms or connected devices (defaults to bBuildDeveloperTools)
 
+$ bNeedsExtraShaderFormats (Boolean): Whether we should include any extra shader formats. By default this is only enabled for Program and Editor targets.
+
 $ bCompileAgainstEditor (Boolean): Enabled for editor builds (TargetType.Editor). Can be overridden for programs (TargetType.Program) that would need to compile against editor code. Not available for other target types. Mainly drives the value of WITH_EDITOR.
 
 $ bCompileSpeedTree (Boolean): Whether we should compile in support for SpeedTree or not.
