Documentation for build tools

diff --git a/Engine/Documentation/Source/Programming/UnrealBuildSystem/TargetFiles/TargetFilesProperties/TargetFilesProperties.INT.udn b/Engine/Documentation/Source/Programming/UnrealBuildSystem/TargetFiles/TargetFilesProperties/TargetFilesProperties.INT.udn
--- a/Engine/Documentation/Source/Programming/UnrealBuildSystem/TargetFiles/TargetFilesProperties/TargetFilesProperties.INT.udn
+++ b/Engine/Documentation/Source/Programming/UnrealBuildSystem/TargetFiles/TargetFilesProperties/TargetFilesProperties.INT.udn
@@ -360,6 +360,8 @@ $ bCompileSpeedTree (Boolean): Whether we should compile in support for SpeedTre
 
 $ bWithServerCode (Boolean): Compile server-only code.
 
+$ bFNameOutlineNumber (Boolean): Compile with FName storing the number part in the name table. Saves memory when most names are not numbered and those that are are referenced multiple times. The game and engine must ensure they reuse numbered names similarly to name strings to avoid leaking memory.
+
 $ bWithPushModel (Boolean): When enabled, Push Model Networking support will be compiled in. This can help reduce CPU overhead of networking, at the cost of more memory. Always enabled in editor builds.
 
 $ bIncludePluginsForTargetPlatforms (Boolean): Whether to allow plugins which support all target platforms.
