Documentation for build tools

diff --git a/Engine/Documentation/Source/Programming/UnrealBuildSystem/ModuleFiles/ModuleFilesProperties/ModuleFilesProperties.INT.udn b/Engine/Documentation/Source/Programming/UnrealBuildSystem/ModuleFiles/ModuleFilesProperties/ModuleFilesProperties.INT.udn
--- a/Engine/Documentation/Source/Programming/UnrealBuildSystem/ModuleFiles/ModuleFilesProperties/ModuleFilesProperties.INT.udn
+++ b/Engine/Documentation/Source/Programming/UnrealBuildSystem/ModuleFiles/ModuleFilesProperties/ModuleFilesProperties.INT.udn
@@ -22,6 +22,8 @@ $ ModuleDirectory (String): Property for the directory containing this module. U
 
 $ TestsDirectory (String): Returns module's low level tests directory "Tests".
 
+$ BuildOrderSettings (SourceFilesBuildOrderSettings): Optional compilation order override rules for module's source files.
+
 $ IsVcPackageSupported (Boolean): Returns if VcPkg is supported for the build configuration.
 
 
