Documentation for build tools

diff --git a/Engine/Documentation/Source/Programming/UnrealBuildSystem/TargetFiles/TargetFilesProperties/TargetFilesProperties.INT.udn b/Engine/Documentation/Source/Programming/UnrealBuildSystem/TargetFiles/TargetFilesProperties/TargetFilesProperties.INT.udn
--- a/Engine/Documentation/Source/Programming/UnrealBuildSystem/TargetFiles/TargetFilesProperties/TargetFilesProperties.INT.udn
+++ b/Engine/Documentation/Source/Programming/UnrealBuildSystem/TargetFiles/TargetFilesProperties/TargetFilesProperties.INT.udn
@@ -16,8 +16,6 @@ $ ProjectFile (FileReference): Path to the project file for the project containi
 
 $ Version (ReadOnlyBuildVersion): The current build version
 
-$ NameSuffix (String): Optional suffix to append to Name
-
 $ bIncludeAllTests (Boolean): Decides whether to compile with all the "Tests" folders from all dependent modules.
 
 $ HostPlatform (UnrealTargetPlatform): Gets the host platform being built on
