Documentation for build tools

diff --git a/Engine/Documentation/Source/Programming/UnrealBuildSystem/TargetFiles/TargetFilesProperties/TargetFilesProperties.INT.udn b/Engine/Documentation/Source/Programming/UnrealBuildSystem/TargetFiles/TargetFilesProperties/TargetFilesProperties.INT.udn
--- a/Engine/Documentation/Source/Programming/UnrealBuildSystem/TargetFiles/TargetFilesProperties/TargetFilesProperties.INT.udn
+++ b/Engine/Documentation/Source/Programming/UnrealBuildSystem/TargetFiles/TargetFilesProperties/TargetFilesProperties.INT.udn
@@ -148,6 +148,8 @@ $ bUseLauncherChecks (Boolean): Whether to check that the process was launched t
 
 $ bUseChecksInShipping (Boolean): Whether to turn on checks (asserts) for test/shipping builds.
 
+$ bTCHARIsUTF8 (Boolean): Whether to turn on UTF-8 mode, mapping TCHAR to UTF8CHAR.
+
 $ bUseEstimatedUtcNow (Boolean): Whether to use the EstimatedUtcNow or PlatformUtcNow.  EstimatedUtcNow is appropriate in cases where PlatformUtcNow can be slow.
 
 $ bCompileFreeType (Boolean): True if we need FreeType support.
