Documentation for build tools

diff --git a/Engine/Documentation/Source/Programming/UnrealBuildSystem/Configuration/BuildConfigProperties/BuildConfigProperties.INT.udn b/Engine/Documentation/Source/Programming/UnrealBuildSystem/Configuration/BuildConfigProperties/BuildConfigProperties.INT.udn
--- a/Engine/Documentation/Source/Programming/UnrealBuildSystem/Configuration/BuildConfigProperties/BuildConfigProperties.INT.udn
+++ b/Engine/Documentation/Source/Programming/UnrealBuildSystem/Configuration/BuildConfigProperties/BuildConfigProperties.INT.udn
@@ -138,6 +138,8 @@ $ bPublicSymbolsByDefault : Whether to expose all symbols as public by default o
 
 $ CppStandard : Which C++ stanard to use for compiling this target
 
+$ bStopSNDBSCompilationAfterErrors : When enabled, SN-DBS will stop compiling targets after a compile error occurs.  Recommended, as it saves computing resources for others.
+
 $ bXGENoWatchdogThread : Whether to use the no_watchdog_thread option to prevent VS2015 toolchain stalls.
 
 $ bShowXGEMonitor : Whether to display the XGE build monitor.
@@ -264,6 +266,12 @@ $ bLogActionCommandLines : Whether to log command lines for actions being execut
 
 $ bPrintActionTargetNames : Add target names for each action executed
 
+### SNDBS
+
+$ bAllowOverVpn : When set to false, SNDBS will not be enabled when running connected to the coordinator over VPN. Configure VPN-assigned subnets via the VpnSubnets parameter.
+
+$ VpnSubnets : List of subnets containing IP addresses assigned by VPN
+
 ### XGE
 
 $ bAllowOverVpn : When set to false, XGE will not be enabled when running connected to the coordinator over VPN. Configure VPN-assigned subnets via the VpnSubnets parameter.
