Documentation for build tools

diff --git a/Engine/Documentation/Source/Programming/UnrealBuildSystem/ModuleFiles/ModuleFilesProperties/ModuleFilesProperties.INT.udn b/Engine/Documentation/Source/Programming/UnrealBuildSystem/ModuleFiles/ModuleFilesProperties/ModuleFilesProperties.INT.udn
--- a/Engine/Documentation/Source/Programming/UnrealBuildSystem/ModuleFiles/ModuleFilesProperties/ModuleFilesProperties.INT.udn
+++ b/Engine/Documentation/Source/Programming/UnrealBuildSystem/ModuleFiles/ModuleFilesProperties/ModuleFilesProperties.INT.udn
@@ -140,6 +140,8 @@ $ SubclassRules (List&lt;String&gt;): Subclass rules files which invalidate the
 
 $ bRequiresImplementModule (Nullable&lt;Boolean&gt;): Whether this module requires the IMPLEMENT_MODULE macro to be implemented. Most UE modules require this, since we use the IMPLEMENT_MODULE macro to do other global overloads (eg. operator new/delete forwarding to GMalloc).
 
+$ VersePath (String): If this module has associated Verse code, this is the Verse root path of it
+
 $ CppStandard (CppStandardVersion): Which stanard to use for compiling this module
 
 $ ModuleSymbolVisibility (SymbolVisibility): Control visibility of symbols
