Documentation for build tools

diff --git a/Engine/Documentation/Source/Programming/UnrealBuildSystem/Configuration/BuildConfigProperties/BuildConfigProperties.INT.udn b/Engine/Documentation/Source/Programming/UnrealBuildSystem/Configuration/BuildConfigProperties/BuildConfigProperties.INT.udn
--- a/Engine/Documentation/Source/Programming/UnrealBuildSystem/Configuration/BuildConfigProperties/BuildConfigProperties.INT.udn
+++ b/Engine/Documentation/Source/Programming/UnrealBuildSystem/Configuration/BuildConfigProperties/BuildConfigProperties.INT.udn
@@ -184,6 +184,8 @@ $ bAllowHotReloadFromIDE : True if hot-reload from IDE is allowed.
 
 $ bForceDebugUnrealHeaderTool : If true, the Debug version of UnrealHeaderTool will be built and run instead of the Development version.
 
+$ bUseBuiltInUnrealHeaderTool : If true, use UHT internal to UBT
+
 ### WindowsPlatform
 
 $ MaxRootPathLength : Maximum recommended root path length.
