Documentation for build tools

diff --git a/Engine/Documentation/Source/Programming/UnrealBuildSystem/Configuration/BuildConfigProperties/BuildConfigProperties.INT.udn b/Engine/Documentation/Source/Programming/UnrealBuildSystem/Configuration/BuildConfigProperties/BuildConfigProperties.INT.udn
--- a/Engine/Documentation/Source/Programming/UnrealBuildSystem/Configuration/BuildConfigProperties/BuildConfigProperties.INT.udn
+++ b/Engine/Documentation/Source/Programming/UnrealBuildSystem/Configuration/BuildConfigProperties/BuildConfigProperties.INT.udn
@@ -314,7 +314,7 @@ $ bIncludeLocalizationFiles : True if we should include localization files in th
 
 $ bIncludeTemplateFiles : True if we should include template files in the generated project.
 
-$ IncludeEnginePrograms : True if we should include program projects in the generated solution.
+$ bIncludeEnginePrograms : True if we should include program projects in the generated solution.
 
 $ IncludeCppSource : Whether to include C++ targets
 
@@ -368,7 +368,7 @@ $ bIncludeConfigFiles : True if we should include config (.ini) files in the gen
 
 $ bIncludeTemplateFiles : True if we should include template files in the generated project.
 
-$ IncludeEnginePrograms : True if we should include program projects in the generated solution.
+$ bIncludeEnginePrograms : True if we should include program projects in the generated solution.
 
 $ IncludeCppSource : Whether to include C++ targets
 
@@ -392,7 +392,7 @@ $ bIncludeConfigFiles : True if we should include config (.ini) files in the gen
 
 $ bIncludeTemplateFiles : True if we should include template files in the generated project.
 
-$ IncludeEnginePrograms : True if we should include program projects in the generated solution.
+$ bIncludeEnginePrograms : True if we should include program projects in the generated solution.
 
 $ IncludeCppSource : Whether to include C++ targets
 
@@ -416,7 +416,7 @@ $ bIncludeConfigFiles : True if we should include config (.ini) files in the gen
 
 $ bIncludeTemplateFiles : True if we should include template files in the generated project.
 
-$ IncludeEnginePrograms : True if we should include program projects in the generated solution.
+$ bIncludeEnginePrograms : True if we should include program projects in the generated solution.
 
 $ IncludeCppSource : Whether to include C++ targets
 
@@ -440,7 +440,7 @@ $ bIncludeConfigFiles : True if we should include config (.ini) files in the gen
 
 $ bIncludeTemplateFiles : True if we should include template files in the generated project.
 
-$ IncludeEnginePrograms : True if we should include program projects in the generated solution.
+$ bIncludeEnginePrograms : True if we should include program projects in the generated solution.
 
 $ IncludeCppSource : Whether to include C++ targets
 
@@ -464,7 +464,7 @@ $ bIncludeConfigFiles : True if we should include config (.ini) files in the gen
 
 $ bIncludeTemplateFiles : True if we should include template files in the generated project.
 
-$ IncludeEnginePrograms : True if we should include program projects in the generated solution.
+$ bIncludeEnginePrograms : True if we should include program projects in the generated solution.
 
 $ IncludeCppSource : Whether to include C++ targets
 
@@ -488,7 +488,7 @@ $ bIncludeConfigFiles : True if we should include config (.ini) files in the gen
 
 $ bIncludeTemplateFiles : True if we should include template files in the generated project.
 
-$ IncludeEnginePrograms : True if we should include program projects in the generated solution.
+$ bIncludeEnginePrograms : True if we should include program projects in the generated solution.
 
 $ IncludeCppSource : Whether to include C++ targets
 
@@ -512,7 +512,7 @@ $ bIncludeConfigFiles : True if we should include config (.ini) files in the gen
 
 $ bIncludeTemplateFiles : True if we should include template files in the generated project.
 
-$ IncludeEnginePrograms : True if we should include program projects in the generated solution.
+$ bIncludeEnginePrograms : True if we should include program projects in the generated solution.
 
 $ IncludeCppSource : Whether to include C++ targets
 
@@ -536,7 +536,7 @@ $ bIncludeConfigFiles : True if we should include config (.ini) files in the gen
 
 $ bIncludeTemplateFiles : True if we should include template files in the generated project.
 
-$ IncludeEnginePrograms : True if we should include program projects in the generated solution.
+$ bIncludeEnginePrograms : True if we should include program projects in the generated solution.
 
 $ IncludeCppSource : Whether to include C++ targets
 
@@ -566,7 +566,7 @@ $ bIncludeConfigFiles : True if we should include config (.ini) files in the gen
 
 $ bIncludeTemplateFiles : True if we should include template files in the generated project.
 
-$ IncludeEnginePrograms : True if we should include program projects in the generated solution.
+$ bIncludeEnginePrograms : True if we should include program projects in the generated solution.
 
 $ IncludeCppSource : Whether to include C++ targets
 
@@ -590,7 +590,7 @@ $ bIncludeConfigFiles : True if we should include config (.ini) files in the gen
 
 $ bIncludeTemplateFiles : True if we should include template files in the generated project.
 
-$ IncludeEnginePrograms : True if we should include program projects in the generated solution.
+$ bIncludeEnginePrograms : True if we should include program projects in the generated solution.
 
 $ IncludeCppSource : Whether to include C++ targets
 
@@ -626,7 +626,7 @@ $ bIncludeConfigFiles : True if we should include config (.ini) files in the gen
 
 $ bIncludeTemplateFiles : True if we should include template files in the generated project.
 
-$ IncludeEnginePrograms : True if we should include program projects in the generated solution.
+$ bIncludeEnginePrograms : True if we should include program projects in the generated solution.
 
 $ IncludeCppSource : Whether to include C++ targets
 
@@ -650,7 +650,7 @@ $ bIncludeConfigFiles : True if we should include config (.ini) files in the gen
 
 $ bIncludeTemplateFiles : True if we should include template files in the generated project.
 
-$ IncludeEnginePrograms : True if we should include program projects in the generated solution.
+$ bIncludeEnginePrograms : True if we should include program projects in the generated solution.
 
 $ IncludeCppSource : Whether to include C++ targets
 
