Documentation for build tools

diff --git a/Engine/Documentation/Source/Programming/UnrealBuildSystem/Configuration/BuildConfigProperties/BuildConfigProperties.INT.udn b/Engine/Documentation/Source/Programming/UnrealBuildSystem/Configuration/BuildConfigProperties/BuildConfigProperties.INT.udn
--- a/Engine/Documentation/Source/Programming/UnrealBuildSystem/Configuration/BuildConfigProperties/BuildConfigProperties.INT.udn
+++ b/Engine/Documentation/Source/Programming/UnrealBuildSystem/Configuration/BuildConfigProperties/BuildConfigProperties.INT.udn
@@ -284,6 +284,10 @@ $ bAllowRemoteLinking : Whether to allow remote linking
 
 $ bUseVCCompilerMode : Whether to enable the VCCompiler=true setting. This requires an additional license for VC tools.
 
+$ MinActions : Minimum number of actions to use XGE execution.
+
+$ bUnavailableIfInUse : Check for a concurrent XGE build and treat the XGE executor as unavailable if it's in use. This will allow UBT to fall back to another executor such as the parallel executor.
+
 ### BuildMode
 
 $ bIgnoreJunk : Whether to skip checking for files identified by the junk manifest.
