Documentation for build tools

diff --git a/Engine/Documentation/Source/Programming/UnrealBuildSystem/Configuration/BuildConfigProperties/BuildConfigProperties.INT.udn b/Engine/Documentation/Source/Programming/UnrealBuildSystem/Configuration/BuildConfigProperties/BuildConfigProperties.INT.udn
--- a/Engine/Documentation/Source/Programming/UnrealBuildSystem/Configuration/BuildConfigProperties/BuildConfigProperties.INT.udn
+++ b/Engine/Documentation/Source/Programming/UnrealBuildSystem/Configuration/BuildConfigProperties/BuildConfigProperties.INT.udn
@@ -126,6 +126,8 @@ $ bCreateMapFile : Outputs a map file as part of the build.
 
 $ bAllowRuntimeSymbolFiles : True if runtime symbols files should be generated as a post build step for some platforms. These files are used by the engine to resolve symbol names of callstack backtraces in logs.
 
+$ PackagePath : Package full path (directory + filename) where to store input files used at link time Normally used to debug a linker crash for platforms that support it
+
 $ bCheckSystemHeadersForModification : Whether headers in system paths should be checked for modification when determining outdated actions.
 
 $ bFlushBuildDirOnRemoteMac : Whether to clean Builds directory on a remote Mac before building.
diff --git a/Engine/Documentation/Source/Programming/UnrealBuildSystem/TargetFiles/TargetFilesProperties/TargetFilesProperties.INT.udn b/Engine/Documentation/Source/Programming/UnrealBuildSystem/TargetFiles/TargetFilesProperties/TargetFilesProperties.INT.udn
--- a/Engine/Documentation/Source/Programming/UnrealBuildSystem/TargetFiles/TargetFilesProperties/TargetFilesProperties.INT.udn
+++ b/Engine/Documentation/Source/Programming/UnrealBuildSystem/TargetFiles/TargetFilesProperties/TargetFilesProperties.INT.udn
@@ -266,6 +266,8 @@ $ bCreateMapFile (Boolean): Outputs a map file as part of the build.
 
 $ bAllowRuntimeSymbolFiles (Boolean): True if runtime symbols files should be generated as a post build step for some platforms. These files are used by the engine to resolve symbol names of callstack backtraces in logs.
 
+$ PackagePath (String): Package full path (directory + filename) where to store input files used at link time Normally used to debug a linker crash for platforms that support it
+
 $ BundleVersion (String): Bundle version for Mac apps.
 
 $ bDeployAfterCompile (Boolean): Whether to deploy the executable after compilation on platforms that require deployment.
