Documentation for build tools

diff --git a/Engine/Documentation/Source/Programming/UnrealBuildSystem/ModuleFiles/ModuleFilesProperties/ModuleFilesProperties.INT.udn b/Engine/Documentation/Source/Programming/UnrealBuildSystem/ModuleFiles/ModuleFilesProperties/ModuleFilesProperties.INT.udn
--- a/Engine/Documentation/Source/Programming/UnrealBuildSystem/ModuleFiles/ModuleFilesProperties/ModuleFilesProperties.INT.udn
+++ b/Engine/Documentation/Source/Programming/UnrealBuildSystem/ModuleFiles/ModuleFilesProperties/ModuleFilesProperties.INT.udn
@@ -134,8 +134,6 @@ $ PrecompileForTargets (PrecompileTargetsType): Which targets this module should
 
 $ ExternalDependencies (List&lt;String&gt;): External files which invalidate the makefile if modified. Relative paths are resolved relative to the .build.cs file.
 
-$ AdditionalCodeGenDirectories (List&lt;String&gt;): External directories containing generated interop files.
-
 $ SubclassRules (List&lt;String&gt;): Subclass rules files which invalidate the makefile if modified.
 
 $ bRequiresImplementModule (Nullable&lt;Boolean&gt;): Whether this module requires the IMPLEMENT_MODULE macro to be implemented. Most UE modules require this, since we use the IMPLEMENT_MODULE macro to do other global overloads (eg. operator new/delete forwarding to GMalloc).
