Documentation for build tools

diff --git a/Engine/Documentation/Source/Programming/UnrealBuildSystem/TargetFiles/TargetFilesProperties/TargetFilesProperties.INT.udn b/Engine/Documentation/Source/Programming/UnrealBuildSystem/TargetFiles/TargetFilesProperties/TargetFilesProperties.INT.udn
--- a/Engine/Documentation/Source/Programming/UnrealBuildSystem/TargetFiles/TargetFilesProperties/TargetFilesProperties.INT.udn
+++ b/Engine/Documentation/Source/Programming/UnrealBuildSystem/TargetFiles/TargetFilesProperties/TargetFilesProperties.INT.udn
@@ -98,6 +98,8 @@ $ bCompileCEF3 (Boolean): Whether to compile CEF3 support.
 
 $ bCompileISPC (Boolean): Whether to compile using ISPC.
 
+$ bCompileIntelMetricsDiscovery (Boolean): Whether to compile IntelMetricsDiscovery.
+
 $ bCompilePython (Boolean): Whether to compile in python support
 
 $ bForceBuildTargetPlatforms (Boolean): Whether to force compiling the target platform modules, even if they wouldn't normally be built.
