Documentation for build tools

diff --git a/Engine/Documentation/Source/Programming/UnrealBuildSystem/ModuleFiles/ModuleFilesProperties/ModuleFilesProperties.INT.udn b/Engine/Documentation/Source/Programming/UnrealBuildSystem/ModuleFiles/ModuleFilesProperties/ModuleFilesProperties.INT.udn
--- a/Engine/Documentation/Source/Programming/UnrealBuildSystem/ModuleFiles/ModuleFilesProperties/ModuleFilesProperties.INT.udn
+++ b/Engine/Documentation/Source/Programming/UnrealBuildSystem/ModuleFiles/ModuleFilesProperties/ModuleFilesProperties.INT.udn
@@ -154,6 +154,8 @@ $ PCHUsage (PCHUsageMode): Precompiled header usage for this module
 
 $ DefaultBuildSettings (BuildSettingsVersion): Which engine version's build settings to use by default.
 
+$ IncludeOrderVersion (EngineIncludeOrderVersion): What version of include order to use when compiling this module. Can be overridden via -ForceIncludeOrder on the command line or in a module's rules.
+
 $ ShadowVariableWarningLevel (WarningLevel): How to treat shadow variable warnings
 
 $ UnsafeTypeCastWarningLevel (WarningLevel): How to treat unsafe implicit type cast warnings (e.g., double->float or int64->int32)
diff --git a/Engine/Documentation/Source/Programming/UnrealBuildSystem/TargetFiles/TargetFilesProperties/TargetFilesProperties.INT.udn b/Engine/Documentation/Source/Programming/UnrealBuildSystem/TargetFiles/TargetFilesProperties/TargetFilesProperties.INT.udn
--- a/Engine/Documentation/Source/Programming/UnrealBuildSystem/TargetFiles/TargetFilesProperties/TargetFilesProperties.INT.udn
+++ b/Engine/Documentation/Source/Programming/UnrealBuildSystem/TargetFiles/TargetFilesProperties/TargetFilesProperties.INT.udn
@@ -30,6 +30,8 @@ $ bIsEngineInstalled (Boolean): Expose a setting for whether or not the engine i
 ### Read/Write Properties
 $ Type (TargetType): The type of target.
 
+$ ForcedIncludeOrder (Nullable&lt;EngineIncludeOrderVersion&gt;): Force the include order to a specific version. Overrides any Target and Module rules.
+
 $ bUsesSteam (Boolean): Whether the target uses Steam.
 
 $ bUsesCEF3 (Boolean): Whether the target uses CEF3.
@@ -344,6 +346,8 @@ $ Name (String): The name of this target
 
 $ DefaultBuildSettings (BuildSettingsVersion): Specifies the engine version to maintain backwards-compatible default build settings with (eg. DefaultSettingsVersion.Release_4_23, DefaultSettingsVersion.Release_4_24). Specify DefaultSettingsVersion.Latest to always use defaults for the current engine version, at the risk of introducing build errors while upgrading.
 
+$ IncludeOrderVersion (EngineIncludeOrderVersion): What version of include order to use when compiling this target. Can be overridden via -ForceIncludeOrder on the command line. ModuleRules.IncludeOrderVersion takes precedence.
+
 $ bAllowHotReload (Boolean): Whether this target supports hot reload
 
 $ bBuildEditor (Boolean): Whether we are compiling editor code or not. Prefer the more explicit bCompileAgainstEditor instead.
