Documentation for build tools

diff --git a/Engine/Documentation/Source/Programming/UnrealBuildSystem/Configuration/BuildConfigProperties/BuildConfigProperties.INT.udn b/Engine/Documentation/Source/Programming/UnrealBuildSystem/Configuration/BuildConfigProperties/BuildConfigProperties.INT.udn
--- a/Engine/Documentation/Source/Programming/UnrealBuildSystem/Configuration/BuildConfigProperties/BuildConfigProperties.INT.udn
+++ b/Engine/Documentation/Source/Programming/UnrealBuildSystem/Configuration/BuildConfigProperties/BuildConfigProperties.INT.udn
@@ -128,6 +128,8 @@ $ bAllowRuntimeSymbolFiles : True if runtime symbols files should be generated a
 
 $ PackagePath : Package full path (directory + filename) where to store input files used at link time Normally used to debug a linker crash for platforms that support it
 
+$ CrashDiagnosticDirectory : Directory where to put crash report files for platforms that support it
+
 $ bCheckSystemHeadersForModification : Whether headers in system paths should be checked for modification when determining outdated actions.
 
 $ bFlushBuildDirOnRemoteMac : Whether to clean Builds directory on a remote Mac before building.
diff --git a/Engine/Documentation/Source/Programming/UnrealBuildSystem/TargetFiles/TargetFilesProperties/TargetFilesProperties.INT.udn b/Engine/Documentation/Source/Programming/UnrealBuildSystem/TargetFiles/TargetFilesProperties/TargetFilesProperties.INT.udn
--- a/Engine/Documentation/Source/Programming/UnrealBuildSystem/TargetFiles/TargetFilesProperties/TargetFilesProperties.INT.udn
+++ b/Engine/Documentation/Source/Programming/UnrealBuildSystem/TargetFiles/TargetFilesProperties/TargetFilesProperties.INT.udn
@@ -272,6 +272,8 @@ $ bAllowRuntimeSymbolFiles (Boolean): True if runtime symbols files should be ge
 
 $ PackagePath (String): Package full path (directory + filename) where to store input files used at link time Normally used to debug a linker crash for platforms that support it
 
+$ CrashDiagnosticDirectory (String): Directory where to put crash report files for platforms that support it
+
 $ BundleVersion (String): Bundle version for Mac apps.
 
 $ bDeployAfterCompile (Boolean): Whether to deploy the executable after compilation on platforms that require deployment.
