Documentation for build tools

diff --git a/Engine/Documentation/Source/Programming/UnrealBuildSystem/TargetFiles/TargetFilesProperties/TargetFilesProperties.INT.udn b/Engine/Documentation/Source/Programming/UnrealBuildSystem/TargetFiles/TargetFilesProperties/TargetFilesProperties.INT.udn
--- a/Engine/Documentation/Source/Programming/UnrealBuildSystem/TargetFiles/TargetFilesProperties/TargetFilesProperties.INT.udn
+++ b/Engine/Documentation/Source/Programming/UnrealBuildSystem/TargetFiles/TargetFilesProperties/TargetFilesProperties.INT.udn
@@ -48,6 +48,8 @@ $ UndecoratedConfiguration (UnrealTargetConfiguration): Specifies the configurat
 
 $ bBuildAllModules (Boolean): Build all the modules that are valid for this target type. Used for CIS and making installed engine builds.
 
+$ VSTestRunSettingsFile (FileReference): Set this to reference a VSTest run settings file from generated projects.
+
 $ BuildPlugins (List&lt;String&gt;): Additional plugins that are built for this target type but not enabled.
 
 $ bRuntimeDependenciesComeFromBuildPlugins (Boolean): If this is true, then the BuildPlugins list will be used to populate RuntimeDependencies, rather than EnablePlugins
