Documentation for build tools

diff --git a/Engine/Documentation/Source/Programming/UnrealBuildSystem/Configuration/BuildConfigProperties/BuildConfigProperties.INT.udn b/Engine/Documentation/Source/Programming/UnrealBuildSystem/Configuration/BuildConfigProperties/BuildConfigProperties.INT.udn
--- a/Engine/Documentation/Source/Programming/UnrealBuildSystem/Configuration/BuildConfigProperties/BuildConfigProperties.INT.udn
+++ b/Engine/Documentation/Source/Programming/UnrealBuildSystem/Configuration/BuildConfigProperties/BuildConfigProperties.INT.udn
@@ -12,6 +12,8 @@ $ bPrintDebugInfo : Whether debug info should be written to the console.
 
 $ bAllowHybridExecutor : Whether the hybrid executor will be used (a remote executor and local executor).
 
+$ bAllowHordeCompute : Whether Horde remote compute may be used. Highly experimental, disabled by default.
+
 $ bAllowXGE : Whether XGE may be used.
 
 $ bAllowFASTBuild : Whether FASTBuild may be used.
@@ -250,6 +252,12 @@ $ bStopOnError : Whether to stop on error
 
 $ MsvcCRTRedistVersion : Which MSVC CRT Redist version to use
 
+### HordeExecutor
+
+$ NumRemoteParallelProcesses : How many processes that will be executed in parallel remotely
+
+$ RemoteProcessOnly : Only run work remotely, for testing
+
 ### HybridExecutor
 
 $ MaxLocalActions : Maximum number of actions to execute locally.
