Documentation for build tools

diff --git a/Engine/Documentation/Source/Programming/UnrealBuildSystem/ModuleFiles/ModuleFilesProperties/ModuleFilesProperties.INT.udn b/Engine/Documentation/Source/Programming/UnrealBuildSystem/ModuleFiles/ModuleFilesProperties/ModuleFilesProperties.INT.udn
--- a/Engine/Documentation/Source/Programming/UnrealBuildSystem/ModuleFiles/ModuleFilesProperties/ModuleFilesProperties.INT.udn
+++ b/Engine/Documentation/Source/Programming/UnrealBuildSystem/ModuleFiles/ModuleFilesProperties/ModuleFilesProperties.INT.udn
@@ -138,6 +138,8 @@ $ bRequiresImplementModule (Nullable&lt;Boolean&gt;): Whether this module requir
 
 $ VersePath (String): If this module has associated Verse code, this is the Verse root path of it
 
+$ bValidateCircularDependencies (Boolean): Whether circular dependencies will be validated against the allow list Circular module dependencies result in slower builds. Disabling this option is strongly discouraged. This option is ignored for Engine modules which will always be validated against the allow list.
+
 $ CppStandard (CppStandardVersion): Which stanard to use for compiling this module
 
 $ ModuleSymbolVisibility (SymbolVisibility): Control visibility of symbols
