Remove UE4 mentions and images in enhanced tooltip files.

diff --git a/Engine/Documentation/Source/Shared/Editor/FbxErrors/FbxErrors.INT.udn b/Engine/Documentation/Source/Shared/Editor/FbxErrors/FbxErrors.INT.udn
--- a/Engine/Documentation/Source/Shared/Editor/FbxErrors/FbxErrors.INT.udn
+++ b/Engine/Documentation/Source/Shared/Editor/FbxErrors/FbxErrors.INT.udn
@@ -379,7 +379,7 @@ The importer is trying to recreate skeleton, but it failed because an asset exis
 
 (#SkeletalMesh_ExceedsMaxBoneCount)
 [EXCERPT:SkeletalMesh_ExceedsMaxBoneCount]
-Your Skeletal Mesh has too many bones, Unreal Engine 4's current max bone count is 65536.
+Your Skeletal Mesh has too many bones, Unreal Engine's current max bone count is 65536.
 
 [/EXCERPT:SkeletalMesh_ExceedsMaxBoneCount]
 
diff --git a/Engine/Documentation/Source/Shared/Editor/HLOD/HLOD.INT.udn b/Engine/Documentation/Source/Shared/Editor/HLOD/HLOD.INT.udn
--- a/Engine/Documentation/Source/Shared/Editor/HLOD/HLOD.INT.udn
+++ b/Engine/Documentation/Source/Shared/Editor/HLOD/HLOD.INT.udn
@@ -3,10 +3,9 @@ Title:HLOD System
 Crumbs:
 
 [EXCERPT:main]
-###Hierarchical Level of Detail (HLODs)
-Has been developed as a way to increase performance and keep the quality bar high for projects using Unreal Engine 4.11 or later. 
+Hierarchical Levels of Detail (HLODs) are a way to increase performance while keeping the quality bar high. 
 
-HLODs can replace multiple Static Mesh Actors with single, combined Static Mesh Actor at long view distances. This helps reduce the number of Actors that need to be rendered for the scene, increasing performance by lowering the number of draw calls per frame. 
+Each HLOD combines multiple Static Mesh Objects into a single object with simplified geometry. At long distances, the HLOD is drawn instead of those individual Static Mesh Objects. This reduces the number of Actors that need to be rendered each frame, increasing performance.
 [VAR:ToolTipFullLink]
 Engine/HLOD
 [/VAR]
diff --git a/Engine/Documentation/Source/Shared/Editor/MapErrors/MapErrors.INT.udn b/Engine/Documentation/Source/Shared/Editor/MapErrors/MapErrors.INT.udn
--- a/Engine/Documentation/Source/Shared/Editor/MapErrors/MapErrors.INT.udn
+++ b/Engine/Documentation/Source/Shared/Editor/MapErrors/MapErrors.INT.udn
@@ -10,7 +10,7 @@ Tags:error
 
 [TOC(start:2 end:3)]
 
-Unreal Engine 4 provides a mechanism for alerting you to potential problems in the form of map check errors and warnings instead of having to hunt down issues manually without knowing what is causing the problem or how to fix it. These errors are an automated way to find out where problems lie and allows you to get them fixed before they have a serious impact on your game. 
+Unreal Engine provides a mechanism for alerting you to potential problems in the form of map check errors and warnings instead of having to hunt down issues manually without knowing what is causing the problem or how to fix it. These errors are an automated way to find out where problems lie and allows you to get them fixed before they have a serious impact on your game. 
 
 
 
diff --git a/Engine/Documentation/Source/Shared/Editors/Persona/PersonaEditorModes.INT.udn b/Engine/Documentation/Source/Shared/Editors/Persona/PersonaEditorModes.INT.udn
--- a/Engine/Documentation/Source/Shared/Editors/Persona/PersonaEditorModes.INT.udn
+++ b/Engine/Documentation/Source/Shared/Editors/Persona/PersonaEditorModes.INT.udn
@@ -11,7 +11,7 @@ The skeleton mode lets you view properties of a skeleton asset.
 
 ### Overview
 
-In Unreal Engine 4, Skeletons are assets that are separate from Skeletal Meshes. Each Skeletal Mesh requires a Skeleton to be associated with it, and all animation assets (such as Animation Sequences, Animation Montages, and Blend Spaces) are associated with the Skeleton, not with the Skeletal Mesh.
+In Unreal Engine, Skeletons are assets that are separate from Skeletal Meshes. Each Skeletal Mesh requires a Skeleton to be associated with it, and all animation assets (such as Animation Sequences, Animation Montages, and Blend Spaces) are associated with the Skeleton, not with the Skeletal Mesh.
 
 ### What is a Skeleton?
 
@@ -21,7 +21,7 @@ If you are familiar with digital content creation, you may already feel like you
 A typical character skeleton as you might see in Maya.
 
 
-An important distinction is that a Skeleton asset is NOT the same thing as the bone hierarchy found within a Skeletal Mesh. In Unreal Engine 4, the Skeleton asset is less like a 3D representation of an actual skeleton and more like a list which holds bone and hierarchy information for a specific type of character or Skeletal Mesh. It is through this list that animation associations are made. 
+An important distinction is that a Skeleton asset is NOT the same thing as the bone hierarchy found within a Skeletal Mesh. In Unreal Engine, the Skeleton asset is less like a 3D representation of an actual skeleton and more like a list which holds bone and hierarchy information for a specific type of character or Skeletal Mesh. It is through this list that animation associations are made. 
 
 In terms of its place in the animation system, it is probably best to think of a Skeleton asset as an intermediary between Skeletal Meshes and animations. Multiple Skeletal Meshes can use the same Skeleton asset, and any Skeletal Meshes using the same Skeleton asset will be able to share animations.
 
@@ -85,7 +85,7 @@ Changing the properties listed in defaults mode will save the values in the asse
 
 ### Overview
 
-In Unreal Engine 4, Skeletons are assets that are separate from Skeletal Meshes. Each Skeletal Mesh requires a Skeleton to be associated with it, and all animation assets (such as Animation Sequences, Animation Montages, and Blend Spaces) are associated with the Skeleton, not with the Skeletal Mesh.
+In Unreal Engine, Skeletons are assets that are separate from Skeletal Meshes. Each Skeletal Mesh requires a Skeleton to be associated with it, and all animation assets (such as Animation Sequences, Animation Montages, and Blend Spaces) are associated with the Skeleton, not with the Skeletal Mesh.
 
 ### What is a Skeleton?
 
@@ -95,7 +95,7 @@ If you are familiar with digital content creation, you may already feel like you
 A typical character skeleton as you might see in Maya.
 
 
-An important distinction is that a Skeleton asset is NOT the same thing as the bone hierarchy found within a Skeletal Mesh. In Unreal Engine 4, the Skeleton asset is less like a 3D representation of an actual skeleton and more like a list which holds bone and hierarchy information for a specific type of character or Skeletal Mesh. It is through this list that animation associations are made. 
+An important distinction is that a Skeleton asset is NOT the same thing as the bone hierarchy found within a Skeletal Mesh. In Unreal Engine, the Skeleton asset is less like a 3D representation of an actual skeleton and more like a list which holds bone and hierarchy information for a specific type of character or Skeletal Mesh. It is through this list that animation associations are made. 
 
 In terms of its place in the animation system, it is probably best to think of a Skeleton asset as an intermediary between Skeletal Meshes and animations. Multiple Skeletal Meshes can use the same Skeleton asset, and any Skeletal Meshes using the same Skeleton asset will be able to share animations.
 
@@ -153,8 +153,6 @@ Animation Notifies tab lists all the custom animation notifies that are used.
 
 This does not display blueprint notifies (such as Play sound or Play particle effect), but it does display custom made notifies. 
 
-![](CustomAnimNotify.PNG)
-
 In this panel, you can remove, rename the custom animation notifies. If you remove or rename, please note that you will have to resave all animations that have this notify. 
 Otherwise it will not be removed or renamed correctly.
 
diff --git a/Engine/Documentation/Source/Shared/FullBlueprintEditor/FullBlueprintEditor.INT.udn b/Engine/Documentation/Source/Shared/FullBlueprintEditor/FullBlueprintEditor.INT.udn
--- a/Engine/Documentation/Source/Shared/FullBlueprintEditor/FullBlueprintEditor.INT.udn
+++ b/Engine/Documentation/Source/Shared/FullBlueprintEditor/FullBlueprintEditor.INT.udn
@@ -31,9 +31,7 @@ This will compile the current Blueprint (default keyboard short cut of F7). Any
 [EXCERPT:EditGlobalOptions]
 
 
-This will open the Blueprint Properties in the Details Pane of the Blueprint Editor:
-
-![](k2_blueprintProps.png)
+This will open the Blueprint Properties in the Details Pane of the Blueprint Editor.
 
 Here you can add Blueprint Interfaces, re-parent the Blueprint, and change some other properties that exist solely on the Blueprint.
 
diff --git a/Engine/Documentation/Source/Shared/Glossary/D/Glossary_D.INT.udn b/Engine/Documentation/Source/Shared/Glossary/D/Glossary_D.INT.udn
--- a/Engine/Documentation/Source/Shared/Glossary/D/Glossary_D.INT.udn
+++ b/Engine/Documentation/Source/Shared/Glossary/D/Glossary_D.INT.udn
@@ -63,7 +63,7 @@ and collision properties. The values set for these properties in the Defaults mo
 used when new instances of the _Blueprint_ are created.
 [/EXCERPT:Glossary_Defaults_Mode]
 [EXCERPT:Glossary_Details_panel]
-The **Details panel** is a panel in the Level Editor that contains information about the properties for various elements. Many of the tools and editors in UE4 contain Details panels: the Level Editor, Matinee, Persona, and so on.
+The **Details panel** is a panel in the Level Editor that contains information about the properties for various elements. Many of the tools and editors in UE contain Details panels: the Level Editor, Matinee, Persona, and so on.
 [/EXCERPT:Glossary_Details_panel]
 [EXCERPT:Glossary_Director_Group]
 The **Director Group** is a specific type of Matinee group that controls camera viewpoints. You can use the Director Group to switch camera viewpoints between shots. Each Matinee sequence can have only one Director Group.
diff --git a/Engine/Documentation/Source/Shared/Glossary/M/Glossary_M.INT.udn b/Engine/Documentation/Source/Shared/Glossary/M/Glossary_M.INT.udn
--- a/Engine/Documentation/Source/Shared/Glossary/M/Glossary_M.INT.udn
+++ b/Engine/Documentation/Source/Shared/Glossary/M/Glossary_M.INT.udn
@@ -86,7 +86,7 @@ A **Material** is an asset that can be applied to a mesh to control the visual l
 
 In more technical terms, when light from the scene hits the surface, a Material is used to calculate how that light interacts with that surface. These calculations are done using incoming data that is input to the Material from a variety of images (textures) and math expressions, as well as from various property settings inherent to the Material itself.
 
-Unreal Engine 4 utilizes a _physically-based shading model_. This means that rather than defining a Material using arbitrary properties (such as Diffuse Color and Specular Power), you instead use properties more easily relatable to the real world. These include Base Color, Metallic, Specular, and Roughness.
+Unreal Engine utilizes a _physically-based shading model_. This means that rather than defining a Material using arbitrary properties (such as Diffuse Color and Specular Power), you instead use properties more easily relatable to the real world. These include Base Color, Metallic, Specular, and Roughness.
 [/EXCERPT:Glossary_Material]
 [EXCERPT:Glossary_Matinee]
 The **Matinee** animation tool provides the ability to animate the properties of Actors over time, to create either dynamic gameplay or cinematic in-game sequences. The system is based on the use of specialized animation tracks in which you can place keyframes to set the values of certain properties of the Actors in the level. The Matinee Editor is similar to the non-linear editors used for video editing, making it familiar to video professionals.
diff --git a/Engine/Documentation/Source/Shared/Glossary/O/Glossary_O.INT.udn b/Engine/Documentation/Source/Shared/Glossary/O/Glossary_O.INT.udn
--- a/Engine/Documentation/Source/Shared/Glossary/O/Glossary_O.INT.udn
+++ b/Engine/Documentation/Source/Shared/Glossary/O/Glossary_O.INT.udn
@@ -19,6 +19,6 @@ Object
 
 [COMMENT:none]
 [EXCERPT:Glossary_Object]
-The base building blocks in the Unreal Engine are called **Objects** and contain a lot of the essential "under the hood" functionality for your game assets. Just about everything in Unreal Engine 4 inherits (or gets some functionality) from an Object. In C++, `UObject` is the base class of all objects; it implements features such as garbage collections, metadata (`UProperty`) support for exposing variables to the Unreal Editor, and serialization for loading and saving.
+The base building blocks in the Unreal Engine are called **Objects** and contain a lot of the essential "under the hood" functionality for your game assets. Just about everything in Unreal Engine inherits (or gets some functionality) from an Object. In C++, `UObject` is the base class of all objects; it implements features such as garbage collections, metadata (`UProperty`) support for exposing variables to the Unreal Editor, and serialization for loading and saving.
 [/EXCERPT:Glossary_Object]
 [/COMMENT]
\ No newline at end of file
diff --git a/Engine/Documentation/Source/Shared/GraphNodes/Material/Material.INT.udn b/Engine/Documentation/Source/Shared/GraphNodes/Material/Material.INT.udn
--- a/Engine/Documentation/Source/Shared/GraphNodes/Material/Material.INT.udn
+++ b/Engine/Documentation/Source/Shared/GraphNodes/Material/Material.INT.udn
@@ -616,7 +616,7 @@ Engine/Rendering/Materials/ExpressionReference/Math#linearinterpolate
 [EXCERPT:UMaterialExpressionMultiply]
 The Multiply expression takes two inputs, multiplies them together, and outputs the result. Similar to Photoshop's multiply layer blend. The multiplication happens per channel, meaning that the R channel of the first is multiplied by the second; the G channel of the first is multiplied by the second, and so on. Both inputs must have the same number of values unless one of the values is a single float value. 
 
-Note: Do not forget that materials in UE4 are not limited to [0,1]! If colors/values are greater than 1, Multiply will actually brighten colors. 
+Note: Do not forget that materials in UE are not limited to [0,1]! If colors/values are greater than 1, Multiply will actually brighten colors. 
 
 [VAR:ToolTipFullLink]
 Engine/Rendering/Materials/ExpressionReference/Math#multiply
@@ -826,7 +826,7 @@ Engine/Rendering/Materials/ExpressionReference/Utility#blackbody
 [/EXCERPT:UMaterialExpressionBlackBody]
 
 [EXCERPT:UMaterialExpressionBumpOffset]
-BumpOffset is the Unreal Engine 4 term for what is commonly known as 'parallax mapping'. The BumpOffset expression allows a material to give the illusion of depth without the need for additional geometry.
+BumpOffset is the Unreal Engine term for what is commonly known as 'parallax mapping'. The BumpOffset expression allows a material to give the illusion of depth without the need for additional geometry.
 
  BumpOffset materials use a grayscale heightmap to give depth information. The brighter the value in the heightmap, the more 'popped out' the material will be; these areas will parallax (shift) as a camera moves across the surface. Darker areas in the heightmap are 'further away' and will shift the least. 
 
@@ -976,7 +976,7 @@ Engine/Rendering/Materials/ExpressionReference/Vector#cameravectorws
 [/EXCERPT:UMaterialExpressionCameraVectorWS]
 
 [EXCERPT:UMaterialExpressionLightVector]
-This expression is deprecated in Unreal Engine 4 as lighting calculations are now deferred.
+This expression is deprecated as lighting calculations are now deferred.
 
 [VAR:ToolTipFullLink]
 Engine/Rendering/Materials/ExpressionReference/Vector#lightvector
diff --git a/Engine/Documentation/Source/Shared/Types/ABoxReflectionCapture/ABoxReflectionCapture.INT.udn b/Engine/Documentation/Source/Shared/Types/ABoxReflectionCapture/ABoxReflectionCapture.INT.udn
--- a/Engine/Documentation/Source/Shared/Types/ABoxReflectionCapture/ABoxReflectionCapture.INT.udn
+++ b/Engine/Documentation/Source/Shared/Types/ABoxReflectionCapture/ABoxReflectionCapture.INT.udn
@@ -5,5 +5,5 @@ Description:Overview of the BoxReflectionCapture class
 
 [EXCERPT:ABoxReflectionCapture]
 
-The Box Reflection Capture Actor is one of two Reflection Capture Actors available in Unreal Engine 4. Many of these can be placed around the world with almost no performance impact as they are calculated before run time.
+The Box Reflection Capture Actor captures an image of the surrounding area for use in reflections, best suited for use in rectangular hallways and interior rooms. You can place any number of these Actors with little impact on runtime performance.
 [/EXCERPT:ABoxReflectionCapture]
diff --git a/Engine/Documentation/Source/Shared/Types/APostProcessVolume/APostProcessVolume.INT.udn b/Engine/Documentation/Source/Shared/Types/APostProcessVolume/APostProcessVolume.INT.udn
--- a/Engine/Documentation/Source/Shared/Types/APostProcessVolume/APostProcessVolume.INT.udn
+++ b/Engine/Documentation/Source/Shared/Types/APostProcessVolume/APostProcessVolume.INT.udn
@@ -1,13 +1,14 @@
-﻿Availability:NoPublish
+Availability:NoPublish
 Title:APostProcessVolume
 Crumbs:
 Description:Overview of the PostProcessVolume class
 
 [EXCERPT:APostProcessVolume]
 
-The PostProcessVolume is a special type of volume that can be added to a level. Unreal Engine 3 had the post processing chains, but in Unreal Engine 4 all this got replaced by
-using the PostProcessSettings which can be conveniently adjusted using one or multiple PostProcessVolume. By using multiple volumes (or other objects) with priority and weight
-the settings can be softly blended depending on the camera position and manipulating weights.
+The PostProcessVolume applies post-processing effects to camera renderings, such as bloom, lens flares, exposure control, color correction, and so on.
+The effects start to apply when the camera is at a distance from the volume set by the Blend Radius setting, and are at full intensity when the camera is inside
+the bounds of the volume. You can blend different effects together by creating multiple overlapping PostProcessVolumes, and control their blending by assigning each 
+volume a relative priority and weight. You can also make a volume's settings apply to the entire level by enabling its Infinite Extent (Unbound) setting.
 
 [VAR:ToolTipFullLink]
 Engine/Rendering/PostProcessEffects
diff --git a/Engine/Documentation/Source/Shared/Types/FBodyInstance/FBodyInstance.INT.udn b/Engine/Documentation/Source/Shared/Types/FBodyInstance/FBodyInstance.INT.udn
--- a/Engine/Documentation/Source/Shared/Types/FBodyInstance/FBodyInstance.INT.udn
+++ b/Engine/Documentation/Source/Shared/Types/FBodyInstance/FBodyInstance.INT.udn
@@ -22,24 +22,6 @@ Engine/Physics/PhysicsBodies#physics
 
 
 
-[EXCERPT:bUseAsyncScene]
-
-PhysX runs two physics modes: Synchronous and Asynchronous, and Unreal Engine 4 has two scenes that match these.
-
-The Synchronous scene is where most classic physics tasks take place such as a falling crate or vehicle. Where order really makes a difference on gameplay, the Synchronous scene is used.
-
-In the Asynchronous is where things like Destructible Actors are simulated. Here, order does not matter as most Physics objects are not interacting directly with each other in a way that affects gameplay.
-
-By default, simulating Actors in one scene will not interact with Actors in the other scene. The exception is Static Actors, as they are placed into both scenes.
-
-Under the hood, there is a bit more going on about blocking and execution order, consult the PhysX documentation for more information.
-[VAR:ToolTipFullLink]
-Engine/Physics/PhysicsBodies#physics
-[/VAR]
-[/EXCERPT:bUseAsyncScene]
-
-
-
 
 [EXCERPT:PhysMaterialOverride]
 
