Documentation for build tools

diff --git a/Engine/Documentation/Source/Programming/UnrealBuildSystem/TargetFiles/TargetFilesProperties/TargetFilesProperties.INT.udn b/Engine/Documentation/Source/Programming/UnrealBuildSystem/TargetFiles/TargetFilesProperties/TargetFilesProperties.INT.udn
--- a/Engine/Documentation/Source/Programming/UnrealBuildSystem/TargetFiles/TargetFilesProperties/TargetFilesProperties.INT.udn
+++ b/Engine/Documentation/Source/Programming/UnrealBuildSystem/TargetFiles/TargetFilesProperties/TargetFilesProperties.INT.udn
@@ -18,6 +18,10 @@ $ Version (ReadOnlyBuildVersion): The current build version
 
 $ IsTestTarget (Boolean): Whether this is a low level tests target.
 
+$ ExplicitTestsTarget (Boolean): Whether this is a test target explicitly defined. Explicitley defined test targets always inherit from TestTargetRules and define their own tests. Implicit test targets are created from existing targets when building with -Mode=Test and they include tests from all dependencies.
+
+$ WithLowLevelTests (Boolean): Controls the value of WITH_LOW_LEVEL_TESTS that dictates whether module-specific low level tests are compiled in or not.
+
 $ HostPlatform (UnrealTargetPlatform): Gets the host platform being built on
 
 $ bGenerateProjectFiles (Boolean): Expose the bGenerateProjectFiles flag to targets, so we can modify behavior as appropriate for better intellisense
@@ -116,12 +120,6 @@ $ bCompileCustomSQLitePlatform (Boolean): Whether we should compile SQLite using
 
 $ bUseCacheFreedOSAllocs (Boolean): Whether to utilize cache freed OS allocs with MallocBinned
 
-$ bCompileAgainstEngine (Boolean): Enabled for all builds that include the engine project.  Disabled only when building standalone apps that only link with Core.
-
-$ bCompileAgainstCoreUObject (Boolean): Enabled for all builds that include the CoreUObject project.  Disabled only when building standalone apps that only link with Core.
-
-$ bCompileAgainstApplicationCore (Boolean): Enabled for builds that need to initialize the ApplicationCore module. Command line utilities do not normally need this.
-
 $ bCompileRecast (Boolean): Whether to compile Recast navmesh generation.
 
 $ bCompileNavmeshSegmentLinks (Boolean): Whether to compile with navmesh segment links.
@@ -366,6 +364,12 @@ $ bBuildTargetDeveloperTools (Boolean): Whether to compile the developer tools t
 
 $ bNeedsExtraShaderFormats (Boolean): Whether we should include any extra shader formats. By default this is only enabled for Program and Editor targets.
 
+$ bCompileAgainstEngine (Boolean): Enabled for all builds that include the engine project.  Disabled only when building standalone apps that only link with Core.
+
+$ bCompileAgainstCoreUObject (Boolean): Enabled for all builds that include the CoreUObject project.  Disabled only when building standalone apps that only link with Core.
+
+$ bCompileAgainstApplicationCore (Boolean): Enabled for builds that need to initialize the ApplicationCore module. Command line utilities do not normally need this.
+
 $ bCompileAgainstEditor (Boolean): Enabled for editor builds (TargetType.Editor). Can be overridden for programs (TargetType.Program) that would need to compile against editor code. Not available for other target types. Mainly drives the value of WITH_EDITOR.
 
 $ bCompileSpeedTree (Boolean): Whether we should compile in support for SpeedTree or not.
