Documentation for build tools

diff --git a/Engine/Documentation/Source/Programming/UnrealBuildSystem/Configuration/BuildConfigProperties/BuildConfigProperties.INT.udn b/Engine/Documentation/Source/Programming/UnrealBuildSystem/Configuration/BuildConfigProperties/BuildConfigProperties.INT.udn
--- a/Engine/Documentation/Source/Programming/UnrealBuildSystem/Configuration/BuildConfigProperties/BuildConfigProperties.INT.udn
+++ b/Engine/Documentation/Source/Programming/UnrealBuildSystem/Configuration/BuildConfigProperties/BuildConfigProperties.INT.udn
@@ -274,6 +274,8 @@ $ ProcessorCountMultiplier : Processor count multiplier for local execution. Can
 
 $ MemoryPerActionBytes : Free memory per action in bytes, used to limit the number of parallel actions if the machine is memory starved. Set to 0 to disable free memory checking.
 
+$ ProcessPriority : The priority to set for spawned processes. Valid Settings: Idle, BelowNormal, Normal, AboveNormal, High Default: BelowNormal
+
 $ bStopCompilationAfterErrors : When enabled, will stop compiling targets after a compile error occurs.
 
 $ bShowCompilationTimes : Whether to show compilation times along with worst offenders or not.
