Documentation for build tools

diff --git a/Engine/Documentation/Source/Programming/UnrealBuildSystem/TargetFiles/TargetFilesProperties/TargetFilesProperties.INT.udn b/Engine/Documentation/Source/Programming/UnrealBuildSystem/TargetFiles/TargetFilesProperties/TargetFilesProperties.INT.udn
--- a/Engine/Documentation/Source/Programming/UnrealBuildSystem/TargetFiles/TargetFilesProperties/TargetFilesProperties.INT.udn
+++ b/Engine/Documentation/Source/Programming/UnrealBuildSystem/TargetFiles/TargetFilesProperties/TargetFilesProperties.INT.udn
@@ -70,7 +70,7 @@ $ PakSigningKeysFile (String): Path to the set of pak signing keys to embed in t
 
 $ SolutionDirectory (String): Allows a Program Target to specify it's own solution folder path.
 
-$ bBuildInSolutionByDefault (Nullable&lt;Boolean&gt;): Whether the target should be included in the default solution build configuration
+$ bBuildInSolutionByDefault (Nullable&lt;Boolean&gt;): Whether the target should be included in the default solution build configuration Setting this to false will skip building when running in the IDE
 
 $ bShouldCompileAsDLL (Boolean): Whether this target should be compiled as a DLL.  Requires LinkType to be set to TargetLinkType.Monolithic.
 
