Documentation for build tools

diff --git a/Engine/Documentation/Source/Programming/UnrealBuildSystem/TargetFiles/TargetFilesProperties/TargetFilesProperties.INT.udn b/Engine/Documentation/Source/Programming/UnrealBuildSystem/TargetFiles/TargetFilesProperties/TargetFilesProperties.INT.udn
--- a/Engine/Documentation/Source/Programming/UnrealBuildSystem/TargetFiles/TargetFilesProperties/TargetFilesProperties.INT.udn
+++ b/Engine/Documentation/Source/Programming/UnrealBuildSystem/TargetFiles/TargetFilesProperties/TargetFilesProperties.INT.udn
@@ -36,6 +36,8 @@ $ Type (TargetType): The type of target.
 
 $ ForcedIncludeOrder (Nullable&lt;EngineIncludeOrderVersion&gt;): Force the include order to a specific version. Overrides any Target and Module rules.
 
+$ OutputFile (String): Path to the output file for the main executable, relative to the Engine or project directory. This setting is only typically useful for non-UE programs, since the engine uses paths relative to the executable to find other known folders (eg. Content).
+
 $ bUsesSteam (Boolean): Whether the target uses Steam.
 
 $ bUsesCEF3 (Boolean): Whether the target uses CEF3.
