Documentation for build tools

diff --git a/Engine/Documentation/Source/Programming/UnrealBuildSystem/TargetFiles/TargetFilesProperties/TargetFilesProperties.INT.udn b/Engine/Documentation/Source/Programming/UnrealBuildSystem/TargetFiles/TargetFilesProperties/TargetFilesProperties.INT.udn
--- a/Engine/Documentation/Source/Programming/UnrealBuildSystem/TargetFiles/TargetFilesProperties/TargetFilesProperties.INT.udn
+++ b/Engine/Documentation/Source/Programming/UnrealBuildSystem/TargetFiles/TargetFilesProperties/TargetFilesProperties.INT.udn
@@ -286,6 +286,8 @@ $ bCheckSystemHeadersForModification (Boolean): Whether headers in system paths
 
 $ bDisableLinking (Boolean): Whether to disable linking for this target.
 
+$ bIgnoreBuildOutputs (Boolean): Whether to ignore tracking build outputs for this target.
+
 $ bFormalBuild (Boolean): Indicates that this is a formal build, intended for distribution. This flag is automatically set to true when Build.version has a changelist set. The only behavior currently bound to this flag is to compile the default resource file separately for each binary so that the OriginalFilename field is set correctly. By default, we only compile the resource once to reduce build times.
 
 $ bFlushBuildDirOnRemoteMac (Boolean): Whether to clean Builds directory on a remote Mac before building.
